[UBA]
* Added and enabled overlapped network send for windows
* Added iocp path (disabled)
* Added per-socket validation when sockets are timing out

diff --git a/Engine/Source/Programs/UnrealBuildAccelerator/Agent/Private/UbaAgent.cpp b/Engine/Source/Programs/UnrealBuildAccelerator/Agent/Private/UbaAgent.cpp
--- a/Engine/Source/Programs/UnrealBuildAccelerator/Agent/Private/UbaAgent.cpp
+++ b/Engine/Source/Programs/UnrealBuildAccelerator/Agent/Private/UbaAgent.cpp
@@ -1024,7 +1024,12 @@ namespace uba
 				networkBackend = new NetworkBackendQuic(logWriter);
 			else
 			#endif
-				networkBackend = new NetworkBackendTcp(logWriter);
+			{
+				NetworkBackendTcpCreateInfo info(logWriter);
+				//info.statusUpdateSeconds = 2;
+				networkBackend = new NetworkBackendTcp(info);
+			}
+
 			auto backendGuard = MakeGuard([networkBackend]() { delete networkBackend; });
 
 			NetworkClientCreateInfo ncci(logWriter);
@@ -1330,10 +1335,6 @@ namespace uba
 			else
 				logger.Info(TC("----------- Session %s started -----------"), sessionClient->GetId());
 
-			#if 0 
-			u64 lastLogTime = GetTime();
-			#endif
-
 			u32 connectionCount = 1;
 
 			//#if PLATFORM_WINDOWS
@@ -1451,19 +1452,6 @@ namespace uba
 							LoggerWithWriter(g_consoleLogWriter, TC("")).Info(TC("%s"), terminationReason.data);
 					}
 				}
-
-				#if 0 // Not needed anymore I think.. horde is now pinging in the background for itself
-				if (quiet)
-				{
-					u64 time = GetTime();
-					u64 timePast = time - lastLogTime;
-					if (TimeToMs(timePast) > 90 * 1000)
-					{
-						logger.Info(TC("Break the silence after %s (here to keep horde happy...). Agent active process count: %u"), TimeToText(timePast).str, sessionClient->GetActiveProcessCount());
-						lastLogTime = time;
-					}
-				}
-				#endif
 			}
 
 			disconnectAndStopLoggingThread.Execute();
diff --git a/Engine/Source/Programs/UnrealBuildAccelerator/Common/Private/UbaNetworkBackendMemory.cpp b/Engine/Source/Programs/UnrealBuildAccelerator/Common/Private/UbaNetworkBackendMemory.cpp
--- a/Engine/Source/Programs/UnrealBuildAccelerator/Common/Private/UbaNetworkBackendMemory.cpp
+++ b/Engine/Source/Programs/UnrealBuildAccelerator/Common/Private/UbaNetworkBackendMemory.cpp
@@ -197,7 +197,7 @@ namespace uba
 		outRecv = 0;
 	}
 
-	void NetworkBackendMemory::Validate(Logger& logger)
+	void NetworkBackendMemory::Validate(Logger& logger, const Vector<void*>& connections, bool full)
 	{
 		logger.Info(TC("  NetworkBackendMemory"));
 		logger.Info(TC("     NoInfo"));
diff --git a/Engine/Source/Programs/UnrealBuildAccelerator/Common/Private/UbaNetworkBackendTcp.cpp b/Engine/Source/Programs/UnrealBuildAccelerator/Common/Private/UbaNetworkBackendTcp.cpp
--- a/Engine/Source/Programs/UnrealBuildAccelerator/Common/Private/UbaNetworkBackendTcp.cpp
+++ b/Engine/Source/Programs/UnrealBuildAccelerator/Common/Private/UbaNetworkBackendTcp.cpp
@@ -44,11 +44,28 @@
 #define UBA_EMULATE_BAD_INTERNET 0
 
 #define UBA_USE_OVERLAPPED_SEND PLATFORM_WINDOWS
+#define UBA_USE_OVERLAPPED_SEND_WITH_LOCK 0
+#define UBA_USE_IOCP 0
 
 namespace uba
 {
 	constexpr u32 MaxHeaderSize = 24;
 
+	////////////////////////////////////////////////////////////////////////////////////////////////////
+
+	bool ShutdownSocket(Logger& logger, SOCKET s, const tchar* hint);
+	bool CloseSocket(Logger& logger, SOCKET s, const tchar* hint);
+	bool SetKeepAlive(Logger& logger, SOCKET socket);
+	bool SetBlocking(Logger& logger, SOCKET socket, bool blocking);
+	bool SetTimeout(Logger& logger, SOCKET socket, u32 timeoutMs);
+	bool SetLinger(Logger& logger, SOCKET socket, u32 lingerSeconds);
+	bool SetRecvBuf(Logger& logger, SOCKET socket, u32 windowSize);
+	bool SetSendBuf(Logger& logger, SOCKET socket, u32 windowSize);
+	bool DisableNagle(Logger& logger, SOCKET socket);
+	bool LogTcpInfo(Logger& logger, SOCKET socket);
+
+	////////////////////////////////////////////////////////////////////////////////////////////////////
+
 	struct NetworkBackendTcp::ListenEntry
 	{
 		StringBuffer<128> ip;
@@ -59,6 +76,17 @@
 		Thread thread;
 	};
 
+	////////////////////////////////////////////////////////////////////////////////////////////////////
+
+	#if UBA_USE_OVERLAPPED_SEND || UBA_USE_IOCP
+	struct Overlapped : OVERLAPPED
+	{
+		bool isSend = false;
+	};
+	#endif
+
+	////////////////////////////////////////////////////////////////////////////////////////////////////
+
 	struct NetworkBackendTcp::Connection
 	{
 		Connection(Logger& l, SOCKET s) : logger(l), socket(s), ready(true) { CreateGuid(uid); }
@@ -92,17 +120,24 @@
 
 		bool allowLess = false;
 
+		// Temporary state for iocp
+		#if UBA_USE_IOCP
+		Overlapped overlapped;
+		WSABUF wsaBuf = {};
+		u8 header[MaxHeaderSize];
+		u8* bodyData = nullptr;
+		u32 bodySize = 0;
+		u32 left = 0;
+		void* bodyContext = nullptr;
+		bool receivingHeader = true;
+		Event exited;
+		#endif
+
 		Connection(const Connection&) = delete;
 		void operator=(const Connection&) = delete;
 	};
 
-	bool SetKeepAlive(Logger& logger, SOCKET socket);
-	bool SetBlocking(Logger& logger, SOCKET socket, bool blocking);
-	bool SetTimeout(Logger& logger, SOCKET socket, u32 timeoutMs);
-	bool SetLinger(Logger& logger, SOCKET socket, u32 lingerSeconds);
-	bool SetRecvBuf(Logger& logger, SOCKET socket, u32 windowSize);
-	bool SetSendBuf(Logger& logger, SOCKET socket, u32 windowSize);
-	bool DisableNagle(Logger& logger, SOCKET socket);
+	////////////////////////////////////////////////////////////////////////////////////////////////////
 
 	struct NetworkBackendTcp::RecvCache
 	{
@@ -110,14 +145,41 @@
 		u32 byteCount = 0;
 	};
 
+	////////////////////////////////////////////////////////////////////////////////////////////////////
+
+	void NetworkBackendTcpCreateInfo::Apply(Config& config, const tchar* tableName)
+	{
+		const ConfigTable* tablePtr = config.GetTable(tableName);
+		if (!tablePtr)
+			return;
+		const ConfigTable& table = *tablePtr;
+		table.GetValueAsBool(disableNagle, TC("DisableNagle"));
+		table.GetValueAsU32(statusUpdateSeconds, TC("StatusUpdateSeconds"));
+	}
+
+	////////////////////////////////////////////////////////////////////////////////////////////////////
+
 	bool NetworkBackendTcp::EnsureInitialized(Logger& logger)
 	{
 		#if PLATFORM_WINDOWS
 		WSADATA wsaData;
 		if (!m_wsaInitDone)
+		{
 			if (int res = WSAStartup(MAKEWORD(2, 2), &wsaData))
 				return logger.Error(TC("WSAStartup failed (%d)"), res);
-		m_wsaInitDone = true;
+
+			#if UBA_USE_IOCP
+			//if (!IsRunningWine())
+			{
+				m_iocpHandle = CreateIoCompletionPort(INVALID_HANDLE_VALUE, NULL, 0, 0);
+
+				for (Thread& t : m_iocpThreads)
+					t.Start([this]() { ThreadIocp(); return 0; });
+			}
+			#endif
+
+			m_wsaInitDone = true;
+		}
 		#else
 		static bool initOnce = []()
 			{
@@ -127,44 +189,20 @@
 			}();
 
 		#endif
+
 		return true;
 	}
 
-	bool ShutdownSocket(Logger& logger, SOCKET s, const tchar* hint)
-	{
-		if (s == INVALID_SOCKET)
-			return true;
-		if (shutdown(s, SD_BOTH) != SOCKET_ERROR)
-			return true;
-		if (WSAGetLastError() == WSAENOTCONN)
-			return true;
-		logger.Info(TC("Failed to shutdown socket %llu in %s (%s)"), u64(s), hint, LastErrorToText(WSAGetLastError()).data);
-		return false;
-	}
-
-	bool CloseSocket(Logger& logger, SOCKET s, const tchar* hint)
-	{
-		if (s == INVALID_SOCKET)
-			return true;
-		if (closesocket(s) != SOCKET_ERROR)
-			return true;
-		logger.Info(TC("Failed to close socket %llu in %s (%s)"), u64(s), hint, LastErrorToText(WSAGetLastError()).data);
-		return false;
-	}
-
-	void NetworkBackendTcpCreateInfo::Apply(Config& config, const tchar* tableName)
-	{
-		const ConfigTable* tablePtr = config.GetTable(tableName);
-		if (!tablePtr)
-			return;
-		const ConfigTable& table = *tablePtr;
-		table.GetValueAsBool(disableNagle, TC("DisableNagle"));
-	}
-
 	NetworkBackendTcp::NetworkBackendTcp(const NetworkBackendTcpCreateInfo& info, const tchar* prefix)
 		: m_logger(info.logWriter, prefix)
 	{
 		m_disableNagle = info.disableNagle;
+
+		if (info.statusUpdateSeconds)
+		{
+			m_tcpStatusLoop.Create(true);
+			m_tcpStatusThread.Start([this, sus = info.statusUpdateSeconds]() { ThreadStatus(sus); return 0; });
+		}
 	}
 
 	NetworkBackendTcp::~NetworkBackendTcp()
@@ -186,6 +224,15 @@
 		}
 		m_connections.clear();
 
+		#if UBA_USE_IOCP
+		m_iocpLoop = false;
+		for (Thread& t : m_iocpThreads)
+			t.Wait();
+		#endif
+
+		m_tcpStatusLoop.Set();
+		m_tcpStatusThread.Wait();
+
 		#if PLATFORM_WINDOWS
 		if (m_wsaInitDone)
 			WSACleanup();
@@ -238,6 +285,15 @@
 		conn.bodyCallback = b;
 		conn.recvHint = recvHint;
 		conn.ready.Set();
+
+		#if UBA_USE_IOCP
+		if (m_iocpHandle && !conn.wsaBuf.buf)
+		{
+			conn.exited.Create(true);
+			conn.left = headerSize;
+			PostIocpRead(conn, conn.header, headerSize);
+		}
+		#endif
 	}
 
 	void NetworkBackendTcp::SetRecvTimeout(void* connection, u32 timeoutMs, void* context, RecvTimeoutCallback* callback)
@@ -456,7 +512,13 @@
 			SCOPED_FUTEX(m_connectionsLock, lock);
 			auto it = m_connections.emplace(m_connections.end(), logger, clientSocket);
 			auto& conn = *it;
-			conn.recvThread.Start([this, connPtr = &conn] { ThreadRecv(*connPtr); return 0; }, TC("UbaTcpRecv"));
+			#if UBA_USE_IOCP
+			if (m_iocpHandle)
+				CreateIoCompletionPort((HANDLE)clientSocket, m_iocpHandle, (ULONG_PTR)&conn, 0);
+			else
+			#endif
+				conn.recvThread.Start([this, connPtr = &conn] { ThreadRecv(*connPtr); return 0; }, TC("UbaTcpRecv"));
+
 			lock.Leave();
 
 			if (!entry.connectedFunc(&conn, remoteSockAddr))
@@ -714,7 +776,13 @@
 		SCOPED_FUTEX(m_connectionsLock, lock);
 		auto it = m_connections.emplace(m_connections.end(), logger, socketFd);
 		auto& conn = *it;
-		conn.recvThread.Start([this, connPtr = &conn] { ThreadRecv(*connPtr); return 0; }, TC("UbaTcpRecv"));
+		#if UBA_USE_IOCP
+		if (m_iocpHandle)
+			CreateIoCompletionPort((HANDLE)socketFd, m_iocpHandle, (ULONG_PTR)&conn, 0);
+		else
+		#endif
+			conn.recvThread.Start([this, connPtr = &conn] { ThreadRecv(*connPtr); return 0; }, TC("UbaTcpRecv"));
+
 		lock.Leave();
 
 		if (!connectedFunc(&conn, remoteSocketAddr, timedOut))
@@ -744,6 +812,12 @@
 			Connection& c = *it;
 			if (&c != connection)
 				continue;
+			
+			#if UBA_USE_IOCP
+			if (m_iocpHandle)
+				c.exited.IsSet();
+			#endif
+				
 			it = m_connections.erase(it);
 			break;
 		}
@@ -755,114 +829,16 @@
 		outRecv = m_totalRecv;
 	}
 
-	void NetworkBackendTcp::Validate(Logger& logger)
+	void NetworkBackendTcp::Validate(Logger& logger, const Vector<void*>& connections, bool full)
 	{
 		logger.Info(TC("  NetworkBackendTcp"));
-		logger.Info(TC("     NoInfo"));
+
+		SCOPED_FUTEX(m_connectionsLock, lock);
+		for (auto ptr : connections)
+		{
+			auto& c = *(Connection*)ptr;
+			LogTcpInfo(logger, c.socket);
+		}
 	}
 
-	bool SetBlocking(Logger& logger, SOCKET socket, bool blocking)
-	{
-#if PLATFORM_WINDOWS
-		u_long value = blocking ? 0 : 1;
-		if (ioctlsocket(socket, FIONBIO, &value) == SOCKET_ERROR)
-			return logger.Error(TC("Setting non blocking socket failed (error: %s)"), LastErrorToText(WSAGetLastError()).data);
-#else
-		int flags = fcntl(socket, F_GETFL, 0);
-		if (flags == -1) return false;
-		flags = blocking ? (flags & ~O_NONBLOCK) : (flags | O_NONBLOCK);
-		if (fcntl(socket, F_SETFL, flags) != 0)
-			return logger.Error(TC("Setting non blocking socket failed (error: %s)"), LastErrorToText(WSAGetLastError()).data);
-#endif
-		return true;
-	}
-
-	bool DisableNagle(Logger& logger, SOCKET socket)
-	{
-#if !PLATFORM_MAC
-		u32 value = 1;
-		if (setsockopt(socket, IPPROTO_TCP, TCP_NODELAY, (const char*)&value, sizeof(value)) == SOCKET_ERROR)
-			return logger.Error(TC("setsockopt TCP_NODELAY error: (error: %s)"), LastErrorToText(WSAGetLastError()).data);
-#endif
-		return true;
-	}
-
-	bool SetTimeout(Logger& logger, SOCKET socket, u32 timeoutMs)
-	{
-		#if PLATFORM_WINDOWS
-		DWORD timeout = timeoutMs;
-		#else
-		struct timeval timeout;
-		timeout.tv_sec = timeoutMs / 1000;
-		timeout.tv_usec = (timeoutMs % 1000)*1000;
-		#endif
-
-		if (setsockopt(socket, SOL_SOCKET, SO_SNDTIMEO, (const char*)&timeout, sizeof timeout) == SOCKET_ERROR)
-			return logger.Error(TC("setsockopt SO_SNDTIMEO error: (error: %s)"), LastErrorToText(WSAGetLastError()).data);
-		if (setsockopt(socket, SOL_SOCKET, SO_RCVTIMEO, (const char*)&timeout, sizeof timeout) == SOCKET_ERROR)
-			return logger.Error(TC("setsockopt SO_RCVTIMEO error: (error: %s)"), LastErrorToText(WSAGetLastError()).data);
-		return true;
-	}
-
-	bool SetLinger(Logger& logger, SOCKET socket, u32 lingerSeconds)
-	{
-#if PLATFORM_MAC // Mac does not seem to automatically close the socket if the process crashes
-		struct linger so_linger;
-		so_linger.l_onoff = 1;
-		so_linger.l_linger = u16(lingerSeconds);
-
-		if (setsockopt(socket, SOL_SOCKET, SO_LINGER, (const char*)&so_linger, sizeof(so_linger)) < 0)
-			return logger.Error(TC("setsockopt SO_LINGER error: (error: %s)"), LastErrorToText(WSAGetLastError()).data);
-#endif
-		return true;
-	}
-
-	bool SetRecvBuf(Logger& logger, SOCKET socket, u32 windowSize)
-	{
-		if (setsockopt(socket, SOL_SOCKET, SO_RCVBUF, (const char*)&windowSize, sizeof(windowSize)) == SOCKET_ERROR)
-			return logger.Error(TC("setsockopt TCP_NODELAY error: (error: %s)"), LastErrorToText(WSAGetLastError()).data);
-		return true;
-	}
-
-	bool SetSendBuf(Logger& logger, SOCKET socket, u32 windowSize)
-	{
-		if (setsockopt(socket, SOL_SOCKET, SO_SNDBUF, (const char*)&windowSize, sizeof(windowSize)) == SOCKET_ERROR)
-			return logger.Error(TC("setsockopt TCP_NODELAY error: (error: %s)"), LastErrorToText(WSAGetLastError()).data);
-		return true;
-	}
-
-	bool SetKeepAlive(Logger& logger, SOCKET socket) // This will make sure that WSAPoll exits when the network cable is pulled
-	{
-		
-		int keepAliveTime = KeepAliveIdleSeconds; // Time of idle before it starts sending keep alive
-		int keepAliveTimeInterval = KeepAliveIntervalSeconds; // Time between each check (10 failures mean disconnect)
-
-
-		u32 value = 1;
-		if (setsockopt(socket, SOL_SOCKET, SO_KEEPALIVE, (const char*)&value, sizeof(value)) == SOCKET_ERROR)
-			return logger.Error(TC("setsockopt SO_KEEPALIVE (error: %s)"), LastErrorToText(WSAGetLastError()).data);
-
-#if PLATFORM_WINDOWS	
-		struct tcp_keepalive kaSettings;
-		DWORD bytesReturned;
-		kaSettings.onoff = 1;
-		kaSettings.keepalivetime = keepAliveTime * 1000;
-		kaSettings.keepaliveinterval = keepAliveTimeInterval * 1000; 
-		if (WSAIoctl(socket, SIO_KEEPALIVE_VALS, &kaSettings, sizeof(kaSettings), NULL, 0, &bytesReturned, NULL, NULL) == SOCKET_ERROR)
-			return logger.Error(TC("WSAIoctl SIO_KEEPALIVE_VALS (error: %s)"), LastErrorToText(WSAGetLastError()).data);
-#elif PLATFORM_LINUX
-		if (setsockopt(socket, IPPROTO_TCP, TCP_KEEPIDLE, &keepAliveTime, sizeof(int)) < 0)
-			return logger.Error(TC("setsockopt TCP_KEEPIDLE (error: %s)"), LastErrorToText(WSAGetLastError()).data);
-		if (setsockopt(socket, IPPROTO_TCP, TCP_KEEPINTVL, &keepAliveTimeInterval, sizeof(int)) < 0)
-			return logger.Error(TC("setsockopt TCP_KEEPINTVL (error: %s)"), LastErrorToText(WSAGetLastError()).data);
-		int keepAliveProbes = KeepAliveProbeCount; // Number of tests before timing out
-		if (setsockopt(socket, IPPROTO_TCP, TCP_KEEPCNT, &keepAliveProbes, sizeof(int)) < 0)
-			return logger.Error(TC("setsockopt TCP_KEEPCNT (error: %s)"), LastErrorToText(WSAGetLastError()).data);
-#else // PLATFORM_MAC
-		if (setsockopt(socket, IPPROTO_TCP, TCP_KEEPALIVE, &keepAliveTime, sizeof(int)) < 0)
-			return logger.Error(TC("setsockopt TCP_KEEPALIVE (error: %s)"), LastErrorToText(WSAGetLastError()).data);
-#endif
-		return true;
-	}
-
 	bool NetworkBackendTcp::SendSocket(Connection& connection, Logger& logger, const void* b, u64 bufferLen, const tchar* hint)
@@ -880,19 +856,35 @@
 
 		Event ev(true);
 
+		// Right now we are experimenting with if we can ignore taking a lock around the entire thing and only the WSASend.
+		// Documentation is slightly unclear and network forums claim the order of data is correct even though multiple
+		// threads call WSASend and then wait on event.
+		#if UBA_USE_OVERLAPPED_SEND_WITH_LOCK
 		SCOPED_FUTEX(connection.sendLock, lock);
-		SOCKET socket = connection.socket;
-		if (socket == INVALID_SOCKET)
-			return false;
+		#endif
 
 		char* buffer = (char*)b;
 		u64 left = bufferLen;
 		while (left)
 		{
-			WSAOVERLAPPED overlapped {};
+			Overlapped overlapped {};
+			overlapped.isSend = true;
 			overlapped.hEvent = ev.GetHandle();
 			WSABUF buf { u32(left), buffer };
+
+			#if !UBA_USE_OVERLAPPED_SEND_WITH_LOCK
+			SCOPED_FUTEX(connection.sendLock, lock);
+			#endif
+
+			SOCKET socket = connection.socket;
+			if (socket == INVALID_SOCKET)
+				return false;
 			int res = WSASend(socket, &buf, 1, NULL, 0, &overlapped, NULL);
+
+			#if !UBA_USE_OVERLAPPED_SEND_WITH_LOCK
+			lock.Leave();
+			#endif
+
 			if (res == SOCKET_ERROR)
 			{
 				u32 lastError = WSAGetLastError();
@@ -908,10 +900,19 @@
 			if (!ev.IsSet(40*1000))
 			{
 				logger.Info(TC("SendSocket - WSASend returned timeout after 40 seconds for socket %i (%s)"), socket, hint);
-				if (!ev.IsSet(DefaultNetworkSendTimeoutSeconds * 1000))
+				u64 startTime = GetTime();
+				while (true)
 				{
-					logger.Info(TC("SendSocket - WSASend returned timeout after 10 minutes for socket %i (%s)"), socket, hint);
-					return false;
+					LogTcpInfo(logger, socket);
+					if (ev.IsSet(2000))
+						break;
+
+					u64 sinceStartSeconds = TimeToMs(GetTime() - startTime)/1000;
+					if (sinceStartSeconds >= DefaultNetworkSendTimeoutSeconds)
+					{
+						logger.Info(TC("SendSocket - WSASend returned timeout after 10 minutes for socket %i (%s)"), socket, hint);
+						return false;
+					}
 				}
 			}
 
@@ -928,7 +929,16 @@
 			buffer += bytesSent;
 			left -= bytesSent;
 			if (left)
+			{
+				#if !UBA_USE_OVERLAPPED_SEND_WITH_LOCK
+				//#if UBA_LOG_SOCKET_ERRORS
+				logger.Warning(TC("SendSocket - WSASend did not send all data in one call for socket %i. Note that this might be allowed behavior. Hard to read out of documentation. (%s)"), socket, hint);
+				//#endif
+				return false;
+				#else
 				ev.Reset();
+				#endif
+			}
 		}
 
 #else
@@ -1122,6 +1132,195 @@
 		}
 		return true;
 	}
+	
+	#if UBA_USE_IOCP
+	void NetworkBackendTcp::ThreadIocp()
+	{
+		ElevateCurrentThreadPriority();
 
+		while (m_iocpLoop)
+		{
+			DWORD bytesTransferred;
+			ULONG_PTR completionKey;
+			Overlapped* overlapped;
+			u32 timeoutMs = 1000;
+			BOOL result = GetQueuedCompletionStatus(m_iocpHandle, &bytesTransferred, &completionKey, (OVERLAPPED**)&overlapped, timeoutMs);
+
+			if (!result || !overlapped)
+				continue;
+
+			auto& connection = *(Connection*)completionKey;
+			auto& logger = connection.logger;
+
+			if (bytesTransferred == 0)
+			{
+				{
+					SCOPED_CRITICAL_SECTION(connection.shutdownLock, lock2);
+					SOCKET s = connection.socket;
+					CancelIoEx((HANDLE)s, &connection.overlapped);
+
+					{
+						SCOPED_FUTEX(connection.sendLock, lock);
+						connection.socket = INVALID_SOCKET;
+					}
+					if (auto cb = connection.disconnectCallback)
+					{
+						auto context = connection.disconnectContext;
+						connection.disconnectCallback = nullptr;
+						connection.disconnectContext = nullptr;
+						cb(context, connection.uid, &connection);
+					}
+
+					if (s != INVALID_SOCKET)
+					{
+						ShutdownSocket(logger, s, TC("ThreadRecv"));
+						CloseSocket(logger, s, TC("ThreadRecv"));
+					}
+				}
+				connection.exited.Set();
+				continue;
+			}
+
+			if (overlapped->isSend)
+				continue;
+
+			connection.left -= bytesTransferred;
+			if (connection.left)
+			{
+				u8* newPos = (u8*)connection.wsaBuf.buf + bytesTransferred;
+				PostIocpRead(connection, newPos, connection.left);
+				continue;
+			}
+
+			if (connection.receivingHeader)
+			{
+				m_totalRecv += connection.headerSize;
+
+				auto hc = connection.headerCallback;
+				if (!hc)
+				{
+					logger.Error(TC("Tcp connection header callback not set"));
+					continue;
+				}
+
+				u8*& bodyData = connection.bodyData;
+				u32& bodySize = connection.bodySize;
+				bodyData = nullptr;
+				bodySize = 0;
+				if (!hc(connection.recvContext, connection.uid, connection.header, connection.bodyContext, bodyData, bodySize))
+					continue;
+
+				if (connection.socket == INVALID_SOCKET)
+					continue;
+
+				if (!bodySize)
+				{
+					connection.left = connection.headerSize;
+					PostIocpRead(connection, connection.header, connection.left);
+					continue;
+				}
+
+				connection.left = bodySize;
+				connection.receivingHeader = false;
+				PostIocpRead(connection, bodyData, bodySize);
+			}
+			else
+			{
+				auto bc = connection.bodyCallback;
+				if (!bc)
+				{
+					logger.Error(TC("Tcp connection body callback not set"));
+					continue;
+				}
+
+				bool success = true;
+				if (!bc(connection.recvContext, !success, connection.header, connection.bodyContext, connection.bodyData, connection.bodySize))
+					continue;
+				if (!success)
+					continue;
+
+				m_totalRecv += connection.bodySize;
+
+				if (connection.socket == INVALID_SOCKET)
+					continue;
+
+				connection.receivingHeader = true;
+				connection.left = connection.headerSize;
+				PostIocpRead(connection, connection.header, connection.headerSize);
+			}
+		}
+	}
+
+	bool NetworkBackendTcp::PostIocpRead(Connection& connection, u8* data, u32 dataSize)
+	{
+		ZeroMemory(&connection.overlapped, sizeof(OVERLAPPED));
+		connection.wsaBuf.buf = (char*)data;
+		connection.wsaBuf.len = dataSize;
+
+		DWORD flags = 0, bytesRecv = 0;
+		int ret = WSARecv(connection.socket, &connection.wsaBuf, 1, &bytesRecv, &flags, &connection.overlapped, NULL);
+		if (ret == SOCKET_ERROR && WSAGetLastError() != WSA_IO_PENDING)
+		{
+			printf("WSARecv failed: %d\n", WSAGetLastError());
+			//closesocket(clientSocket);
+			//delete context;
+			return false;
+		}
+		return true;
+	}
+	#endif
+
+	void NetworkBackendTcp::ThreadStatus(u32 statusUpdateSeconds)
+	{
+		u32 statusUpdateMs = statusUpdateSeconds*1000;
+
+		#if PLATFORM_WINDOWS
+		MIB_TCPSTATS_LH prevStats{};
+		#endif
+
+		while (!m_tcpStatusLoop.IsSet(statusUpdateMs))
+		{
+			#if PLATFORM_WINDOWS
+			MIB_TCPSTATS_LH stats{};
+			if (GetTcpStatisticsEx(&stats, AF_INET) != NO_ERROR)
+			{
+				m_logger.Info(TC("GetTcpStatisticsEx failed"));
+				return;
+			}
+
+			if (prevStats.dwRtoAlgorithm != 0)
+			{
+				INT64 recv       = stats.dwInSegs - prevStats.dwInSegs;
+				INT64 sent       = stats.dwOutSegs - prevStats.dwOutSegs;
+				INT64 retrans    = stats.dwRetransSegs - prevStats.dwRetransSegs;
+				INT64 inerrs     = stats.dwInErrs - prevStats.dwInErrs;
+				INT64 outrsts    = stats.dwOutRsts - prevStats.dwOutRsts;
+
+				double rrate = sent ? (100.0 * retrans / sent) : 0;
+
+				m_logger.Info(TC("%us delta  recv=%lld  sent=%lld  retrans=%lld  (%.2f%%)  inErr=%lld  RST=%lld"), statusUpdateSeconds, recv, sent, retrans, rrate, inerrs, outrsts);
+			}
+			prevStats = stats;
+			#endif
+
+			#if 0
+			SCOPED_FUTEX(m_connectionsLock, lock);
+			for (auto& conn : m_connections)
+			{
+				SCOPED_CRITICAL_SECTION(conn.shutdownLock, lock2);
+				if (conn.socket == INVALID_SOCKET)
+					continue;
+				SCOPED_FUTEX(conn.sendLock, sendLock);
+				SOCKET s = conn.socket;
+				if (s == INVALID_SOCKET)
+					continue;
+				LogTcpInfo(m_logger, s);
+			}
+			#endif
+		}
+	}
+
+	////////////////////////////////////////////////////////////////////////////////////////////////////
+
 	void TraverseNetworkAddresses(Logger& logger, const Function<bool(const StringBufferBase& addr)>& func)
 	{
@@ -1211,6 +1410,8 @@
 		return true;
 	}
 
+	////////////////////////////////////////////////////////////////////////////////////////////////////
+
 	HttpConnection::HttpConnection()
 	{
 		m_socket = INVALID_SOCKET;
@@ -1366,4 +1567,286 @@
 	{
 		m_connectTimeOutMs = timeOutMs;
 	}
+
+	////////////////////////////////////////////////////////////////////////////////////////////////////
+
+	bool ShutdownSocket(Logger& logger, SOCKET s, const tchar* hint)
+	{
+		if (s == INVALID_SOCKET)
+			return true;
+		if (shutdown(s, SD_BOTH) != SOCKET_ERROR)
+			return true;
+		if (WSAGetLastError() == WSAENOTCONN)
+			return true;
+		logger.Info(TC("Failed to shutdown socket %llu in %s (%s)"), u64(s), hint, LastErrorToText(WSAGetLastError()).data);
+		return false;
+	}
+
+	bool CloseSocket(Logger& logger, SOCKET s, const tchar* hint)
+	{
+		if (s == INVALID_SOCKET)
+			return true;
+		if (closesocket(s) != SOCKET_ERROR)
+			return true;
+		logger.Info(TC("Failed to close socket %llu in %s (%s)"), u64(s), hint, LastErrorToText(WSAGetLastError()).data);
+		return false;
+	}
+
+	bool SetBlocking(Logger& logger, SOCKET socket, bool blocking)
+	{
+#if PLATFORM_WINDOWS
+		u_long value = blocking ? 0 : 1;
+		if (ioctlsocket(socket, FIONBIO, &value) == SOCKET_ERROR)
+			return logger.Error(TC("Setting non blocking socket failed (error: %s)"), LastErrorToText(WSAGetLastError()).data);
+#else
+		int flags = fcntl(socket, F_GETFL, 0);
+		if (flags == -1) return false;
+		flags = blocking ? (flags & ~O_NONBLOCK) : (flags | O_NONBLOCK);
+		if (fcntl(socket, F_SETFL, flags) != 0)
+			return logger.Error(TC("Setting non blocking socket failed (error: %s)"), LastErrorToText(WSAGetLastError()).data);
+#endif
+		return true;
+	}
+
+	bool DisableNagle(Logger& logger, SOCKET socket)
+	{
+#if !PLATFORM_MAC
+		u32 value = 1;
+		if (setsockopt(socket, IPPROTO_TCP, TCP_NODELAY, (const char*)&value, sizeof(value)) == SOCKET_ERROR)
+			return logger.Error(TC("setsockopt TCP_NODELAY error: (error: %s)"), LastErrorToText(WSAGetLastError()).data);
+#endif
+		return true;
+	}
+
+	bool SetTimeout(Logger& logger, SOCKET socket, u32 timeoutMs)
+	{
+		#if PLATFORM_WINDOWS
+		DWORD timeout = timeoutMs;
+		#else
+		struct timeval timeout;
+		timeout.tv_sec = timeoutMs / 1000;
+		timeout.tv_usec = (timeoutMs % 1000)*1000;
+		#endif
+
+		if (setsockopt(socket, SOL_SOCKET, SO_SNDTIMEO, (const char*)&timeout, sizeof timeout) == SOCKET_ERROR)
+			return logger.Error(TC("setsockopt SO_SNDTIMEO error: (error: %s)"), LastErrorToText(WSAGetLastError()).data);
+		if (setsockopt(socket, SOL_SOCKET, SO_RCVTIMEO, (const char*)&timeout, sizeof timeout) == SOCKET_ERROR)
+			return logger.Error(TC("setsockopt SO_RCVTIMEO error: (error: %s)"), LastErrorToText(WSAGetLastError()).data);
+		return true;
+	}
+
+	bool SetLinger(Logger& logger, SOCKET socket, u32 lingerSeconds)
+	{
+#if PLATFORM_MAC // Mac does not seem to automatically close the socket if the process crashes
+		struct linger so_linger;
+		so_linger.l_onoff = 1;
+		so_linger.l_linger = u16(lingerSeconds);
+
+		if (setsockopt(socket, SOL_SOCKET, SO_LINGER, (const char*)&so_linger, sizeof(so_linger)) < 0)
+			return logger.Error(TC("setsockopt SO_LINGER error: (error: %s)"), LastErrorToText(WSAGetLastError()).data);
+#endif
+		return true;
+	}
+
+	bool SetRecvBuf(Logger& logger, SOCKET socket, u32 windowSize)
+	{
+		if (setsockopt(socket, SOL_SOCKET, SO_RCVBUF, (const char*)&windowSize, sizeof(windowSize)) == SOCKET_ERROR)
+			return logger.Error(TC("setsockopt TCP_NODELAY error: (error: %s)"), LastErrorToText(WSAGetLastError()).data);
+		return true;
+	}
+
+	bool SetSendBuf(Logger& logger, SOCKET socket, u32 windowSize)
+	{
+		if (setsockopt(socket, SOL_SOCKET, SO_SNDBUF, (const char*)&windowSize, sizeof(windowSize)) == SOCKET_ERROR)
+			return logger.Error(TC("setsockopt TCP_NODELAY error: (error: %s)"), LastErrorToText(WSAGetLastError()).data);
+		return true;
+	}
+
+	bool SetKeepAlive(Logger& logger, SOCKET socket) // This will make sure that WSAPoll exits when the network cable is pulled
+	{
+		
+		int keepAliveTime = KeepAliveIdleSeconds; // Time of idle before it starts sending keep alive
+		int keepAliveTimeInterval = KeepAliveIntervalSeconds; // Time between each check (10 failures mean disconnect)
+
+
+		u32 value = 1;
+		if (setsockopt(socket, SOL_SOCKET, SO_KEEPALIVE, (const char*)&value, sizeof(value)) == SOCKET_ERROR)
+			return logger.Error(TC("setsockopt SO_KEEPALIVE (error: %s)"), LastErrorToText(WSAGetLastError()).data);
+
+#if PLATFORM_WINDOWS	
+		struct tcp_keepalive kaSettings;
+		DWORD bytesReturned;
+		kaSettings.onoff = 1;
+		kaSettings.keepalivetime = keepAliveTime * 1000;
+		kaSettings.keepaliveinterval = keepAliveTimeInterval * 1000; 
+		if (WSAIoctl(socket, SIO_KEEPALIVE_VALS, &kaSettings, sizeof(kaSettings), NULL, 0, &bytesReturned, NULL, NULL) == SOCKET_ERROR)
+			return logger.Error(TC("WSAIoctl SIO_KEEPALIVE_VALS (error: %s)"), LastErrorToText(WSAGetLastError()).data);
+#elif PLATFORM_LINUX
+		if (setsockopt(socket, IPPROTO_TCP, TCP_KEEPIDLE, &keepAliveTime, sizeof(int)) < 0)
+			return logger.Error(TC("setsockopt TCP_KEEPIDLE (error: %s)"), LastErrorToText(WSAGetLastError()).data);
+		if (setsockopt(socket, IPPROTO_TCP, TCP_KEEPINTVL, &keepAliveTimeInterval, sizeof(int)) < 0)
+			return logger.Error(TC("setsockopt TCP_KEEPINTVL (error: %s)"), LastErrorToText(WSAGetLastError()).data);
+		int keepAliveProbes = KeepAliveProbeCount; // Number of tests before timing out
+		if (setsockopt(socket, IPPROTO_TCP, TCP_KEEPCNT, &keepAliveProbes, sizeof(int)) < 0)
+			return logger.Error(TC("setsockopt TCP_KEEPCNT (error: %s)"), LastErrorToText(WSAGetLastError()).data);
+#else // PLATFORM_MAC
+		if (setsockopt(socket, IPPROTO_TCP, TCP_KEEPALIVE, &keepAliveTime, sizeof(int)) < 0)
+			return logger.Error(TC("setsockopt TCP_KEEPALIVE (error: %s)"), LastErrorToText(WSAGetLastError()).data);
+#endif
+		return true;
+	}
+
+	#if !defined(TCP_INFO_v0) && PLATFORM_WINDOWS
+	enum TCPSTATE {};
+	struct TCP_INFO_v0
+	{
+		TCPSTATE State;
+		ULONG    Mss;
+		ULONG64  ConnectionTimeMs;
+		BOOLEAN  TimestampsEnabled;
+		ULONG    RttUs;
+		ULONG    MinRttUs;
+		ULONG    BytesInFlight;
+		ULONG    Cwnd;
+		ULONG    SndWnd;
+		ULONG    RcvWnd;
+		ULONG    RcvBuf;
+		ULONG64  BytesOut;
+		ULONG64  BytesIn;
+		ULONG    BytesReordered;
+		ULONG    BytesRetrans;
+		ULONG    FastRetrans;
+		ULONG    DupAcksIn;
+		ULONG    TimeoutEpisodes;
+		UCHAR    SynRetrans;
+	};
+	#endif
+
+	#ifndef SIO_TCP_INFO     // older MinGW headers
+	#define SIO_TCP_INFO  _WSAIORW(IOC_VENDOR,0x50)
+	#endif
+
+
+#if PLATFORM_WINDOWS
+	struct LinuxTcpInfo        
+	{
+		u8 tcpi_state;         
+		u8 tcpi_ca_state;
+		u8 tcpi_retransmits;
+		u8 tcpi_probes;
+		u8 tcpi_backoff;
+		u8 tcpi_options;
+		u8 tcpi_snd_wscale : 4,
+		   tcpi_rcv_wscale : 4;
+
+		u32 tcpi_rto;          
+		u32 tcpi_ato;
+		u32 tcpi_snd_mss;
+		u32 tcpi_rcv_mss;
+
+		u32 tcpi_unacked;      
+		u32 tcpi_sacked;
+		u32 tcpi_lost;
+		u32 tcpi_retrans;      
+		u32 tcpi_fackets;
+
+		u32 tcpi_last_data_sent;
+		u32 tcpi_last_ack_sent;
+		u32 tcpi_last_data_recv;
+		u32 tcpi_last_ack_recv;
+
+		u32 tcpi_pmtu;
+		u32 tcpi_rcv_ssthresh;
+		u32 tcpi_rtt;           
+		u32 tcpi_rttvar;        
+		u32 tcpi_snd_ssthresh;
+		u32 tcpi_snd_cwnd;      
+		u32 tcpi_advmss;
+		u32 tcpi_reordering;
+
+		u32 tcpi_rcv_rtt;
+		u32 tcpi_rcv_space;     
+
+		u32 tcpi_total_retrans; 
+
+		u64 tcpi_pacing_rate;    
+		u64 tcpi_max_pacing_rate;
+		u64 tcpi_bytes_acked;    
+		u64 tcpi_bytes_received; 
+		u64 tcpi_segs_out;       
+		u64 tcpi_segs_in;        
+
+		u32 tcpi_notsent_bytes;  
+		u32 tcpi_min_rtt;        
+		u32 tcpi_data_segs_in;
+		u32 tcpi_data_segs_out;
+
+		u64 tcpi_delivery_rate;  
+
+		u64 tcpi_busy_time;      
+		u64 tcpi_rwnd_limited;   
+		u64 tcpi_sndbuf_limited; 
+
+		u32 tcpi_delivered;
+		u32 tcpi_delivered_ce;
+
+		u64 tcpi_bytes_sent;
+		u64 tcpi_bytes_retrans;
+		u64 tcpi_dsack_dups;
+		u64 tcpi_reord_seen;
+
+		u32 tcpi_rtt_min;
+		u32 tcpi_rcv_rtt_min;
+	};
+#endif
+
+
+	bool LogTcpInfo(Logger& logger, SOCKET socket)
+	{
+	#if PLATFORM_WINDOWS
+		if (IsRunningWine())
+		{
+			static HMODULE wineDll = LoadLibrary(TC("UbaWine.dll.so"));
+			if (!wineDll)
+				return false;
+			using GetLinuxTcpInfoFunc = int WINAPI(SOCKET, void*, int*);
+			static auto GetLinuxTcpInfo = (GetLinuxTcpInfoFunc*)GetProcAddress(wineDll, "GetLinuxTcpInfo");
+			if (!GetLinuxTcpInfo)
+				return logger.Error(TC("GetLinuxTcpInfo not found in UbaWine dll"));
+			LinuxTcpInfo ti {};
+			int size = sizeof(LinuxTcpInfo);
+			int res = GetLinuxTcpInfo(socket, &ti, &size);
+			if (res != 0)
+				return logger.Error(TC("GetLinuxTcpInfo failed with error %i"), res);
+
+			u32 bytes_inflight_est = ti.tcpi_unacked * ti.tcpi_snd_mss + ti.tcpi_notsent_bytes;
+			logger.Info(TC("%p - RTT=%uus  inFlight=%u  retrans=%u  dupacks=%u  state=%u"), socket, ti.tcpi_rtt, bytes_inflight_est, ti.tcpi_retrans, ti.tcpi_snd_cwnd, ti.tcpi_state);
+			return true;
+		}
+
+		TCP_INFO_v0 ti;
+		ULONG version = 0; // Specify 0 to retrieve the v0 version of this structure.
+		ULONG bytesReturned;
+
+		if (WSAIoctl(socket, SIO_TCP_INFO, &version, sizeof(ULONG), &ti, sizeof(TCP_INFO_v0), &bytesReturned, NULL, NULL) == SOCKET_ERROR)
+			return logger.Error(TC("WSAIoctl SIO_TCP_INFO failed (%s)"), LastErrorToText(WSAGetLastError()).data);
+		logger.Info(TC("rtt=%uµs  inFlight=%u  retrans=%u  dupacks=%u  state=%u"), ti.RttUs, ti.BytesInFlight, ti.BytesRetrans, ti.DupAcksIn, ti.State);
+		return true;
+	#else
+		return false;
+		/*
+		struct tcp_info ti = {};
+		socklen_t len = sizeof(ti);
+		if (getsockopt(s, IPPROTO_TCP, TCP_INFO, &ti, &len) == 0)
+		{
+			printf("%s  rtt=%uµs  unacked=%u  retrans=%u  snd_cwnd=%u  state=%u\n",
+				   tag,
+				   ti.tcpi_rtt, ti.tcpi_unacked, ti.tcpi_retrans,
+				   ti.tcpi_snd_cwnd, ti.tcpi_state);
+		}
+		*/
+	#endif
+	}
+
+	////////////////////////////////////////////////////////////////////////////////////////////////////
 }
diff --git a/Engine/Source/Programs/UnrealBuildAccelerator/Common/Private/UbaNetworkClient.cpp b/Engine/Source/Programs/UnrealBuildAccelerator/Common/Private/UbaNetworkClient.cpp
--- a/Engine/Source/Programs/UnrealBuildAccelerator/Common/Private/UbaNetworkClient.cpp
+++ b/Engine/Source/Programs/UnrealBuildAccelerator/Common/Private/UbaNetworkClient.cpp
@@ -141,6 +141,7 @@ namespace uba
 		StackBinaryWriter<1024> handshakeData;
 		if (m_cryptoKey)
 		{
+			handshakeData.WriteByte(1);
 			// If we have a crypto key we start by sending a predefined 128 bytes blob that is encrypted.
 			// If server decrypt it to the same blob, we're good on that part
 			u8* encryptedBuffer = handshakeData.AllocWrite(sizeof(EncryptionHandshakeString));//[1024];
@@ -422,9 +423,9 @@ namespace uba
 		logger.Info(TC(""));
 	}
 
-	void NetworkClient::ValidateNetwork(Logger& logger)
+	void NetworkClient::ValidateNetwork(Logger& logger, bool full)
 	{
-		UnorderedSet<NetworkBackend*> backends;
+		UnorderedMap<NetworkBackend*, Vector<void*>> backends;
 		u32 connectionsCount;
 		{
 			LogStallScope lss(logger, LogEntryType_Info, 1, TC("  Connections lock took %s"));
@@ -435,12 +436,15 @@ namespace uba
 			u32 connectionIndex = 0;
 			for (auto& c : m_connections)
 			{
-				#if UBA_TRACK_NETWORK_TIMES
-				logger.Info(TC("  Connection %u - LastSend %s, LastHeaderRecv: %s, LastBodyRecv: %s"), connectionIndex++, TimeToText(c.lastSendTime - m_startTime).str, TimeToText(c.lastHeaderRecvTime - m_startTime).str, TimeToText(c.lastBodyRecvTime - m_startTime).str);
-				#else
-				logger.Info(TC("  Connection %u"), connectionIndex++);
-				#endif
-				backends.insert(c.backend);
+				if (full)
+				{
+					#if UBA_TRACK_NETWORK_TIMES
+					logger.Info(TC("  Connection %u - LastSend %s, LastHeaderRecv: %s, LastBodyRecv: %s"), connectionIndex++, TimeToText(c.lastSendTime - m_startTime).str, TimeToText(c.lastHeaderRecvTime - m_startTime).str, TimeToText(c.lastBodyRecvTime - m_startTime).str);
+					#else
+					logger.Info(TC("  Connection %u"), connectionIndex++);
+					#endif
+				}
+				backends[c.backend].push_back(c.backendConnection);
 			}
 		}
 
@@ -449,6 +453,7 @@ namespace uba
 			SCOPED_FUTEX(m_connectionsItLock, lock);
 		}
 
+		if (full)
 		{
 			LogStallScope lss(logger, LogEntryType_Info, 1, TC("  ActiveMessages lock took %s"));
 			SCOPED_WRITE_LOCK(m_activeMessagesLock, lock);
@@ -467,8 +472,8 @@ namespace uba
 				}
 		}
 
-		for (auto backend : backends)
-			backend->Validate(logger);
+		for (auto& kv : backends)
+			kv.first->Validate(logger, kv.second, full);
 	}
 
 	void NetworkClient::RegisterOnConnected(const OnConnectedFunction& function)
diff --git a/Engine/Source/Programs/UnrealBuildAccelerator/Common/Private/UbaNetworkServer.cpp b/Engine/Source/Programs/UnrealBuildAccelerator/Common/Private/UbaNetworkServer.cpp
--- a/Engine/Source/Programs/UnrealBuildAccelerator/Common/Private/UbaNetworkServer.cpp
+++ b/Engine/Source/Programs/UnrealBuildAccelerator/Common/Private/UbaNetworkServer.cpp
@@ -124,12 +124,15 @@ namespace uba
 			m_backend.SetRecvTimeout(m_backendConnection, m_server.m_receiveTimeoutMs, this, [](void* context, u32 timeoutMs, const tchar* recvHint, const tchar* hint)
 				{
 					auto& conn = *(Connection*)context;
-					conn.m_server.m_logger.Warning(TC("Connection %u timed out after %u seconds (%s%s)"), conn.m_id, timeoutMs/1000, recvHint, hint);
+					u32 clientId = ~0u;
+					if (auto c = conn.m_client)
+						clientId = c->id;
+					conn.m_server.m_logger.Warning(TC("Connection %u (Client %u) timed out after %u seconds (%s%s)"), conn.m_id, clientId, timeoutMs/1000, recvHint, hint);
 					return false;
 				});
 
 			if (requiresCrypto)
-				m_backend.SetRecvCallbacks(m_backendConnection, this, 0, ReceiveHandshakeHeader, ReceiveHandshakeBody, TC("ReceiveHandshake"));
+				m_backend.SetRecvCallbacks(m_backendConnection, this, 1, ReceiveHandshakeHeader, ReceiveHandshakeBody, TC("ReceiveHandshake"));
 			else
 				m_backend.SetRecvCallbacks(m_backendConnection, this, 4, ReceiveVersion, nullptr, TC("ReceiveVersion"));
 		}
@@ -276,7 +279,10 @@ namespace uba
 			if (conn.m_recvTimeoutCount < KeepAliveProbeCount)
 				return true;
 			constexpr u32 totalTimeoutSeconds = KeepAliveIdleSeconds + KeepAliveIntervalSeconds*KeepAliveProbeCount;
-			conn.m_server.m_logger.Warning(TC("Connection %u timed out after %u seconds (%s%s)"), conn.m_id, totalTimeoutSeconds, recvHint, hint);
+			u32 clientId = ~0u;
+			if (auto c = conn.m_client)
+				clientId = c->id;
+			conn.m_server.m_logger.Warning(TC("Connection %u (Client %u) timed out after %u seconds (%s%s)"), conn.m_id, clientId, totalTimeoutSeconds, recvHint, hint);
 			return false;
 		}
 
@@ -534,6 +540,7 @@ namespace uba
 		CryptoKey cryptoKey = connection.m_cryptoKey;
 		if (cryptoKey)
 		{
+			//TrackHintScope ths(tws, TCV("Decrypt"));
 			TimerScope ts(connection.m_decryptTimer);
 			if (!Crypto::Decrypt(server.m_logger, cryptoKey, context.buffer.data(), context.dataSize))
 			{
@@ -557,25 +564,28 @@ namespace uba
 		mi.connectionId = connection.m_id;
 		mi.messageId = context.id;
 
-		if (!rec.func)
-		{
-			server.m_logger.Error(TC("WORKER FUNCTION NOT FOUND. id: %u, serviceid: %u type: %s, client: %u"), context.id, context.serviceId, rec.toString(context.messageType).data, connection.m_client->id);
-			connection.SetShouldDisconnect();
-			size = MessageErrorSize;
-		}
-		else if (!rec.func({&connection}, {tws}, mi, reader, writer))
 		{
-			if (connection.SetShouldDisconnect())
+			//TrackHintScope ths(tws, TCV("HandleMessage"));
+			if (!rec.func)
 			{
-				#if UBA_DEBUG
-				server.m_logger.Error(TC("WORKER FUNCTION FAILED. id: %u, serviceid: %u type: %s, client: %u"), context.id, context.serviceId, rec.toString(context.messageType).data, connection.m_client->id);
-				#endif
+				server.m_logger.Error(TC("WORKER FUNCTION NOT FOUND. id: %u, serviceid: %u type: %s, client: %u"), context.id, context.serviceId, rec.toString(context.messageType).data, connection.m_client->id);
+				connection.SetShouldDisconnect();
+				size = MessageErrorSize;
+			}
+			else if (!rec.func({&connection}, {tws}, mi, reader, writer))
+			{
+				if (connection.SetShouldDisconnect())
+				{
+					#if UBA_DEBUG
+					server.m_logger.Error(TC("WORKER FUNCTION FAILED. id: %u, serviceid: %u type: %s, client: %u"), context.id, context.serviceId, rec.toString(context.messageType).data, connection.m_client->id);
+					#endif
+				}
+				size = MessageErrorSize;
+			}
+			else
+			{
+				size = u32(writer.GetPosition());
 			}
-			size = MessageErrorSize;
-		}
-		else
-		{
-			size = u32(writer.GetPosition());
 		}
 
 		if (mi.messageId)
@@ -585,6 +595,7 @@ namespace uba
 			u32 bodySize = u32(size - HeaderSize);
 			if (cryptoKey && size != MessageErrorSize && bodySize)
 			{
+				//TrackHintScope ths(tws, TCV("Encrypt"));
 				TimerScope ts(connection.m_encryptTimer);
 				u8* bodyData = writer.GetData() + HeaderSize;
 				if (!Crypto::Encrypt(server.m_logger, cryptoKey, bodyData, bodySize))
@@ -602,6 +613,7 @@ namespace uba
 			//if (size == MessageErrorSize)
 			//	UBA_ASSERT(false);
 
+			//TrackHintScope ths(tws, TCV("Send"));
 			connection.Send(writer.GetData(), size == MessageErrorSize ? HeaderSize : size, TC("MessageResponse"));
 		}
 			
diff --git a/Engine/Source/Programs/UnrealBuildAccelerator/Common/Private/UbaSessionClient.cpp b/Engine/Source/Programs/UnrealBuildAccelerator/Common/Private/UbaSessionClient.cpp
--- a/Engine/Source/Programs/UnrealBuildAccelerator/Common/Private/UbaSessionClient.cpp
+++ b/Engine/Source/Programs/UnrealBuildAccelerator/Common/Private/UbaSessionClient.cpp
@@ -1651,7 +1651,10 @@ namespace uba
 			logger.Info(TC("Took more than %u seconds to send/receive ping%s"), timeoutSeconds, (m_pingTimeoutSecondsPrintCallstacks ? TC(". Printing callstacks") : TC("")));
 			if (m_pingTimeoutSecondsPrintCallstacks)
 				PrintAllCallstacks(logger);
-			m_client.ValidateNetwork(logger);
+			m_client.ValidateNetwork(logger, true);
+
+			while (!response.done.IsSet(2000))
+				m_client.ValidateNetwork(logger, false);
 
 		}
 		
@@ -1661,7 +1664,6 @@ namespace uba
 		{
 			LoggerWithWriter logger(g_consoleLogWriter);
 			logger.Info(TC("Ping finished after %s"), TimeToText(GetTime() - time).str);
-			m_client.ValidateNetwork(logger);
 		}
 
 		if (!msg.ProcessAsyncResults(reader) || msg.GetError())
diff --git a/Engine/Source/Programs/UnrealBuildAccelerator/Common/Public/UbaNetworkBackend.h b/Engine/Source/Programs/UnrealBuildAccelerator/Common/Public/UbaNetworkBackend.h
--- a/Engine/Source/Programs/UnrealBuildAccelerator/Common/Public/UbaNetworkBackend.h
+++ b/Engine/Source/Programs/UnrealBuildAccelerator/Common/Public/UbaNetworkBackend.h
@@ -68,7 +68,7 @@ namespace uba
 		virtual void DeleteConnection(void* connection) = 0;
 
 		virtual void GetTotalSendAndRecv(u64& outSend, u64& outRecv) = 0;
-		virtual void Validate(Logger& logger) = 0;
+		virtual void Validate(Logger& logger, const Vector<void*>& connections, bool full) = 0;
 	};
 
 	struct NetworkBackend::SendContext
diff --git a/Engine/Source/Programs/UnrealBuildAccelerator/Common/Public/UbaNetworkBackendMemory.h b/Engine/Source/Programs/UnrealBuildAccelerator/Common/Public/UbaNetworkBackendMemory.h
--- a/Engine/Source/Programs/UnrealBuildAccelerator/Common/Public/UbaNetworkBackendMemory.h
+++ b/Engine/Source/Programs/UnrealBuildAccelerator/Common/Public/UbaNetworkBackendMemory.h
@@ -28,7 +28,7 @@ namespace uba
 		virtual void DeleteConnection(void* connection) override;
 
 		virtual void GetTotalSendAndRecv(u64& outSend, u64& outRecv) override;
-		virtual void Validate(Logger& logger) override;
+		virtual void Validate(Logger& logger, const Vector<void*>& connections, bool full) override;
 
 	private:
 		struct Connection;
diff --git a/Engine/Source/Programs/UnrealBuildAccelerator/Common/Public/UbaNetworkBackendTcp.h b/Engine/Source/Programs/UnrealBuildAccelerator/Common/Public/UbaNetworkBackendTcp.h
--- a/Engine/Source/Programs/UnrealBuildAccelerator/Common/Public/UbaNetworkBackendTcp.h
+++ b/Engine/Source/Programs/UnrealBuildAccelerator/Common/Public/UbaNetworkBackendTcp.h
@@ -10,15 +10,20 @@ namespace uba
 {
 	class Config;
 
+	////////////////////////////////////////////////////////////////////////////////////////////////////
+
 	struct NetworkBackendTcpCreateInfo
 	{
 		NetworkBackendTcpCreateInfo(LogWriter& w = g_consoleLogWriter) : logWriter(w) {}
 		LogWriter& logWriter;
+		u32 statusUpdateSeconds = 0; // Zero means never
 		bool disableNagle = true;
 
 		void Apply(Config& config, const tchar* tableName = TC("NetworkBackendTcp"));
 	};
 
+	////////////////////////////////////////////////////////////////////////////////////////////////////
+
 	class NetworkBackendTcp : public NetworkBackend
 	{
 	public:
@@ -39,7 +44,7 @@ namespace uba
 		virtual void DeleteConnection(void* connection) override;
 
 		virtual void GetTotalSendAndRecv(u64& outSend, u64& outRecv) override;
-		virtual void Validate(Logger& logger) override;
+		virtual void Validate(Logger& logger, const Vector<void*>& connections, bool full) override;
 
 	private:
 		bool EnsureInitialized(Logger& logger);
@@ -52,6 +57,11 @@ namespace uba
 		bool SendSocket(Connection& connection, Logger& logger, const void* b, u64 bufferLen, const tchar* hint);
 		bool RecvSocket(Connection& connection, RecvCache& recvCache, void* b, u32& bufferLen, const tchar* hint, bool isFirstCall, bool allowLess);
 
+		void ThreadIocp();
+		bool PostIocpRead(Connection& connection, u8* data, u32 dataSize);
+
+		void ThreadStatus(u32 statusUpdateSeconds);
+
 		LoggerWithWriter m_logger;
 		Futex m_listenEntriesLock;
 		List<ListenEntry> m_listenEntries;
@@ -67,5 +77,14 @@ namespace uba
 		#if PLATFORM_WINDOWS
 		bool m_wsaInitDone = false;
 		#endif
+
+		Thread m_iocpThreads[4];
+		void* m_iocpHandle = nullptr;
+		bool m_iocpLoop = true;
+
+		Event m_tcpStatusLoop;
+		Thread m_tcpStatusThread;
 	};
+
+	////////////////////////////////////////////////////////////////////////////////////////////////////
 };
diff --git a/Engine/Source/Programs/UnrealBuildAccelerator/Common/Public/UbaNetworkClient.h b/Engine/Source/Programs/UnrealBuildAccelerator/Common/Public/UbaNetworkClient.h
--- a/Engine/Source/Programs/UnrealBuildAccelerator/Common/Public/UbaNetworkClient.h
+++ b/Engine/Source/Programs/UnrealBuildAccelerator/Common/Public/UbaNetworkClient.h
@@ -49,7 +49,7 @@ namespace uba
 		bool IsConnected(u32 waitTimeoutMs = 0);
 		bool IsOrWasConnected(u32 waitTimeoutMs = 0);
 
-		void ValidateNetwork(Logger& logger);
+		void ValidateNetwork(Logger& logger, bool full);
 		void PrintSummary(Logger& logger);
 
 		using OnConnectedFunction = Function<void()>;
