[UBA]
* Fixed crash in UbaCli -listimports on windows

diff --git a/Engine/Source/Programs/UnrealBuildAccelerator/Cli/Private/UbaCli.cpp b/Engine/Source/Programs/UnrealBuildAccelerator/Cli/Private/UbaCli.cpp
--- a/Engine/Source/Programs/UnrealBuildAccelerator/Cli/Private/UbaCli.cpp
+++ b/Engine/Source/Programs/UnrealBuildAccelerator/Cli/Private/UbaCli.cpp
@@ -837,7 +837,7 @@ namespace uba
 				{
 					if (printImports)
 					{
-						if (*loaderPaths)
+						if (loaderPaths && *loaderPaths)
 						{
 							logger.Info(TC("LoaderPaths:"));
 							for (auto it=loaderPaths;*it; ++it)
