[UBA]
* Added more info to socket timeout logging

diff --git a/Engine/Source/Programs/UnrealBuildAccelerator/Common/Private/UbaNetworkBackendTcp.cpp b/Engine/Source/Programs/UnrealBuildAccelerator/Common/Private/UbaNetworkBackendTcp.cpp
--- a/Engine/Source/Programs/UnrealBuildAccelerator/Common/Private/UbaNetworkBackendTcp.cpp
+++ b/Engine/Source/Programs/UnrealBuildAccelerator/Common/Private/UbaNetworkBackendTcp.cpp
@@ -169,12 +169,13 @@ namespace uba
 				return logger.Error(TC("WSAStartup failed (%d)"), res);
 
 			#if UBA_USE_IOCP
-			//if (!IsRunningWine())
+			if (!IsRunningWine())
 			{
 				m_iocpHandle = CreateIoCompletionPort(INVALID_HANDLE_VALUE, NULL, 0, 0);
 
-				for (Thread& t : m_iocpThreads)
-					t.Start([this]() { ThreadIocp(); return 0; });
+				//for (Thread& t : m_iocpThreads)
+				//	t.Start([this]() { ThreadIocp(); return 0; });
+				m_iocpThreads[0].Start([this]() { ThreadIocp(); return 0; });
 			}
 			#endif
 
@@ -1253,15 +1254,18 @@ namespace uba
 
 	bool NetworkBackendTcp::PostIocpRead(Connection& connection, u8* data, u32 dataSize)
 	{
+		UBA_ASSERT(data);
 		ZeroMemory(&connection.overlapped, sizeof(OVERLAPPED));
 		connection.wsaBuf.buf = (char*)data;
 		connection.wsaBuf.len = dataSize;
 
+		SOCKET socket = connection.socket;
+
 		DWORD flags = 0, bytesRecv = 0;
-		int ret = WSARecv(connection.socket, &connection.wsaBuf, 1, &bytesRecv, &flags, &connection.overlapped, NULL);
+		int ret = WSARecv(socket, &connection.wsaBuf, 1, &bytesRecv, &flags, &connection.overlapped, NULL);
 		if (ret == SOCKET_ERROR && WSAGetLastError() != WSA_IO_PENDING)
 		{
-			printf("WSARecv failed: %d\n", WSAGetLastError());
+			m_logger.Info(TC("WSARecv failed for socket %i trying to receive %u bytes (%s)"), socket, dataSize, LastErrorToText(WSAGetLastError()).data);
 			//closesocket(clientSocket);
 			//delete context;
 			return false;
@@ -1830,7 +1834,7 @@ namespace uba
 
 		if (WSAIoctl(socket, SIO_TCP_INFO, &version, sizeof(ULONG), &ti, sizeof(TCP_INFO_v0), &bytesReturned, NULL, NULL) == SOCKET_ERROR)
 			return logger.Error(TC("WSAIoctl SIO_TCP_INFO failed (%s)"), LastErrorToText(WSAGetLastError()).data);
-		logger.Info(TC("rtt=%uµs  inFlight=%u  retrans=%u  dupacks=%u  state=%u"), ti.RttUs, ti.BytesInFlight, ti.BytesRetrans, ti.DupAcksIn, ti.State);
+		logger.Info(TC("RTT=%uµs  InFlight=%u  Retrans=%u  Dupacks=%u  State=%u RcvWnd=%u  RcvBuf=%u"), ti.RttUs, ti.BytesInFlight, ti.BytesRetrans, ti.DupAcksIn, ti.State, ti.RcvWnd, ti.RcvBuf);
 		return true;
 	#else
 		return false;
