[UBA]
* Added more tcp statistics on stalls

diff --git a/Engine/Source/Programs/UnrealBuildAccelerator/Common/Private/UbaNetworkBackendTcp.cpp b/Engine/Source/Programs/UnrealBuildAccelerator/Common/Private/UbaNetworkBackendTcp.cpp
--- a/Engine/Source/Programs/UnrealBuildAccelerator/Common/Private/UbaNetworkBackendTcp.cpp
+++ b/Engine/Source/Programs/UnrealBuildAccelerator/Common/Private/UbaNetworkBackendTcp.cpp
@@ -834,6 +834,8 @@ namespace uba
 	{
 		logger.Info(TC("  NetworkBackendTcp"));
 
+		PrintTcpStatistics(0, nullptr);
+
 		SCOPED_FUTEX(m_connectionsLock, lock);
 		for (auto ptr : connections)
 		{
@@ -904,6 +906,7 @@ namespace uba
 				u64 startTime = GetTime();
 				while (true)
 				{
+					PrintTcpStatistics(0, nullptr);
 					LogTcpInfo(logger, socket);
 					if (ev.IsSet(2000))
 						break;
@@ -1285,26 +1288,7 @@ namespace uba
 		while (!m_tcpStatusLoop.IsSet(statusUpdateMs))
 		{
 			#if PLATFORM_WINDOWS
-			MIB_TCPSTATS_LH stats{};
-			if (GetTcpStatisticsEx(&stats, AF_INET) != NO_ERROR)
-			{
-				m_logger.Info(TC("GetTcpStatisticsEx failed"));
-				return;
-			}
-
-			if (prevStats.dwRtoAlgorithm != 0)
-			{
-				INT64 recv       = stats.dwInSegs - prevStats.dwInSegs;
-				INT64 sent       = stats.dwOutSegs - prevStats.dwOutSegs;
-				INT64 retrans    = stats.dwRetransSegs - prevStats.dwRetransSegs;
-				INT64 inerrs     = stats.dwInErrs - prevStats.dwInErrs;
-				INT64 outrsts    = stats.dwOutRsts - prevStats.dwOutRsts;
-
-				double rrate = sent ? (100.0 * retrans / sent) : 0;
-
-				m_logger.Info(TC("%us delta  recv=%lld  sent=%lld  retrans=%lld  (%.2f%%)  inErr=%lld  RST=%lld"), statusUpdateSeconds, recv, sent, retrans, rrate, inerrs, outrsts);
-			}
-			prevStats = stats;
+			PrintTcpStatistics(statusUpdateSeconds, &prevStats);
 			#endif
 
 			#if 0
@@ -1324,6 +1308,43 @@ namespace uba
 		}
 	}
 
+	void NetworkBackendTcp::PrintTcpStatistics(u32 statusUpdateSeconds, void* prevStatsPtr)
+	{
+		#if PLATFORM_WINDOWS
+		auto prevStats = (MIB_TCPSTATS_LH*)prevStatsPtr;
+		MIB_TCPSTATS_LH stats{};
+		if (GetTcpStatisticsEx(&stats, AF_INET) != NO_ERROR)
+		{
+			m_logger.Info(TC("GetTcpStatisticsEx failed"));
+			return;
+		}
+
+		INT64 recv       = stats.dwInSegs;
+		INT64 sent       = stats.dwOutSegs;
+		INT64 retrans    = stats.dwRetransSegs;
+		INT64 inerrs     = stats.dwInErrs;
+		INT64 outrsts    = stats.dwOutRsts;
+
+		StringBuffer<128> temp;
+
+		if (prevStats && prevStats->dwRtoAlgorithm != 0)
+		{
+			recv       -= prevStats->dwInSegs;
+			sent       -= prevStats->dwOutSegs;
+			retrans    -= prevStats->dwRetransSegs;
+			inerrs     -= prevStats->dwInErrs;
+			outrsts    -= prevStats->dwOutRsts;
+
+			temp.Appendf(TC("%us delta"), statusUpdateSeconds);
+			*prevStats = stats;
+		}
+
+		double rrate = sent ? (100.0 * retrans / sent) : 0;
+		m_logger.Info(TC("%s  recv=%lld  sent=%lld  retrans=%lld  (%.2f%%)  inErr=%lld  RST=%lld"), temp.data, recv, sent, retrans, rrate, inerrs, outrsts);
+		#endif
+	}
+
+
 	////////////////////////////////////////////////////////////////////////////////////////////////////
 
 	void TraverseNetworkAddresses(Logger& logger, const Function<bool(const StringBufferBase& addr)>& func)
@@ -1824,7 +1845,7 @@ namespace uba
 				return logger.Error(TC("GetLinuxTcpInfo failed with error %i"), res);
 
 			u32 bytes_inflight_est = ti.tcpi_unacked * ti.tcpi_snd_mss + ti.tcpi_notsent_bytes;
-			logger.Info(TC("%p - RTT=%uus  inFlight=%u  retrans=%u  dupacks=%u  state=%u"), socket, ti.tcpi_rtt, bytes_inflight_est, ti.tcpi_retrans, ti.tcpi_snd_cwnd, ti.tcpi_state);
+			logger.Info(TC("%i - RTT=%uus  InFlight=%u  Retrans=%u (Tot %u)  Dupacks=%u  State=%u RcvWnd=%u  RcvScale=%u"), (int)socket, ti.tcpi_rtt, bytes_inflight_est, ti.tcpi_retrans, ti.tcpi_total_retrans, ti.tcpi_snd_cwnd, ti.tcpi_state, ti.tcpi_rcv_space, ti.tcpi_rcv_wscale);
 			return true;
 		}
 
diff --git a/Engine/Source/Programs/UnrealBuildAccelerator/Common/Public/UbaNetworkBackendTcp.h b/Engine/Source/Programs/UnrealBuildAccelerator/Common/Public/UbaNetworkBackendTcp.h
--- a/Engine/Source/Programs/UnrealBuildAccelerator/Common/Public/UbaNetworkBackendTcp.h
+++ b/Engine/Source/Programs/UnrealBuildAccelerator/Common/Public/UbaNetworkBackendTcp.h
@@ -61,6 +61,7 @@ namespace uba
 		bool PostIocpRead(Connection& connection, u8* data, u32 dataSize);
 
 		void ThreadStatus(u32 statusUpdateSeconds);
+		void PrintTcpStatistics(u32 statusUpdateSeconds, void* prevStats);
 
 		LoggerWithWriter m_logger;
 		Futex m_listenEntriesLock;
