[UBA]
* Fixed bug in uba where it treated files that was used for reading as throwaway files (causing pak compression path to fail)

diff --git a/Engine/Source/Programs/UnrealBuildAccelerator/Detours/Private/Windows/UbaDetoursFunctionsKernelBase.inl b/Engine/Source/Programs/UnrealBuildAccelerator/Detours/Private/Windows/UbaDetoursFunctionsKernelBase.inl
--- a/Engine/Source/Programs/UnrealBuildAccelerator/Detours/Private/Windows/UbaDetoursFunctionsKernelBase.inl
+++ b/Engine/Source/Programs/UnrealBuildAccelerator/Detours/Private/Windows/UbaDetoursFunctionsKernelBase.inl
@@ -376,7 +376,7 @@ BOOL Detoured_ReadFile(HANDLE hFile, LPVOID lpBuffer, DWORD nNumberOfBytesToRead
 				DEBUG_LOG_DETOURED(L"ReadFile", L"%llu %u (%ls) -> FAILED TO MAP", uintptr_t(hFile), nNumberOfBytesToRead, HandleToName(hFile));
 				return FALSE;
 			}
-			UBA_ASSERT(fi.fileMapMem || !fi.memoryFile->isThrowAway);
+			UBA_ASSERTF(fi.fileMapMem || !fi.memoryFile->isThrowAway, TC("Trying to read throw-away file %s"), HandleToName(hFile));
 
 			u8* mem = fi.fileMapMem ? fi.fileMapMem : fi.memoryFile->baseAddress;
 			u64 size = fi.fileMapMem ? fi.fileMapMemSize : fi.memoryFile->writtenSize;
diff --git a/Engine/Source/Programs/UnrealBuildAccelerator/Detours/Private/Windows/UbaDetoursFunctionsNtDll.inl b/Engine/Source/Programs/UnrealBuildAccelerator/Detours/Private/Windows/UbaDetoursFunctionsNtDll.inl
--- a/Engine/Source/Programs/UnrealBuildAccelerator/Detours/Private/Windows/UbaDetoursFunctionsNtDll.inl
+++ b/Engine/Source/Programs/UnrealBuildAccelerator/Detours/Private/Windows/UbaDetoursFunctionsNtDll.inl
@@ -856,7 +856,7 @@ NTSTATUS NTAPI Shared_NtCreateFile(bool IsCreateFunc, PHANDLE hFileHandle, ACCES
 	bool useContent = IsContentUse(DesiredAccess, CreateDisposition);
 	bool isWrite = IsWrite(DesiredAccess, CreateDisposition);
 	bool isThrowAway = g_rules->IsThrowAway(fileName, g_runningRemote);
-	bool keepInMemory = KeepInMemory(fileName, isWrite) || ((isWrite || isDeleteOnClose) && g_rules->IsOutputFile(fileName)) || isThrowAway;
+	bool keepInMemory = KeepInMemory(fileName, isWrite) || ((isWrite || isDeleteOnClose) && g_rules->IsOutputFile(fileName)) || (isWrite && isThrowAway);
 
 #if UBA_DEBUG_LOG_ENABLED
 	const wchar_t* isWriteStr = isWrite ? L" WRITE" : L""; (void)isWriteStr;
