[UBA]
* Fixed so realpath is using directory table to check if path exists and returns nullptr and errno ENOENT if it doesn't

diff --git a/Engine/Source/Programs/UnrealBuildAccelerator/Detours/Private/PosixOS/UbaDetoursFunctionsPosix.cpp b/Engine/Source/Programs/UnrealBuildAccelerator/Detours/Private/PosixOS/UbaDetoursFunctionsPosix.cpp
--- a/Engine/Source/Programs/UnrealBuildAccelerator/Detours/Private/PosixOS/UbaDetoursFunctionsPosix.cpp
+++ b/Engine/Source/Programs/UnrealBuildAccelerator/Detours/Private/PosixOS/UbaDetoursFunctionsPosix.cpp
@@ -988,6 +988,22 @@ UBA_EXPORT char* UBA_WRAPPER(realpath)(const char* path, char* resolved_path)
 	// TODO
 	StringBuffer<> fixedPath;
 	FixPath(fixedPath, path);
+
+	FileAttributes fileAttr;
+	const char* realName = Shared_GetFileAttributes(fileAttr, fixedPath.data);(void)realName;
+	if (!fileAttr.useCache)
+	{
+		auto res = TRUE_WRAPPER(realpath)(path, resolved_path);
+		DEBUG_LOG_TRUE("realpath", "(%s) -> %s (%s)", path, res, StrError(res == 0 ? -1 : 0, errno));
+		return res;
+	}
+
+	if (!fileAttr.exists)
+	{
+		errno = ENOENT;
+		return nullptr;
+	}
+
 	if (!resolved_path)
 		resolved_path = (char*)malloc(fixedPath.count + 1);
 	memcpy(resolved_path, fixedPath.data, fixedPath.count + 1);
@@ -1420,8 +1436,9 @@ UBA_EXPORT int UBA_WRAPPER(closedir)(DIR* dirp)
 UBA_EXPORT int UBA_WRAPPER(glob)(const char* pattern, int flags, int (*errfunc)(const char* epath, int eerrno), glob_t* pglob)
 {
 	UBA_INIT_DETOUR(glob, pattern, flags, errfunc, pglob);
-	DEBUG_LOG_TRUE("glob", "%s", pattern);
-	return TRUE_WRAPPER(glob)(pattern, flags, errfunc, pglob);
+	int res = TRUE_WRAPPER(glob)(pattern, flags, errfunc, pglob);
+	DEBUG_LOG_TRUE("glob", "%s -> %i", pattern, res);
+	return res;
 }
 
 #if UBA_DEBUG && PLATFORM_LINUX
@@ -1454,8 +1471,9 @@ UBA_EXPORT long UBA_WRAPPER(pathconf)(char *path, int name)
 UBA_EXPORT int UBA_WRAPPER(lstat)(const char *path, struct stat *buf)
 {
 	UBA_INIT_DETOUR(lstat, path, buf);
-	DEBUG_LOG_TRUE("lstat", "");
-	return TRUE_WRAPPER(lstat)(path, buf);
+	int res = TRUE_WRAPPER(lstat)(path, buf);
+	DEBUG_LOG_TRUE("lstat", "%s -> %i", path, res);
+	return res;
 }
 
 UBA_EXPORT int UBA_WRAPPER(stat)(const char* file, struct stat* attr)
