[UBA]
* Added option to enable/disable overlapped WSASend on windows using useOverlappedSend option both via ubaagent cmd line and config file
* Added option to enable/disable iocp on windows with "useIocp" option/cmd line.
* Reduced spam when tcp is timing out

diff --git a/Engine/Source/Programs/UnrealBuildAccelerator/Agent/Private/UbaAgent.cpp b/Engine/Source/Programs/UnrealBuildAccelerator/Agent/Private/UbaAgent.cpp
--- a/Engine/Source/Programs/UnrealBuildAccelerator/Agent/Private/UbaAgent.cpp
+++ b/Engine/Source/Programs/UnrealBuildAccelerator/Agent/Private/UbaAgent.cpp
@@ -120,6 +120,9 @@ namespace uba
 		logger.Info(TC("  -killtcphogs            If failing to bind listen socket UbaAgent will attempt to kill processes holding it and then retry"));
 		logger.Info(TC("  -populateCasFromXcodeVersion=<version>   Prepopulate cas database with files from local xcode installation that matches the version."));
 		logger.Info(TC("  -populateCasFromAllXcodes   Prepopulate cas database with files from local xcode installation that matches the version."));
+		#elif PLATFORM_WINDOWS
+		logger.Info(TC("  -useOverlappedSend      Enable/Disable overlapped send for tcp"));
+		logger.Info(TC("  -useIocp                Enable/Disable iocp for tcp"));
 		#endif
 		logger.Info(TC(""));
 		return false;
@@ -371,6 +374,8 @@ namespace uba
 		bool printSummary = false;
 		bool killRandom = false;
 		bool useCrawler = false;
+		bool useOverlappedSend = true;
+		bool useIocp = false;
 		bool downloadDetoursLib = true;
 		bool useExceptionHandler = true;
 		bool resendCas = false;
@@ -398,10 +403,10 @@ namespace uba
 			StringBuffer<> name;
 			StringBuffer<> value;
 
-			auto parseValue = [&](auto& out)
+			auto parseValue = [&](auto& out, bool allowEmpty = false)
 				{
 					if (value.IsEmpty())
-						return PrintHelp(StringBuffer().Appendf(TC("%s needs a value"), name.data).data);
+						return allowEmpty ? true : PrintHelp(StringBuffer().Appendf(TC("%s needs a value"), name.data).data);
 					if (!ExpandEnvironmentVariables(value, PrintHelp))
 						return false;
 					std::decay_t<decltype(out)> temp;
@@ -602,6 +607,18 @@ namespace uba
 			{
 				useCrawler = true;
 			}
+			#if PLATFORM_WINDOWS
+			else if (name.Equals(TCV("-useOverlappedSend")))
+			{
+				if (!parseValue(useOverlappedSend, true))
+					return false;
+			}
+			else if (name.Equals(TCV("-useIocp")))
+			{
+				if (!parseValue(useIocp, true))
+					return false;
+			}
+			#endif
 			else if (name.Equals(TCV("-memwait")))
 			{
 				if (!parseValue(memWaitLoadPercent))
@@ -1026,6 +1043,8 @@ namespace uba
 			#endif
 			{
 				NetworkBackendTcpCreateInfo info(logWriter);
+				info.useOverlappedSend = useOverlappedSend;
+				info.useIocp = useIocp;
 				//info.statusUpdateSeconds = 2;
 				networkBackend = new NetworkBackendTcp(info);
 			}
diff --git a/Engine/Source/Programs/UnrealBuildAccelerator/Common/Private/UbaNetworkBackendTcp.cpp b/Engine/Source/Programs/UnrealBuildAccelerator/Common/Private/UbaNetworkBackendTcp.cpp
--- a/Engine/Source/Programs/UnrealBuildAccelerator/Common/Private/UbaNetworkBackendTcp.cpp
+++ b/Engine/Source/Programs/UnrealBuildAccelerator/Common/Private/UbaNetworkBackendTcp.cpp
@@ -45,7 +45,7 @@
 
 #define UBA_USE_OVERLAPPED_SEND PLATFORM_WINDOWS
 #define UBA_USE_OVERLAPPED_SEND_WITH_LOCK 0
-#define UBA_USE_IOCP 0
+#define UBA_USE_IOCP PLATFORM_WINDOWS
 
 namespace uba
 {
@@ -113,6 +113,8 @@ namespace uba
 		void* disconnectContext = nullptr;
 		DisconnectCallback* disconnectCallback = nullptr;
 
+		Futex timeoutLock;
+
 		Futex sendLock;
 		CriticalSection shutdownLock;
 
@@ -155,6 +157,8 @@ namespace uba
 		const ConfigTable& table = *tablePtr;
 		table.GetValueAsBool(disableNagle, TC("DisableNagle"));
 		table.GetValueAsU32(statusUpdateSeconds, TC("StatusUpdateSeconds"));
+		table.GetValueAsBool(useOverlappedSend, TC("UseOverlappedSend"));
+		table.GetValueAsBool(useIocp, TC("UseIocp"));
 	}
 
 	////////////////////////////////////////////////////////////////////////////////////////////////////
@@ -169,7 +173,7 @@ namespace uba
 				return logger.Error(TC("WSAStartup failed (%d)"), res);
 
 			#if UBA_USE_IOCP
-			if (!IsRunningWine())
+			if (m_useIocp && !IsRunningWine())
 			{
 				m_iocpHandle = CreateIoCompletionPort(INVALID_HANDLE_VALUE, NULL, 0, 0);
 
@@ -198,6 +202,8 @@ namespace uba
 		: m_logger(info.logWriter, prefix)
 	{
 		m_disableNagle = info.disableNagle;
+		m_useOverlappedSend = info.useOverlappedSend;
+		m_useIocp = info.useIocp;
 
 		if (info.statusUpdateSeconds)
 		{
@@ -834,7 +840,7 @@ namespace uba
 	{
 		logger.Info(TC("  NetworkBackendTcp"));
 
-		PrintTcpStatistics(0, nullptr);
+		PrintTcpStatistics(logger, 0, nullptr);
 
 		SCOPED_FUTEX(m_connectionsLock, lock);
 		for (auto ptr : connections)
@@ -856,96 +862,101 @@ namespace uba
 
 
 #if UBA_USE_OVERLAPPED_SEND
-
-		Event ev(true);
-
-		// Right now we are experimenting with if we can ignore taking a lock around the entire thing and only the WSASend.
-		// Documentation is slightly unclear and network forums claim the order of data is correct even though multiple
-		// threads call WSASend and then wait on event.
-		#if UBA_USE_OVERLAPPED_SEND_WITH_LOCK
-		SCOPED_FUTEX(connection.sendLock, lock);
-		#endif
-
-		char* buffer = (char*)b;
-		u64 left = bufferLen;
-		while (left)
+		if (m_useOverlappedSend)
 		{
-			Overlapped overlapped {};
-			overlapped.isSend = true;
-			overlapped.hEvent = ev.GetHandle();
-			WSABUF buf { u32(left), buffer };
+			Event ev(true);
 
-			#if !UBA_USE_OVERLAPPED_SEND_WITH_LOCK
+			// Right now we are experimenting with if we can ignore taking a lock around the entire thing and only the WSASend.
+			// Documentation is slightly unclear and network forums claim the order of data is correct even though multiple
+			// threads call WSASend and then wait on event.
+			#if UBA_USE_OVERLAPPED_SEND_WITH_LOCK
 			SCOPED_FUTEX(connection.sendLock, lock);
 			#endif
 
-			SOCKET socket = connection.socket;
-			if (socket == INVALID_SOCKET)
-				return false;
-			int res = WSASend(socket, &buf, 1, NULL, 0, &overlapped, NULL);
+			char* buffer = (char*)b;
+			u64 left = bufferLen;
+			while (left)
+			{
+				Overlapped overlapped {};
+				overlapped.isSend = true;
+				overlapped.hEvent = ev.GetHandle();
+				WSABUF buf { u32(left), buffer };
 
-			#if !UBA_USE_OVERLAPPED_SEND_WITH_LOCK
-			lock.Leave();
-			#endif
+				#if !UBA_USE_OVERLAPPED_SEND_WITH_LOCK
+				SCOPED_FUTEX(connection.sendLock, lock);
+				#endif
 
-			if (res == SOCKET_ERROR)
-			{
-				u32 lastError = WSAGetLastError();
-				if (lastError != WSA_IO_PENDING)
-				{
-					#if UBA_LOG_SOCKET_ERRORS
-					logger.Info(TC("WSASend - error for socket %i: %s (%s)"), socket, LastErrorToText(lastError).data, hint);
-					#endif
+				SOCKET socket = connection.socket;
+				if (socket == INVALID_SOCKET)
 					return false;
+				int res = WSASend(socket, &buf, 1, NULL, 0, &overlapped, NULL);
+
+				#if !UBA_USE_OVERLAPPED_SEND_WITH_LOCK
+				lock.Leave();
+				#endif
+
+				if (res == SOCKET_ERROR)
+				{
+					u32 lastError = WSAGetLastError();
+					if (lastError != WSA_IO_PENDING)
+					{
+						#if UBA_LOG_SOCKET_ERRORS
+						logger.Info(TC("WSASend - error for socket %i: %s (%s)"), socket, LastErrorToText(lastError).data, hint);
+						#endif
+						return false;
+					}
 				}
-			}
 
-			if (!ev.IsSet(40*1000))
-			{
-				logger.Info(TC("SendSocket - WSASend returned timeout after 40 seconds for socket %i (%s)"), socket, hint);
-				u64 startTime = GetTime();
-				while (true)
+				if (!ev.IsSet(38*1000))
 				{
-					PrintTcpStatistics(0, nullptr);
-					LogTcpInfo(logger, socket);
-					if (ev.IsSet(2000))
+					SCOPED_FUTEX(connection.timeoutLock, timeoutLock);
+					if (ev.IsSet(2*1000))
 						break;
-
-					u64 sinceStartSeconds = TimeToMs(GetTime() - startTime)/1000;
-					if (sinceStartSeconds >= DefaultNetworkSendTimeoutSeconds)
+					logger.Info(TC("SendSocket - WSASend returned timeout after 40 seconds for socket %i (%s)"), socket, hint);
+					u64 startTime = GetTime();
+					while (true)
 					{
-						logger.Info(TC("SendSocket - WSASend returned timeout after 10 minutes for socket %i (%s)"), socket, hint);
-						return false;
+						PrintTcpStatistics(logger, 0, nullptr);
+						LogTcpInfo(logger, socket);
+
+						if (ev.IsSet(4*1000))
+							break;
+
+						u64 sinceStartSeconds = TimeToMs(GetTime() - startTime)/1000;
+						if (sinceStartSeconds >= DefaultNetworkSendTimeoutSeconds)
+						{
+							logger.Info(TC("SendSocket - WSASend returned timeout after 10 minutes for socket %i (%s)"), socket, hint);
+							return false;
+						}
 					}
 				}
-			}
 
-			DWORD bytesSent;
-			DWORD flags;
-			if (!WSAGetOverlappedResult(socket, &overlapped, &bytesSent, FALSE, &flags))
-			{
-				#if UBA_LOG_SOCKET_ERRORS
-				logger.Info(TC("WSAGetOverlappedResult - error for socket %i: %s (%s)"), socket, LastErrorToText().data, hint);
-				#endif
-				return false;
-			}
+				DWORD bytesSent;
+				DWORD flags;
+				if (!WSAGetOverlappedResult(socket, &overlapped, &bytesSent, FALSE, &flags))
+				{
+					#if UBA_LOG_SOCKET_ERRORS
+					logger.Info(TC("WSAGetOverlappedResult - error for socket %i: %s (%s)"), socket, LastErrorToText().data, hint);
+					#endif
+					return false;
+				}
 
-			buffer += bytesSent;
-			left -= bytesSent;
-			if (left)
-			{
-				#if !UBA_USE_OVERLAPPED_SEND_WITH_LOCK
-				//#if UBA_LOG_SOCKET_ERRORS
-				logger.Warning(TC("SendSocket - WSASend did not send all data in one call for socket %i. Note that this might be allowed behavior. Hard to read out of documentation. (%s)"), socket, hint);
-				//#endif
-				return false;
-				#else
-				ev.Reset();
-				#endif
+				buffer += bytesSent;
+				left -= bytesSent;
+				if (left)
+				{
+					#if !UBA_USE_OVERLAPPED_SEND_WITH_LOCK
+					//#if UBA_LOG_SOCKET_ERRORS
+					logger.Warning(TC("SendSocket - WSASend did not send all data in one call for socket %i. Note that this might be allowed behavior. Hard to read out of documentation. (%s)"), socket, hint);
+					//#endif
+					return false;
+					#else
+					ev.Reset();
+					#endif
+				}
 			}
 		}
-
-#else
+#endif
 
 		SCOPED_FUTEX(connection.sendLock, lock);
 		SOCKET socket = connection.socket;
@@ -1002,7 +1013,6 @@ namespace uba
 			buffer += sent;
 			left -= sent;
 		}
-#endif
 		return true;
 	}
 
@@ -1268,7 +1278,7 @@ namespace uba
 		int ret = WSARecv(socket, &connection.wsaBuf, 1, &bytesRecv, &flags, &connection.overlapped, NULL);
 		if (ret == SOCKET_ERROR && WSAGetLastError() != WSA_IO_PENDING)
 		{
-			m_logger.Info(TC("WSARecv failed for socket %i trying to receive %u bytes (%s)"), socket, dataSize, LastErrorToText(WSAGetLastError()).data);
+			connection.logger.Info(TC("WSARecv failed for socket %i trying to receive %u bytes (%s)"), socket, dataSize, LastErrorToText(WSAGetLastError()).data);
 			//closesocket(clientSocket);
 			//delete context;
 			return false;
@@ -1288,7 +1298,7 @@ namespace uba
 		while (!m_tcpStatusLoop.IsSet(statusUpdateMs))
 		{
 			#if PLATFORM_WINDOWS
-			PrintTcpStatistics(statusUpdateSeconds, &prevStats);
+			PrintTcpStatistics(m_logger, statusUpdateSeconds, &prevStats);
 			#endif
 
 			#if 0
@@ -1308,14 +1318,14 @@ namespace uba
 		}
 	}
 
-	void NetworkBackendTcp::PrintTcpStatistics(u32 statusUpdateSeconds, void* prevStatsPtr)
+	void NetworkBackendTcp::PrintTcpStatistics(Logger& logger, u32 statusUpdateSeconds, void* prevStatsPtr)
 	{
 		#if PLATFORM_WINDOWS
 		auto prevStats = (MIB_TCPSTATS_LH*)prevStatsPtr;
 		MIB_TCPSTATS_LH stats{};
 		if (GetTcpStatisticsEx(&stats, AF_INET) != NO_ERROR)
 		{
-			m_logger.Info(TC("GetTcpStatisticsEx failed"));
+			logger.Info(TC("GetTcpStatisticsEx failed"));
 			return;
 		}
 
@@ -1340,7 +1350,7 @@ namespace uba
 		}
 
 		double rrate = sent ? (100.0 * retrans / sent) : 0;
-		m_logger.Info(TC("%s  recv=%lld  sent=%lld  retrans=%lld  (%.2f%%)  inErr=%lld  RST=%lld"), temp.data, recv, sent, retrans, rrate, inerrs, outrsts);
+		logger.Info(TC("%s  recv=%lld  sent=%lld  retrans=%lld  (%.2f%%)  inErr=%lld  RST=%lld"), temp.data, recv, sent, retrans, rrate, inerrs, outrsts);
 		#endif
 	}
 
@@ -1849,13 +1859,12 @@ namespace uba
 			return true;
 		}
 
-		TCP_INFO_v0 ti;
-		ULONG version = 0; // Specify 0 to retrieve the v0 version of this structure.
-		ULONG bytesReturned;
-
-		if (WSAIoctl(socket, SIO_TCP_INFO, &version, sizeof(ULONG), &ti, sizeof(TCP_INFO_v0), &bytesReturned, NULL, NULL) == SOCKET_ERROR)
-			return logger.Error(TC("WSAIoctl SIO_TCP_INFO failed (%s)"), LastErrorToText(WSAGetLastError()).data);
-		logger.Info(TC("RTT=%uµs  InFlight=%u  Retrans=%u  Dupacks=%u  State=%u RcvWnd=%u  RcvBuf=%u"), ti.RttUs, ti.BytesInFlight, ti.BytesRetrans, ti.DupAcksIn, ti.State, ti.RcvWnd, ti.RcvBuf);
+		//TCP_INFO_v0 ti;
+		//ULONG version = 0; // Specify 0 to retrieve the v0 version of this structure.
+		//ULONG bytesReturned;
+		//if (WSAIoctl(socket, SIO_TCP_INFO, &version, sizeof(ULONG), &ti, sizeof(TCP_INFO_v0), &bytesReturned, NULL, NULL) == SOCKET_ERROR)
+		//	return logger.Error(TC("WSAIoctl SIO_TCP_INFO failed (%s)"), LastErrorToText(WSAGetLastError()).data);
+		//logger.Info(TC("RTT=%uµs  InFlight=%u  Retrans=%u  Dupacks=%u  State=%u RcvWnd=%u  RcvBuf=%u"), ti.RttUs, ti.BytesInFlight, ti.BytesRetrans, ti.DupAcksIn, ti.State, ti.RcvWnd, ti.RcvBuf);
 		return true;
 	#else
 		return false;
diff --git a/Engine/Source/Programs/UnrealBuildAccelerator/Common/Public/UbaNetworkBackendTcp.h b/Engine/Source/Programs/UnrealBuildAccelerator/Common/Public/UbaNetworkBackendTcp.h
--- a/Engine/Source/Programs/UnrealBuildAccelerator/Common/Public/UbaNetworkBackendTcp.h
+++ b/Engine/Source/Programs/UnrealBuildAccelerator/Common/Public/UbaNetworkBackendTcp.h
@@ -18,6 +18,8 @@ namespace uba
 		LogWriter& logWriter;
 		u32 statusUpdateSeconds = 0; // Zero means never
 		bool disableNagle = true;
+		bool useOverlappedSend = true;
+		bool useIocp = false;
 
 		void Apply(Config& config, const tchar* tableName = TC("NetworkBackendTcp"));
 	};
@@ -61,7 +63,7 @@ namespace uba
 		bool PostIocpRead(Connection& connection, u8* data, u32 dataSize);
 
 		void ThreadStatus(u32 statusUpdateSeconds);
-		void PrintTcpStatistics(u32 statusUpdateSeconds, void* prevStats);
+		void PrintTcpStatistics(Logger& logger, u32 statusUpdateSeconds, void* prevStats);
 
 		LoggerWithWriter m_logger;
 		Futex m_listenEntriesLock;
@@ -74,6 +76,8 @@ namespace uba
 		Atomic<u64> m_totalRecv;
 
 		bool m_disableNagle;
+		bool m_useOverlappedSend;
+		bool m_useIocp;
 
 		#if PLATFORM_WINDOWS
 		bool m_wsaInitDone = false;
diff --git a/Engine/Source/Programs/UnrealBuildAccelerator/Core/Private/UbaStringBuffer.cpp b/Engine/Source/Programs/UnrealBuildAccelerator/Core/Private/UbaStringBuffer.cpp
--- a/Engine/Source/Programs/UnrealBuildAccelerator/Core/Private/UbaStringBuffer.cpp
+++ b/Engine/Source/Programs/UnrealBuildAccelerator/Core/Private/UbaStringBuffer.cpp
@@ -392,6 +392,21 @@ namespace uba
 		return true;
 	}
 
+	bool StringBufferBase::Parse(bool& out) const
+	{
+		if (Equals(TCV("true")) || Equals(TCV("1")))
+		{
+			out = true;
+			return true;
+		}
+		if (Equals(TCV("false")) || Equals(TCV("0")))
+		{
+			out = false;
+			return true;
+		}
+		return false;
+	}
+
 	bool StringBufferBase::Parse(float& out) const
 	{
 		if (!count)
diff --git a/Engine/Source/Programs/UnrealBuildAccelerator/Core/Public/UbaStringBuffer.h b/Engine/Source/Programs/UnrealBuildAccelerator/Core/Public/UbaStringBuffer.h
--- a/Engine/Source/Programs/UnrealBuildAccelerator/Core/Public/UbaStringBuffer.h
+++ b/Engine/Source/Programs/UnrealBuildAccelerator/Core/Public/UbaStringBuffer.h
@@ -73,6 +73,7 @@ namespace uba
 		bool Parse(u64& out) const;
 		bool Parse(u32& out, u64 offset = 0) const;
 		bool Parse(u16& out, u64 offset = 0) const;
+		bool Parse(bool& out) const;
 		bool Parse(float& out) const;
 		bool Parse(TString& out, u64 offset = 0) const;
 		bool Parse(StringBufferBase& out, u64 offset = 0) const;
