[UBA]
* Fixed so realpath handles vfs when testing if file exists

diff --git a/Engine/Source/Programs/UnrealBuildAccelerator/Detours/Private/PosixOS/UbaDetoursFunctionsPosix.cpp b/Engine/Source/Programs/UnrealBuildAccelerator/Detours/Private/PosixOS/UbaDetoursFunctionsPosix.cpp
--- a/Engine/Source/Programs/UnrealBuildAccelerator/Detours/Private/PosixOS/UbaDetoursFunctionsPosix.cpp
+++ b/Engine/Source/Programs/UnrealBuildAccelerator/Detours/Private/PosixOS/UbaDetoursFunctionsPosix.cpp
@@ -989,8 +989,11 @@ UBA_EXPORT char* UBA_WRAPPER(realpath)(const char* path, char* resolved_path)
 	StringBuffer<> fixedPath;
 	FixPath(fixedPath, path);
 
+	StringBuffer<> devirtualizedPath(fixedPath);
+	DevirtualizePath(devirtualizedPath);
+
 	FileAttributes fileAttr;
-	const char* realName = Shared_GetFileAttributes(fileAttr, fixedPath.data);(void)realName;
+	const char* realName = Shared_GetFileAttributes(fileAttr, devirtualizedPath.data);(void)realName;
 	if (!fileAttr.useCache)
 	{
 		auto res = TRUE_WRAPPER(realpath)(path, resolved_path);
@@ -1001,6 +1004,7 @@ UBA_EXPORT char* UBA_WRAPPER(realpath)(const char* path, char* resolved_path)
 	if (!fileAttr.exists)
 	{
 		errno = ENOENT;
+		DEBUG_LOG_DETOURED("realpath", "(%s) -> nullptr (ENOENT)", path);
 		return nullptr;
 	}
 
