[UBA]
* Enabled full lock around WSASend

diff --git a/Engine/Source/Programs/UnrealBuildAccelerator/Common/Private/UbaNetworkBackendTcp.cpp b/Engine/Source/Programs/UnrealBuildAccelerator/Common/Private/UbaNetworkBackendTcp.cpp
--- a/Engine/Source/Programs/UnrealBuildAccelerator/Common/Private/UbaNetworkBackendTcp.cpp
+++ b/Engine/Source/Programs/UnrealBuildAccelerator/Common/Private/UbaNetworkBackendTcp.cpp
@@ -44,7 +44,7 @@
 #define UBA_EMULATE_BAD_INTERNET 0
 
 #define UBA_USE_OVERLAPPED_SEND PLATFORM_WINDOWS
-#define UBA_USE_OVERLAPPED_SEND_WITH_LOCK 0
+#define UBA_USE_OVERLAPPED_SEND_WITH_LOCK 1
 #define UBA_USE_IOCP PLATFORM_WINDOWS
 
 namespace uba
@@ -113,7 +113,9 @@ namespace uba
 		void* disconnectContext = nullptr;
 		DisconnectCallback* disconnectCallback = nullptr;
 
+		#if !UBA_USE_OVERLAPPED_SEND_WITH_LOCK
 		Futex timeoutLock;
+		#endif
 
 		Futex sendLock;
 		CriticalSection shutdownLock;
@@ -909,7 +911,10 @@ namespace uba
 
 				if (!ev.IsSet(38*1000))
 				{
+					#if !UBA_USE_OVERLAPPED_SEND_WITH_LOCK
 					SCOPED_FUTEX(connection.timeoutLock, timeoutLock);
+					#endif
+
 					if (ev.IsSet(2*1000))
 						break;
 					logger.Info(TC("SendSocket - WSASend returned timeout after 40 seconds for socket %i (%s)"), socket, hint);
