[UBA]
* Fixed bug in overlapped send logic where it sent data twice

diff --git a/Engine/Source/Programs/UnrealBuildAccelerator/Common/Private/UbaNetworkBackendTcp.cpp b/Engine/Source/Programs/UnrealBuildAccelerator/Common/Private/UbaNetworkBackendTcp.cpp
--- a/Engine/Source/Programs/UnrealBuildAccelerator/Common/Private/UbaNetworkBackendTcp.cpp
+++ b/Engine/Source/Programs/UnrealBuildAccelerator/Common/Private/UbaNetworkBackendTcp.cpp
@@ -44,7 +44,7 @@
 #define UBA_EMULATE_BAD_INTERNET 0
 
 #define UBA_USE_OVERLAPPED_SEND PLATFORM_WINDOWS
-#define UBA_USE_OVERLAPPED_SEND_WITH_LOCK 1
+#define UBA_USE_OVERLAPPED_SEND_WITH_LOCK 0
 #define UBA_USE_IOCP PLATFORM_WINDOWS
 
 namespace uba
@@ -960,6 +960,7 @@ namespace uba
 					#endif
 				}
 			}
+			return true;
 		}
 #endif
 
