[UBA]
* Fixed bug in native uba scheduler where the logic for fetching from cache was wrong and returned cache hit even though it was a miss

diff --git a/Engine/Source/Programs/UnrealBuildAccelerator/Common/Private/UbaScheduler.cpp b/Engine/Source/Programs/UnrealBuildAccelerator/Common/Private/UbaScheduler.cpp
--- a/Engine/Source/Programs/UnrealBuildAccelerator/Common/Private/UbaScheduler.cpp
+++ b/Engine/Source/Programs/UnrealBuildAccelerator/Common/Private/UbaScheduler.cpp
@@ -523,7 +523,7 @@ namespace uba
 						bool isHit = false;
 						for (auto& cacheClient : m_cacheClients)
 						{
-							if (!cacheClient->FetchFromCache(cacheResult, si.rootsHandle, si.cacheBucketId, si) && cacheResult.hit)
+							if (!cacheClient->FetchFromCache(cacheResult, si.rootsHandle, si.cacheBucketId, si) || !cacheResult.hit)
 								continue;
 
 							isHit = true;
