[UBA]
* Finished iocp implementation. Default is disabled because requires testing

diff --git a/Engine/Source/Programs/UnrealBuildAccelerator/Agent/Private/UbaAgent.cpp b/Engine/Source/Programs/UnrealBuildAccelerator/Agent/Private/UbaAgent.cpp
--- a/Engine/Source/Programs/UnrealBuildAccelerator/Agent/Private/UbaAgent.cpp
+++ b/Engine/Source/Programs/UnrealBuildAccelerator/Agent/Private/UbaAgent.cpp
@@ -122,7 +122,7 @@ namespace uba
 		logger.Info(TC("  -populateCasFromAllXcodes   Prepopulate cas database with files from local xcode installation that matches the version."));
 		#elif PLATFORM_WINDOWS
 		logger.Info(TC("  -useOverlappedSend      Enable/Disable overlapped send for tcp"));
-		logger.Info(TC("  -useIocp                Enable/Disable iocp for tcp"));
+		logger.Info(TC("  -useIocp[=workerCount]  Enable/Disable iocp for tcp. Defaults to 4 workers is not set"));
 		#endif
 		logger.Info(TC(""));
 		return false;
@@ -375,7 +375,7 @@ namespace uba
 		bool killRandom = false;
 		bool useCrawler = false;
 		bool useOverlappedSend = true;
-		bool useIocp = false;
+		u32 iocpWorkerCount = 0;
 		bool downloadDetoursLib = true;
 		bool useExceptionHandler = true;
 		bool resendCas = false;
@@ -615,7 +615,8 @@ namespace uba
 			}
 			else if (name.Equals(TCV("-useIocp")))
 			{
-				if (!parseValue(useIocp, true))
+				iocpWorkerCount = 4;
+				if (!parseValue(iocpWorkerCount, true))
 					return false;
 			}
 			#endif
@@ -1044,7 +1045,7 @@ namespace uba
 			{
 				NetworkBackendTcpCreateInfo info(logWriter);
 				info.useOverlappedSend = useOverlappedSend;
-				info.useIocp = useIocp;
+				info.iocpWorkerCount = iocpWorkerCount;
 				//info.statusUpdateSeconds = 2;
 				networkBackend = new NetworkBackendTcp(info);
 			}
@@ -1314,7 +1315,7 @@ namespace uba
 						logLines[logLinesIndexPrev].clear();
 					}
 					return 0;
-				});
+				}, TC("UbaLogging"));
 
 			g_sessionClient = sessionClient;
 
diff --git a/Engine/Source/Programs/UnrealBuildAccelerator/CacheService/Private/UbaCacheService.cpp b/Engine/Source/Programs/UnrealBuildAccelerator/CacheService/Private/UbaCacheService.cpp
--- a/Engine/Source/Programs/UnrealBuildAccelerator/CacheService/Private/UbaCacheService.cpp
+++ b/Engine/Source/Programs/UnrealBuildAccelerator/CacheService/Private/UbaCacheService.cpp
@@ -65,6 +65,8 @@ namespace uba
 		logger.Info(TC("  -reportinterval=<sec>   How often the service should report status. Defaults to \"%s\""), TimeToText(MsToTime(DefaultReportIntervalSeconds*1000)).str);
 #if PLATFORM_LINUX
 		logger.Info(TC("  -fork                   Will handle segfaults and restart"));
+#elif PLATFORM_WINDOWS
+		logger.Info(TC("  -useIocp[=workerCount]  Enable/Disable iocp for tcp. Defaults to 4 workers is not set"));
 #endif
 		logger.Info(TC(" Example of how to register crypto key to cache server (when -http=80 is provided)"));
 		logger.Info(TC("   curl http://localhost/addcrypto?3f58aa57466db9999213456789123445"));
@@ -126,6 +128,7 @@ namespace uba
 		u32 maxWorkerCount = DefaultWorkerCount;
 		u32 reportIntervalSeconds = DefaultReportIntervalSeconds;
 		TString configFile;
+		u32 iocpWorkerCount = 0;
 
 		#if PLATFORM_LINUX
 		bool forkProcess = false;
@@ -200,6 +203,13 @@ namespace uba
 			{
 				forkProcess = true;
 			}
+			#elif PLATFORM_WINDOWS
+			else if (name.Equals(TCV("-useIocp")))
+			{
+				iocpWorkerCount = 4;
+				if (!value.IsEmpty() && !value.Parse(iocpWorkerCount))
+					return false;
+			}
 			#endif
 			else if (name.Equals(TCV("-maxworkers")))
 			{
@@ -433,6 +443,7 @@ namespace uba
 		#endif
 
 		NetworkBackendTcpCreateInfo nbtci(logWriter);
+		nbtci.iocpWorkerCount = iocpWorkerCount;
 		nbtci.Apply(config);
 		NetworkBackendTcp networkBackend(nbtci);
 
diff --git a/Engine/Source/Programs/UnrealBuildAccelerator/Common/Private/UbaNetworkBackendTcp.cpp b/Engine/Source/Programs/UnrealBuildAccelerator/Common/Private/UbaNetworkBackendTcp.cpp
--- a/Engine/Source/Programs/UnrealBuildAccelerator/Common/Private/UbaNetworkBackendTcp.cpp
+++ b/Engine/Source/Programs/UnrealBuildAccelerator/Common/Private/UbaNetworkBackendTcp.cpp
@@ -78,15 +78,6 @@ namespace uba
 
 	////////////////////////////////////////////////////////////////////////////////////////////////////
 
-	#if UBA_USE_OVERLAPPED_SEND || UBA_USE_IOCP
-	struct Overlapped : OVERLAPPED
-	{
-		bool isSend = false;
-	};
-	#endif
-
-	////////////////////////////////////////////////////////////////////////////////////////////////////
-
 	struct NetworkBackendTcp::Connection
 	{
 		Connection(Logger& l, SOCKET s) : logger(l), socket(s), ready(true) { CreateGuid(uid); }
@@ -126,15 +117,13 @@ namespace uba
 
 		// Temporary state for iocp
 		#if UBA_USE_IOCP
-		Overlapped overlapped;
+		OVERLAPPED overlapped;
 		WSABUF wsaBuf = {};
 		u8 header[MaxHeaderSize];
 		u8* bodyData = nullptr;
 		u32 bodySize = 0;
-		u32 left = 0;
 		void* bodyContext = nullptr;
 		bool receivingHeader = true;
-		Event exited;
 		#endif
 
 		Connection(const Connection&) = delete;
@@ -159,8 +148,11 @@ namespace uba
 		const ConfigTable& table = *tablePtr;
 		table.GetValueAsBool(disableNagle, TC("DisableNagle"));
 		table.GetValueAsU32(statusUpdateSeconds, TC("StatusUpdateSeconds"));
+
+		#if PLATFORM_WINDOWS
 		table.GetValueAsBool(useOverlappedSend, TC("UseOverlappedSend"));
-		table.GetValueAsBool(useIocp, TC("UseIocp"));
+		table.GetValueAsU32(iocpWorkerCount, TC("IocpWorkerCount"));
+		#endif
 	}
 
 	////////////////////////////////////////////////////////////////////////////////////////////////////
@@ -175,13 +167,12 @@ namespace uba
 				return logger.Error(TC("WSAStartup failed (%d)"), res);
 
 			#if UBA_USE_IOCP
-			if (m_useIocp && !IsRunningWine())
+			if (m_iocpWorkerCount)
 			{
 				m_iocpHandle = CreateIoCompletionPort(INVALID_HANDLE_VALUE, NULL, 0, 0);
 
-				//for (Thread& t : m_iocpThreads)
-				//	t.Start([this]() { ThreadIocp(); return 0; });
-				m_iocpThreads[0].Start([this]() { ThreadIocp(); return 0; });
+				for (u32 i=0; i!=m_iocpWorkerCount; ++i)
+					m_iocpThreads.emplace_back().Start([this]() { ThreadIocp(); return 0; }, TC("UbaIocp"));
 			}
 			#endif
 
@@ -204,13 +195,16 @@ namespace uba
 		: m_logger(info.logWriter, prefix)
 	{
 		m_disableNagle = info.disableNagle;
+
+		#if PLATFORM_WINDOWS
 		m_useOverlappedSend = info.useOverlappedSend;
-		m_useIocp = info.useIocp;
+		m_iocpWorkerCount = u16(info.iocpWorkerCount);
+		#endif
 
 		if (info.statusUpdateSeconds)
 		{
 			m_tcpStatusLoop.Create(true);
-			m_tcpStatusThread.Start([this, sus = info.statusUpdateSeconds]() { ThreadStatus(sus); return 0; });
+			m_tcpStatusThread.Start([this, sus = info.statusUpdateSeconds]() { ThreadStatus(sus); return 0; }, TC("UbaTcpStat"));
 		}
 	}
 
@@ -234,9 +228,14 @@ namespace uba
 		m_connections.clear();
 
 		#if UBA_USE_IOCP
-		m_iocpLoop = false;
-		for (Thread& t : m_iocpThreads)
-			t.Wait();
+		if (m_iocpHandle)
+		{
+			for (u64 i=0; i!=m_iocpThreads.size(); ++i)
+				PostQueuedCompletionStatus(m_iocpHandle, 0, 1, NULL);
+			for (Thread& t : m_iocpThreads)
+				t.Wait();
+			CloseHandle(m_iocpHandle);
+		}
 		#endif
 
 		m_tcpStatusLoop.Set();
@@ -297,11 +296,7 @@ namespace uba
 
 		#if UBA_USE_IOCP
 		if (m_iocpHandle && !conn.wsaBuf.buf)
-		{
-			conn.exited.Create(true);
-			conn.left = headerSize;
 			PostIocpRead(conn, conn.header, headerSize);
-		}
 		#endif
 	}
 
@@ -821,12 +816,6 @@ namespace uba
 			Connection& c = *it;
 			if (&c != connection)
 				continue;
-			
-			#if UBA_USE_IOCP
-			if (m_iocpHandle)
-				c.exited.IsSet();
-			#endif
-				
 			it = m_connections.erase(it);
 			break;
 		}
@@ -879,8 +868,7 @@ namespace uba
 			u64 left = bufferLen;
 			while (left)
 			{
-				Overlapped overlapped {};
-				overlapped.isSend = true;
+				OVERLAPPED overlapped {};
 				overlapped.hEvent = ev.GetHandle();
 				WSABUF buf { u32(left), buffer };
 
@@ -1158,57 +1146,71 @@ namespace uba
 	{
 		ElevateCurrentThreadPriority();
 
-		while (m_iocpLoop)
+		while (true)
 		{
-			DWORD bytesTransferred;
-			ULONG_PTR completionKey;
-			Overlapped* overlapped;
-			u32 timeoutMs = 1000;
-			BOOL result = GetQueuedCompletionStatus(m_iocpHandle, &bytesTransferred, &completionKey, (OVERLAPPED**)&overlapped, timeoutMs);
+			DWORD bytesTransferred = 0;
+			ULONG_PTR completionKey = 0;
+			OVERLAPPED* overlapped = nullptr;
+			BOOL result = GetQueuedCompletionStatus(m_iocpHandle, &bytesTransferred, &completionKey, (OVERLAPPED**)&overlapped, INFINITE);
 
-			if (!result || !overlapped)
-				continue;
+			if (!result && !overlapped)
+			{
+				m_logger.Info(TC("GetQueuedCompletionStatus error (%s)"), LastErrorToText(WSAGetLastError()).data);
+				break;
+			}
+			if (completionKey == 1)
+				break;
 
 			auto& connection = *(Connection*)completionKey;
 			auto& logger = connection.logger;
 
-			if (bytesTransferred == 0)
+			if (bytesTransferred == 0 || !result)
 			{
+				SCOPED_CRITICAL_SECTION(connection.shutdownLock, lock2);
+				SOCKET s = connection.socket;
+				CancelIoEx((HANDLE)s, overlapped);
+
 				{
-					SCOPED_CRITICAL_SECTION(connection.shutdownLock, lock2);
-					SOCKET s = connection.socket;
-					CancelIoEx((HANDLE)s, &connection.overlapped);
+					SCOPED_FUTEX(connection.sendLock, lock);
+					connection.socket = INVALID_SOCKET;
+				}
 
+				// We need to marshal shutdown and callback handling out of iocp thread since callback can call things that requires iocp to loop
+				connection.recvThread.Start([this, s, connPtr = &connection]
 					{
-						SCOPED_FUTEX(connection.sendLock, lock);
-						connection.socket = INVALID_SOCKET;
-					}
-					if (auto cb = connection.disconnectCallback)
-					{
-						auto context = connection.disconnectContext;
-						connection.disconnectCallback = nullptr;
-						connection.disconnectContext = nullptr;
-						cb(context, connection.uid, &connection);
-					}
+						Connection& connection = *connPtr;
+						auto& logger = connection.logger;
 
-					if (s != INVALID_SOCKET)
-					{
-						ShutdownSocket(logger, s, TC("ThreadRecv"));
-						CloseSocket(logger, s, TC("ThreadRecv"));
-					}
-				}
-				connection.exited.Set();
+						SCOPED_CRITICAL_SECTION(connection.shutdownLock, lock2);
+						if (auto cb = connection.disconnectCallback)
+						{
+							auto context = connection.disconnectContext;
+							connection.disconnectCallback = nullptr;
+							connection.disconnectContext = nullptr;
+							cb(context, connection.uid, &connection);
+						}
+
+						if (s != INVALID_SOCKET)
+						{
+							ShutdownSocket(logger, s, TC("ThreadRecv"));
+							CloseSocket(logger, s, TC("ThreadRecv"));
+						}
+						return 0;
+
+					}, TC("UbaTcpSdwn"));
 				continue;
 			}
 
-			if (overlapped->isSend)
+			bool isSend = overlapped != &connection.overlapped;
+			if (isSend)
 				continue;
 
-			connection.left -= bytesTransferred;
-			if (connection.left)
+			UBA_ASSERT(bytesTransferred <= connection.wsaBuf.len);
+			connection.wsaBuf.len -= bytesTransferred;
+			if (connection.wsaBuf.len)
 			{
 				u8* newPos = (u8*)connection.wsaBuf.buf + bytesTransferred;
-				PostIocpRead(connection, newPos, connection.left);
+				PostIocpRead(connection, newPos, connection.wsaBuf.len);
 				continue;
 			}
 
@@ -1235,12 +1237,10 @@ namespace uba
 
 				if (!bodySize)
 				{
-					connection.left = connection.headerSize;
-					PostIocpRead(connection, connection.header, connection.left);
+					PostIocpRead(connection, connection.header, connection.headerSize);
 					continue;
 				}
 
-				connection.left = bodySize;
 				connection.receivingHeader = false;
 				PostIocpRead(connection, bodyData, bodySize);
 			}
@@ -1265,7 +1265,6 @@ namespace uba
 					continue;
 
 				connection.receivingHeader = true;
-				connection.left = connection.headerSize;
 				PostIocpRead(connection, connection.header, connection.headerSize);
 			}
 		}
@@ -1274,19 +1273,19 @@ namespace uba
 	bool NetworkBackendTcp::PostIocpRead(Connection& connection, u8* data, u32 dataSize)
 	{
 		UBA_ASSERT(data);
+		UBA_ASSERT(dataSize < 1*1024*1024); // Sanity
 		ZeroMemory(&connection.overlapped, sizeof(OVERLAPPED));
 		connection.wsaBuf.buf = (char*)data;
 		connection.wsaBuf.len = dataSize;
 
 		SOCKET socket = connection.socket;
 
-		DWORD flags = 0, bytesRecv = 0;
-		int ret = WSARecv(socket, &connection.wsaBuf, 1, &bytesRecv, &flags, &connection.overlapped, NULL);
+		DWORD flags = 0;
+		int ret = WSARecv(socket, &connection.wsaBuf, 1, NULL, &flags, &connection.overlapped, NULL);
 		if (ret == SOCKET_ERROR && WSAGetLastError() != WSA_IO_PENDING)
 		{
 			connection.logger.Info(TC("WSARecv failed for socket %i trying to receive %u bytes (%s)"), socket, dataSize, LastErrorToText(WSAGetLastError()).data);
-			//closesocket(clientSocket);
-			//delete context;
+			PostQueuedCompletionStatus(m_iocpHandle, 0, (ULONG_PTR)&connection, &connection.overlapped );
 			return false;
 		}
 		return true;
diff --git a/Engine/Source/Programs/UnrealBuildAccelerator/Common/Private/UbaNetworkServer.cpp b/Engine/Source/Programs/UnrealBuildAccelerator/Common/Private/UbaNetworkServer.cpp
--- a/Engine/Source/Programs/UnrealBuildAccelerator/Common/Private/UbaNetworkServer.cpp
+++ b/Engine/Source/Programs/UnrealBuildAccelerator/Common/Private/UbaNetworkServer.cpp
@@ -58,7 +58,7 @@ namespace uba
 		{
 			m_context = new WorkerContext(server);
 			m_loop = true;
-			m_thread.Start([&]() { ThreadWorker(server); return 0; }, TC("UbaWrk/NetwSrv"));
+			m_thread.Start([&]() { ThreadWorker(server); return 0; }, TC("UbaWrkNetwSrv"));
 		}
 
 		void Stop(NetworkServer& server)
diff --git a/Engine/Source/Programs/UnrealBuildAccelerator/Common/Private/UbaStorageUtils.cpp b/Engine/Source/Programs/UnrealBuildAccelerator/Common/Private/UbaStorageUtils.cpp
--- a/Engine/Source/Programs/UnrealBuildAccelerator/Common/Private/UbaStorageUtils.cpp
+++ b/Engine/Source/Programs/UnrealBuildAccelerator/Common/Private/UbaStorageUtils.cpp
@@ -198,14 +198,14 @@ namespace uba
 			}
 		}
 
-		u32 timeOutTimeMs = 5*60*1000;
+		u32 timeOutTimeMs = 20*60*1000;
 
 		for (u32 i=0; i!=inFlightCount; ++i)
 		{
 			Entry& entry = entries[i];
 			if (!entry.done.IsSet(timeOutTimeMs))
 			{
-				logger.Error(TC("SendBatchMessages timed out after 5 minutes getting async message response (%u/%u). Received %llu bytes so far. FetchId: %u (%s)"), i, inFlightCount, responseSize, fetchId, hint);
+				logger.Error(TC("SendBatchMessages timed out after 20 minutes getting async message response (%u/%u) This timeout will cause a crash. Received %llu bytes so far. FetchId: %u (%s)"), i, inFlightCount, responseSize, fetchId, hint);
 				timeOutTimeMs = 10;
 			}
 			if (!entry.message.ProcessAsyncResults(entry.reader))
diff --git a/Engine/Source/Programs/UnrealBuildAccelerator/Common/Public/UbaNetworkBackendTcp.h b/Engine/Source/Programs/UnrealBuildAccelerator/Common/Public/UbaNetworkBackendTcp.h
--- a/Engine/Source/Programs/UnrealBuildAccelerator/Common/Public/UbaNetworkBackendTcp.h
+++ b/Engine/Source/Programs/UnrealBuildAccelerator/Common/Public/UbaNetworkBackendTcp.h
@@ -19,7 +19,7 @@ namespace uba
 		u32 statusUpdateSeconds = 0; // Zero means never
 		bool disableNagle = true;
 		bool useOverlappedSend = true;
-		bool useIocp = false;
+		u32 iocpWorkerCount = 4;
 
 		void Apply(Config& config, const tchar* tableName = TC("NetworkBackendTcp"));
 	};
@@ -76,16 +76,15 @@ namespace uba
 		Atomic<u64> m_totalRecv;
 
 		bool m_disableNagle;
-		bool m_useOverlappedSend;
-		bool m_useIocp;
 
 		#if PLATFORM_WINDOWS
 		bool m_wsaInitDone = false;
+		bool m_useOverlappedSend;
+		u16 m_iocpWorkerCount;
 		#endif
 
-		Thread m_iocpThreads[4];
+		List<Thread> m_iocpThreads;
 		void* m_iocpHandle = nullptr;
-		bool m_iocpLoop = true;
 
 		Event m_tcpStatusLoop;
 		Thread m_tcpStatusThread;
