[UBA]
* Optimized stoplisten (shutdown) to create a socket and connect to the listen socket just to skip the poll timeout.
* Also turned off iocp that was accidentally enabled by default

diff --git a/Engine/Source/Programs/UnrealBuildAccelerator/Common/Private/UbaNetworkBackendTcp.cpp b/Engine/Source/Programs/UnrealBuildAccelerator/Common/Private/UbaNetworkBackendTcp.cpp
--- a/Engine/Source/Programs/UnrealBuildAccelerator/Common/Private/UbaNetworkBackendTcp.cpp
+++ b/Engine/Source/Programs/UnrealBuildAccelerator/Common/Private/UbaNetworkBackendTcp.cpp
@@ -73,6 +73,7 @@ namespace uba
 		ListenConnectedFunc connectedFunc;
 		Event listening;
 		Atomic<SOCKET> socket = INVALID_SOCKET;
+		sockaddr_in addr;
 		Thread thread;
 	};
 
@@ -393,7 +394,14 @@ namespace uba
 	{
 		SCOPED_FUTEX(m_listenEntriesLock, lock);
 		for (auto& e : m_listenEntries)
+		{
 			e.socket = INVALID_SOCKET;
+			SOCKET tempSocket = ::socket(AF_INET, SOCK_STREAM, IPPROTO_TCP); // Create a temporary socket just to connect to listen socket to wakeup WSAPoll
+			if (tempSocket == INVALID_SOCKET)
+				continue;
+			::connect(tempSocket, (sockaddr*)&e.addr, sizeof(e.addr));
+			::closesocket(tempSocket);
+		}
 		for (auto& e : m_listenEntries)
 			e.thread.Wait();
 		m_listenEntries.clear();
@@ -423,7 +431,7 @@ namespace uba
 		auto addrGuard = MakeGuard([result]() { FreeAddrInfoW(result); });
 
 		// Create a socket for listening to connections
-		SOCKET listenSocket = socket(result->ai_family, result->ai_socktype, result->ai_protocol);
+		SOCKET listenSocket = ::socket(result->ai_family, result->ai_socktype, result->ai_protocol);
 		if (listenSocket == INVALID_SOCKET)
 			return logger.Error(TC("socket failed (%s)"), LastErrorToText(WSAGetLastError()).data);
 
@@ -445,6 +453,10 @@ namespace uba
 			return false;
 		}
 
+		entry.addr = *(sockaddr_in*)result->ai_addr;
+		if (entry.addr.sin_addr.s_addr == 0) // if 0.0.0.0 then we store 127.0.0.1
+			entry.addr.sin_addr.s_addr = htonl(127 << 24 | 1);
+
 		addrGuard.Execute();
 
 		res = listen(listenSocket, SOMAXCONN);
@@ -459,15 +471,18 @@ namespace uba
 
 		listenEv.Execute();
 
-		while (entry.socket != INVALID_SOCKET)
+		while (true)
 		{
 			WSAPOLLFD p;
 			p.fd = listenSocket;
 			p.revents = 0;
 			p.events = POLLIN;
-			int timeoutMs = 2000;
+			int timeoutMs = 5000;
 			int pollRes = WSAPoll(&p, 1, timeoutMs);
 
+			if (entry.socket == INVALID_SOCKET)
+				break;
+
 			if (pollRes == SOCKET_ERROR)
 			{
 				int lastError = WSAGetLastError();
diff --git a/Engine/Source/Programs/UnrealBuildAccelerator/Common/Public/UbaNetworkBackendTcp.h b/Engine/Source/Programs/UnrealBuildAccelerator/Common/Public/UbaNetworkBackendTcp.h
--- a/Engine/Source/Programs/UnrealBuildAccelerator/Common/Public/UbaNetworkBackendTcp.h
+++ b/Engine/Source/Programs/UnrealBuildAccelerator/Common/Public/UbaNetworkBackendTcp.h
@@ -19,7 +19,7 @@ namespace uba
 		u32 statusUpdateSeconds = 0; // Zero means never
 		bool disableNagle = true;
 		bool useOverlappedSend = true;
-		u32 iocpWorkerCount = 4;
+		u32 iocpWorkerCount = 0; // Iocp is turned off by default
 
 		void Apply(Config& config, const tchar* tableName = TC("NetworkBackendTcp"));
 	};
