[UBA]
* Fixed so tests work with allowKeepFilesInMemory false
* Changed const tchar* to StringView

diff --git a/Engine/Source/Programs/UnrealBuildAccelerator/Common/Private/UbaFile.cpp b/Engine/Source/Programs/UnrealBuildAccelerator/Common/Private/UbaFile.cpp
--- a/Engine/Source/Programs/UnrealBuildAccelerator/Common/Private/UbaFile.cpp
+++ b/Engine/Source/Programs/UnrealBuildAccelerator/Common/Private/UbaFile.cpp
@@ -690,7 +690,7 @@ namespace uba
 		return false;
 #else
 
-		UBA_ASSERTF(false, TC("CopyFileW not implemented"));
+		UBA_ASSERTF(false, TC("CopyFileW not implemented (From %s to %s)"), existingFileName, newFileName);
 		return false;
 #endif
 	}
@@ -838,7 +838,11 @@ namespace uba
 		if (errno == EPERM) // Because of System Integrity Protection we might not be allowed to link this file, fallback to copy
 			return false;
 		#endif
-
+		if (errno == EEXIST)
+		{
+			SetLastError(ERROR_ALREADY_EXISTS);
+			return false;
+		}
 		UBA_ASSERTF(false, TC("CreateHardLinkW %s to %s error handling not implemented (%s)"), existingFileName, newFileName, strerror(errno));
 		return false;
 #endif
diff --git a/Engine/Source/Programs/UnrealBuildAccelerator/Detours/Private/PosixOS/UbaDetoursFunctionsPosix.cpp b/Engine/Source/Programs/UnrealBuildAccelerator/Detours/Private/PosixOS/UbaDetoursFunctionsPosix.cpp
--- a/Engine/Source/Programs/UnrealBuildAccelerator/Detours/Private/PosixOS/UbaDetoursFunctionsPosix.cpp
+++ b/Engine/Source/Programs/UnrealBuildAccelerator/Detours/Private/PosixOS/UbaDetoursFunctionsPosix.cpp
@@ -362,7 +362,7 @@ int Shared_open(const char* funcName, const char* file, int flags, int mode, con
 	{
 		if (!isWrite) // We need to skip SystemTemp.. lots of stuff going on there.
 		{
-			dirTableOffset = Rpc_GetEntryOffset(fileNameKey, fileName.data, fileName.count, false);
+			dirTableOffset = Rpc_GetEntryOffset(fileNameKey, fileName, false);
 			bool allowEarlyOut = true;
 			if (dirTableOffset == ~u32(0))
 			{
@@ -626,7 +626,7 @@ int Shared_fstat(const char* funcName, int fd, struct stat* attr, const True_fst
 	}
 
 	FileAttributes fileAttr = {};
-	const char* realName = Shared_GetFileAttributes(fileAttr, fi.originalName);
+	const char* realName = Shared_GetFileAttributes(fileAttr, ToView(fi.originalName));
 
 	if (!fileAttr.useCache)
 	{
@@ -693,7 +693,7 @@ int Shared_stat(const char* funcName, const char* file, struct stat* attr, const
 	DevirtualizePath(fixedFile);
 
 	FileAttributes fileAttr;
-	const char* realName = Shared_GetFileAttributes(fileAttr, fixedFile.data);
+	const char* realName = Shared_GetFileAttributes(fileAttr, fixedFile);
 
 	if (!fileAttr.useCache)
 	{
@@ -993,7 +993,7 @@ UBA_EXPORT char* UBA_WRAPPER(realpath)(const char* path, char* resolved_path)
 	DevirtualizePath(devirtualizedPath);
 
 	FileAttributes fileAttr;
-	const char* realName = Shared_GetFileAttributes(fileAttr, devirtualizedPath.data);(void)realName;
+	Shared_GetFileAttributes(fileAttr, devirtualizedPath);
 	if (!fileAttr.useCache)
 	{
 		auto res = TRUE_WRAPPER(realpath)(path, resolved_path);
@@ -1529,7 +1529,7 @@ UBA_EXPORT int UBA_WRAPPER(access)(const char* pathname, int mode)
 	}
 
 	FileAttributes attr;
-	const char* realName = Shared_GetFileAttributes(attr, fixedPath.data, checkIfDir);
+	const char* realName = Shared_GetFileAttributes(attr, fixedPath, checkIfDir);
 
 	if (!attr.useCache)
 	{
diff --git a/Engine/Source/Programs/UnrealBuildAccelerator/Detours/Private/UbaDetoursFileMappingTable.cpp b/Engine/Source/Programs/UnrealBuildAccelerator/Detours/Private/UbaDetoursFileMappingTable.cpp
--- a/Engine/Source/Programs/UnrealBuildAccelerator/Detours/Private/UbaDetoursFileMappingTable.cpp
+++ b/Engine/Source/Programs/UnrealBuildAccelerator/Detours/Private/UbaDetoursFileMappingTable.cpp
@@ -246,11 +246,11 @@ namespace uba
 		UpdateWrittenFiles(reader);
 	}
 
-	u32 Rpc_GetEntryOffset(const StringKey& entryNameKey, const tchar* entryName, u64 entryNameLen, bool checkIfDir)
+	u32 Rpc_GetEntryOffset(const StringKey& entryNameKey, StringView entryName, bool checkIfDir)
 	{
 		u32 dirTableOffset = ~u32(0);
 		StringBuffer<MaxPath> entryNameForKey;
-		entryNameForKey.Append(entryName, entryNameLen);
+		entryNameForKey.Append(entryName);
 		if (CaseInsensitiveFs)
 			entryNameForKey.MakeLower();
 		else if (entryNameForKey.count == 1 && entryNameForKey[0] == '/')
@@ -261,21 +261,21 @@ namespace uba
 		if (exists != DirectoryTable::Exists_Maybe)
 			return dirTableOffset;
 
-		const tchar* lastPathSeparator = TStrrchr(entryName, PathSeparator);
+		const tchar* lastPathSeparator = TStrrchr(entryName.data, PathSeparator);
 		if (!lastPathSeparator)
 		{
-			UBA_ASSERTF(lastPathSeparator, TC("No path separator found in %s"), TStrlen(entryName) > 0 ? entryName : TC("(NULL)"));
+			UBA_ASSERTF(lastPathSeparator, TC("No path separator found in %s"), entryName.count > 0 ? entryName.data : TC("(NULL)"));
 			return ~u32(0);
 		}
 
 		#if PLATFORM_WINDOWS
-		UBA_ASSERT(wcsncmp(entryName, g_systemTemp.data, g_systemTemp.count) != 0);
+		UBA_ASSERT(wcsncmp(entryName.data, g_systemTemp.data, g_systemTemp.count) != 0);
 		#endif
 
-		u32 dirNameLen = u32(lastPathSeparator - entryName);
+		u32 dirNameLen = u32(lastPathSeparator - entryName.data);
 		DirHash hash(StringView(entryNameForKey.data, dirNameLen));
 
-		if (Rpc_UpdateDirectory(hash.key, entryName, dirNameLen) == ~u32(0))
+		if (Rpc_UpdateDirectory(hash.key, entryName.data, dirNameLen) == ~u32(0))
 			return ~u32(0);
 
 		SCOPED_WRITE_LOCK(g_directoryTable.m_lookupLock, lookLock);
diff --git a/Engine/Source/Programs/UnrealBuildAccelerator/Detours/Private/UbaDetoursFileMappingTable.h b/Engine/Source/Programs/UnrealBuildAccelerator/Detours/Private/UbaDetoursFileMappingTable.h
--- a/Engine/Source/Programs/UnrealBuildAccelerator/Detours/Private/UbaDetoursFileMappingTable.h
+++ b/Engine/Source/Programs/UnrealBuildAccelerator/Detours/Private/UbaDetoursFileMappingTable.h
@@ -93,7 +93,7 @@ namespace uba
 	u32  Rpc_UpdateDirectory(const StringKey& dirKey, const tchar* dirName, u64 dirNameLen, bool lockDirTable = true);
 	void Rpc_UpdateCloseHandle(const tchar* handleName, u32 closeId, bool deleteOnClose, const tchar* newName, const FileMappingHandle& mappingHandle, u64 mappingWritten, bool success);
 	void Rpc_UpdateTables();
-	u32  Rpc_GetEntryOffset(const StringKey& entryNameKey, const tchar* entryName, u64 entryNameLen, bool checkIfDir = false);
+	u32  Rpc_GetEntryOffset(const StringKey& entryNameKey, StringView entryName, bool checkIfDir = false);
 	void Rpc_GetFullFileName(const tchar*& path, u64& pathLen, StringBufferBase& tempBuf, bool useVirtualName, const tchar* const* loaderPaths = nullptr);
 	void Rpc_GetFullFileName2(const tchar* path, StringBufferBase& outReal, StringBufferBase& outVirtual, const tchar* const* loaderPaths = nullptr);
 	void Rpc_GetWrittenFiles();
diff --git a/Engine/Source/Programs/UnrealBuildAccelerator/Detours/Private/UbaDetoursShared.cpp b/Engine/Source/Programs/UnrealBuildAccelerator/Detours/Private/UbaDetoursShared.cpp
--- a/Engine/Source/Programs/UnrealBuildAccelerator/Detours/Private/UbaDetoursShared.cpp
+++ b/Engine/Source/Programs/UnrealBuildAccelerator/Detours/Private/UbaDetoursShared.cpp
@@ -402,11 +402,10 @@ namespace uba
 	#endif
 
 
-	const tchar* Shared_GetFileAttributes(FileAttributes& outAttr, const tchar* fileName, bool checkIfDir)
+	const tchar* Shared_GetFileAttributes(FileAttributes& outAttr, StringView fileName, bool checkIfDir)
 	{
 		StringBuffer<MaxPath> fileNameForKey;
 		fileNameForKey.Append(fileName);
-		StringView fileNameView(fileName, fileNameForKey.count);
 		if (CaseInsensitiveFs)
 			fileNameForKey.MakeLower();
 
@@ -416,17 +415,18 @@ namespace uba
 
 		memset(&outAttr.data, 0, sizeof(outAttr.data));
 
-		bool keepInMemory = KeepInMemory(fileNameView, false);
+		bool isInsideSystemTemp = fileName.StartsWith(g_systemTemp.data);
+		bool keepInMemory = KeepInMemory(fileName, false);
 		if (keepInMemory)
 		{
 			SCOPED_READ_LOCK(g_mappedFileTable.m_lookupLock, lock);
 			auto it = g_mappedFileTable.m_lookup.find(fileNameKey);
 			if (it == g_mappedFileTable.m_lookup.end() || it->second.deleted)
 			{
-				if (StartsWith(fileName, g_systemTemp.data))
+				if (isInsideSystemTemp)
 				{
 					outAttr.useCache = false;
-					return fileName;
+					return fileName.data;
 				}
 				outAttr.useCache = true;
 				outAttr.exists = false;
@@ -469,13 +469,12 @@ namespace uba
 #endif
 		else if (g_allowDirectoryCache)
 		{
-			bool isInsideSystemTemp = StartsWith(fileName, g_systemTemp.data);
 			// This is an optimization where we populate directory table and use that to figure out if file exists or not..
 			// .. in msvc's case it doesn't matter much because these tables are already up to date when msvc use CreateFile.
 			// .. clang otoh is using CreateFile with tooons of different paths trying to open files.. in remote worker case this becomes super expensive
 			if (!isInsideSystemTemp) // We need to skip SystemTemp.. lots of stuff going on there.
 			{
-				u32 dirTableOffset = Rpc_GetEntryOffset(fileNameKey, fileName, fileNameForKey.count, checkIfDir);
+				u32 dirTableOffset = Rpc_GetEntryOffset(fileNameKey, fileName, checkIfDir);
 
 				if (dirTableOffset == ~u32(0))
 				{
@@ -522,7 +521,7 @@ namespace uba
 						u64 fileSize = info.size;
 
 						// Could be compressed and then directory table size is wrong
-						if (CouldBeCompressedFile(fileNameView))
+						if (CouldBeCompressedFile(fileName))
 						{
 							SCOPED_READ_LOCK(g_mappedFileTable.m_lookupLock, lock);
 							auto findIt = g_mappedFileTable.m_lookup.find(fileNameKey);
@@ -530,11 +529,11 @@ namespace uba
 							{
 								// If file is output file we accept wrong size because size is not supposed to be used anyway.
 								// We don't want to trigger unnecessary download/decompress of file
-								if (!g_rules->IsOutputFile(fileNameView))
+								if (!g_rules->IsOutputFile(fileName))
 								{
 									StringBuffer<> temp;
 									u32 closeId;
-									Rpc_CreateFileW(fileNameView, fileNameKey, AccessFlag_Read, temp.data, temp.capacity, fileSize, closeId, false);
+									Rpc_CreateFileW(fileName, fileNameKey, AccessFlag_Read, temp.data, temp.capacity, fileSize, closeId, false);
 								}
 							}
 							else
@@ -580,13 +579,13 @@ namespace uba
 			else
 			{
 				outAttr.useCache = false;
-				return fileName;
+				return fileName.data;
 			}
 		}
 		else
 		{
 			outAttr.useCache = false;
-			return fileName;
+			return fileName.data;
 		}
 
 #if 0//UBA_DEBUG_VALIDATE
@@ -598,20 +597,20 @@ namespace uba
 			BOOL res = True_GetFileAttributesExW(fileName, GetFileExInfoStandard, &validate); (void)res;
 			if (outAttr.exists)
 			{
-				UBA_ASSERTF(res != 0, L"File %ls exists even though uba claims it is not..", fileName);
+				UBA_ASSERTF(res != 0, L"File %ls exists even though uba claims it is not..", fileName.data);
 				if (validate.dwFileAttributes & FILE_ATTRIBUTE_DIRECTORY)
-					UBA_ASSERTF((outAttr.data.dwFileAttributes & FILE_ATTRIBUTE_DIRECTORY), L"File attributes are wrong for %ls", fileName);
+					UBA_ASSERTF((outAttr.data.dwFileAttributes & FILE_ATTRIBUTE_DIRECTORY), L"File attributes are wrong for %ls", fileName.data);
 				else
 				{
 					validate.ftCreationTime = outAttr.data.ftCreationTime; // Creation time is not really important
 					validate.ftLastAccessTime = outAttr.data.ftLastAccessTime; // Access time is not really important
 					validate.ftLastWriteTime = outAttr.data.ftLastWriteTime; // Write time is important, revisit this
-					UBA_ASSERTF(memcmp(&validate, &outAttr.data, sizeof(WIN32_FILE_ATTRIBUTE_DATA)) == 0, L"File %ls is not up-to-date in cache", fileName);
+					UBA_ASSERTF(memcmp(&validate, &outAttr.data, sizeof(WIN32_FILE_ATTRIBUTE_DATA)) == 0, L"File %ls is not up-to-date in cache", fileName.data);
 				}
 			}
 			else
 			{
-				UBA_ASSERTF(res == 0, L"Can't find file %ls but validation checked that it is there", fileName); // This means most likely that Uba did not update attribute table for added files.
+				UBA_ASSERTF(res == 0, L"Can't find file %ls but validation checked that it is there", fileName.data); // This means most likely that Uba did not update attribute table for added files.
 				DWORD lastError2 = GetLastError();
 				if (lastError2 == ERROR_PATH_NOT_FOUND || lastError2 == ERROR_INVALID_NAME)
 					lastError2 = ERROR_FILE_NOT_FOUND;
@@ -619,6 +618,6 @@ namespace uba
 			}
 		}
 #endif
-		return fileName;
+		return fileName.data;
 	}
 }
diff --git a/Engine/Source/Programs/UnrealBuildAccelerator/Detours/Private/UbaDetoursShared.h b/Engine/Source/Programs/UnrealBuildAccelerator/Detours/Private/UbaDetoursShared.h
--- a/Engine/Source/Programs/UnrealBuildAccelerator/Detours/Private/UbaDetoursShared.h
+++ b/Engine/Source/Programs/UnrealBuildAccelerator/Detours/Private/UbaDetoursShared.h
@@ -127,7 +127,7 @@ namespace uba
 	void Shared_WriteConsole(const char* chars, u32 charCount, bool isError);
 	void Shared_WriteConsole(const wchar_t* chars, u32 charCount, bool isError);
 
-	const tchar* Shared_GetFileAttributes(FileAttributes& outAttr, const tchar* fileName, bool checkIfDir = false);
+	const tchar* Shared_GetFileAttributes(FileAttributes& outAttr, StringView fileName, bool checkIfDir = false);
 
 	void InitSharedVariables();
 
diff --git a/Engine/Source/Programs/UnrealBuildAccelerator/Detours/Private/Windows/UbaDetoursFunctionsKernelBase.inl b/Engine/Source/Programs/UnrealBuildAccelerator/Detours/Private/Windows/UbaDetoursFunctionsKernelBase.inl
--- a/Engine/Source/Programs/UnrealBuildAccelerator/Detours/Private/Windows/UbaDetoursFunctionsKernelBase.inl
+++ b/Engine/Source/Programs/UnrealBuildAccelerator/Detours/Private/Windows/UbaDetoursFunctionsKernelBase.inl
@@ -1029,10 +1029,10 @@ DWORD Detoured_GetFileType(HANDLE hFile)
 	return True_GetFileType(hFile); // Calling NtQueryVolumeInformationFile
 }
 
-BOOL Shared_GetFileAttributesExW(const tchar* caller, LPCWSTR lpFileName, GET_FILEEX_INFO_LEVELS fInfoLevelId, LPVOID lpFileInformation, LPCWSTR originalName)
+BOOL Shared_GetFileAttributesExW(const tchar* caller, StringView fileName, GET_FILEEX_INFO_LEVELS fInfoLevelId, LPVOID lpFileInformation, LPCWSTR originalName)
 {
 	FileAttributes attr;
-	const wchar_t* realName = Shared_GetFileAttributes(attr, lpFileName);
+	const wchar_t* realName = Shared_GetFileAttributes(attr, fileName);
 
 	if (!attr.useCache)
 	{
@@ -1070,7 +1070,7 @@ BOOL Detoured_GetFileAttributesExW(LPCWSTR lpFileName, GET_FILEEX_INFO_LEVELS fI
 		return FALSE;
 	}
 
-	return Shared_GetFileAttributesExW(TC("GetFileAttributesExW"), fixedName.data, fInfoLevelId, lpFileInformation, lpFileName);
+	return Shared_GetFileAttributesExW(TC("GetFileAttributesExW"), fixedName, fInfoLevelId, lpFileInformation, lpFileName);
 }
 
 DWORD Detoured_GetFileAttributesW(LPCWSTR lpFileName)
@@ -1090,7 +1090,7 @@ DWORD Detoured_GetFileAttributesW(LPCWSTR lpFileName)
 	DevirtualizePath(fixedPath);
 
 	WIN32_FILE_ATTRIBUTE_DATA data;
-	if (!Shared_GetFileAttributesExW(TC("GetFileAttributesW"), fixedPath.data, GetFileExInfoStandard, &data, lpFileName))
+	if (!Shared_GetFileAttributesExW(TC("GetFileAttributesW"), fixedPath, GetFileExInfoStandard, &data, lpFileName))
 		return INVALID_FILE_ATTRIBUTES;
 
 	return data.dwFileAttributes;
@@ -1137,7 +1137,7 @@ DWORD Detoured_GetLongPathNameW(LPCWSTR lpszShortPath, LPWSTR lpszLongPath, DWOR
 			StringBuffer<> realName(fixedName);
 			DevirtualizePath(realName);
 			WIN32_FILE_ATTRIBUTE_DATA data;
-			success = Shared_GetFileAttributesExW(TC("GetLongPathNameW"), realName.data, GetFileExInfoStandard, &data, path);
+			success = Shared_GetFileAttributesExW(TC("GetLongPathNameW"), realName, GetFileExInfoStandard, &data, path);
 		}
 
 		DWORD res = 0;
@@ -2226,7 +2226,7 @@ BOOL Detoured_GetFileInformationByHandle(HANDLE hFile, LPBY_HANDLE_FILE_INFORMAT
 			FixPath(fixedName, fi.originalName);
 
 			FileAttributes attr;
-			Shared_GetFileAttributes(attr, fixedName.data);
+			Shared_GetFileAttributes(attr, fixedName);
 
 			if (attr.useCache)
 			{
@@ -2814,7 +2814,7 @@ HMODULE GetModuleHandleNoSyscall(const tchar* moduleName)
 
 HMODULE Shared_LoadLibrary(LPCWSTR lpLibFileName, HANDLE hFile, DWORD dwFlags, bool& mightExist)
 {
-	if (!g_runningRemote && !g_trackInputsMem)
+	if (!g_runningRemote && !g_trackInputsMem || t_disallowDetour)
 		return NULL;
 
 	StringBuffer<> path;
diff --git a/Engine/Source/Programs/UnrealBuildAccelerator/Detours/Private/Windows/UbaDetoursFunctionsNtDll.inl b/Engine/Source/Programs/UnrealBuildAccelerator/Detours/Private/Windows/UbaDetoursFunctionsNtDll.inl
--- a/Engine/Source/Programs/UnrealBuildAccelerator/Detours/Private/Windows/UbaDetoursFunctionsNtDll.inl
+++ b/Engine/Source/Programs/UnrealBuildAccelerator/Detours/Private/Windows/UbaDetoursFunctionsNtDll.inl
@@ -459,7 +459,7 @@ NTSTATUS Detoured_NtQueryFullAttributesFile(POBJECT_ATTRIBUTES ObjectAttributes,
 	DevirtualizePath(fixedName);
 
 	FileAttributes attr;
-	Shared_GetFileAttributes(attr, fixedName.data);
+	Shared_GetFileAttributes(attr, fixedName);
 
 	if (!attr.useCache)
 	{
@@ -856,7 +856,7 @@ NTSTATUS NTAPI Shared_NtCreateFile(bool IsCreateFunc, PHANDLE hFileHandle, ACCES
 	bool useContent = IsContentUse(DesiredAccess, CreateDisposition);
 	bool isWrite = IsWrite(DesiredAccess, CreateDisposition);
 	bool isThrowAway = g_rules->IsThrowAway(fileName, g_runningRemote);
-	bool keepInMemory = KeepInMemory(fileName, isWrite) || ((isWrite || isDeleteOnClose) && g_rules->IsOutputFile(fileName)) || (isWrite && isThrowAway);
+	bool keepInMemory = KeepInMemory(fileName, isWrite) || ((isWrite || isDeleteOnClose) && g_rules->IsOutputFile(fileName) && g_allowKeepFilesInMemory) || (isWrite && isThrowAway);
 
 #if UBA_DEBUG_LOG_ENABLED
 	const wchar_t* isWriteStr = isWrite ? L" WRITE" : L""; (void)isWriteStr;
@@ -904,7 +904,7 @@ NTSTATUS NTAPI Shared_NtCreateFile(bool IsCreateFunc, PHANDLE hFileHandle, ACCES
 			if (!isWrite && !isSystemOrTempFile) // We need to skip SystemTemp.. lots of stuff going on there.
 			{
 				CHECK_PATH(fileNameLower);
-				dirTableOffset = Rpc_GetEntryOffset(fileNameKey, fileName.data, fileName.count, checkIfDir);
+				dirTableOffset = Rpc_GetEntryOffset(fileNameKey, fileName, checkIfDir);
 
 				bool allowEarlyOut = true;
 				if (dirTableOffset == ~u32(0))
diff --git a/Engine/Source/Programs/UnrealBuildAccelerator/Detours/Private/Windows/UbaDetoursFunctionsUcrtBase.inl b/Engine/Source/Programs/UnrealBuildAccelerator/Detours/Private/Windows/UbaDetoursFunctionsUcrtBase.inl
--- a/Engine/Source/Programs/UnrealBuildAccelerator/Detours/Private/Windows/UbaDetoursFunctionsUcrtBase.inl
+++ b/Engine/Source/Programs/UnrealBuildAccelerator/Detours/Private/Windows/UbaDetoursFunctionsUcrtBase.inl
@@ -165,7 +165,7 @@ errno_t Detoured__waccess_s(const wchar_t* path, int mode)
 	DevirtualizePath(fixedPath);
 
 	FileAttributes attr;
-	const wchar_t* realName = Shared_GetFileAttributes(attr, path);
+	const wchar_t* realName = Shared_GetFileAttributes(attr, fixedPath);
 
 	if (!attr.useCache)
 	{
