[UBA]
* Added const to Config object sent into NetworkBackendTcpCreateInfo::Apply

diff --git a/Engine/Source/Programs/UnrealBuildAccelerator/Common/Private/UbaNetworkBackendTcp.cpp b/Engine/Source/Programs/UnrealBuildAccelerator/Common/Private/UbaNetworkBackendTcp.cpp
--- a/Engine/Source/Programs/UnrealBuildAccelerator/Common/Private/UbaNetworkBackendTcp.cpp
+++ b/Engine/Source/Programs/UnrealBuildAccelerator/Common/Private/UbaNetworkBackendTcp.cpp
@@ -141,7 +141,7 @@ namespace uba
 
 	////////////////////////////////////////////////////////////////////////////////////////////////////
 
-	void NetworkBackendTcpCreateInfo::Apply(Config& config, const tchar* tableName)
+	void NetworkBackendTcpCreateInfo::Apply(const Config& config, const tchar* tableName)
 	{
 		const ConfigTable* tablePtr = config.GetTable(tableName);
 		if (!tablePtr)
diff --git a/Engine/Source/Programs/UnrealBuildAccelerator/Common/Public/UbaNetworkBackendTcp.h b/Engine/Source/Programs/UnrealBuildAccelerator/Common/Public/UbaNetworkBackendTcp.h
--- a/Engine/Source/Programs/UnrealBuildAccelerator/Common/Public/UbaNetworkBackendTcp.h
+++ b/Engine/Source/Programs/UnrealBuildAccelerator/Common/Public/UbaNetworkBackendTcp.h
@@ -21,7 +21,7 @@ namespace uba
 		bool useOverlappedSend = true;
 		u32 iocpWorkerCount = 0; // Iocp is turned off by default
 
-		void Apply(Config& config, const tchar* tableName = TC("NetworkBackendTcp"));
+		void Apply(const Config& config, const tchar* tableName = TC("NetworkBackendTcp"));
 	};
 
 	////////////////////////////////////////////////////////////////////////////////////////////////////
