[UBA]
* Fixed so FixPath return false for certain bad paths instead of crashing

diff --git a/Engine/Source/Programs/UnrealBuildAccelerator/Core/Public/UbaPathUtils.h b/Engine/Source/Programs/UnrealBuildAccelerator/Core/Public/UbaPathUtils.h
--- a/Engine/Source/Programs/UnrealBuildAccelerator/Core/Public/UbaPathUtils.h
+++ b/Engine/Source/Programs/UnrealBuildAccelerator/Core/Public/UbaPathUtils.h
@@ -156,7 +156,10 @@ namespace uba
 				--charLen;
 			}
 
-			UBA_ASSERTF(workingDir && *workingDir, TC("No working dir provided but path is relative (%s)"), buffer);
+			bool hasWorkDir = workingDir && *workingDir;
+			UBA_ASSERTF(hasWorkDir, TC("No working dir provided but path is relative (%s)"), buffer);
+			if (!hasWorkDir)
+				return false;
 			tchar temp2[1024];
 			UBA_ASSERTF(workingDirCharLen + charLen < sizeof_array(temp2), TC("%llu + %llu < %i"), workingDirCharLen, charLen, sizeof_array(temp2));
 			memcpy(temp2, workingDir, workingDirCharLen*sizeof(tchar));
@@ -350,11 +353,12 @@ namespace uba
 #endif
 	}
 
-	inline void FixPath(const tchar* fileName, const tchar* workingDir, u64 workingDirCharLen, StringBufferBase& buffer)
+	inline bool FixPath(const tchar* fileName, const tchar* workingDir, u64 workingDirCharLen, StringBufferBase& buffer)
 	{
 		u32 count = buffer.count;
-		FixPath2(fileName, workingDir, workingDirCharLen, buffer.data + count, buffer.capacity - count, &buffer.count);
+		bool res = FixPath2(fileName, workingDir, workingDirCharLen, buffer.data + count, buffer.capacity - count, &buffer.count);
 		buffer.count += count;
+		return res;
 	}
 
 	inline bool IsAbsolutePath(const tchar* path)
