[UBA]
* Fixed crash in detoured function WriteFile if "byteswritten" was null

diff --git a/Engine/Source/Programs/UnrealBuildAccelerator/Detours/Private/Windows/UbaDetoursFunctionsKernelBase.inl b/Engine/Source/Programs/UnrealBuildAccelerator/Detours/Private/Windows/UbaDetoursFunctionsKernelBase.inl
--- a/Engine/Source/Programs/UnrealBuildAccelerator/Detours/Private/Windows/UbaDetoursFunctionsKernelBase.inl
+++ b/Engine/Source/Programs/UnrealBuildAccelerator/Detours/Private/Windows/UbaDetoursFunctionsKernelBase.inl
@@ -634,6 +634,10 @@ void WriteStdFile(LPCVOID lpBuffer, DWORD nNumberOfBytesToWrite, bool isError)
 
 BOOL Detoured_WriteFile(HANDLE hFile, LPCVOID lpBuffer, DWORD nNumberOfBytesToWrite, LPDWORD lpNumberOfBytesWritten, LPOVERLAPPED lpOverlapped)
 {
+	DWORD temp = 0;
+	if (!lpNumberOfBytesWritten)
+		lpNumberOfBytesWritten = &temp;
+
 	DETOURED_CALL(WriteFile);
 	HANDLE trueHandle = hFile;
 	UBA_ASSERT(!isListDirectoryHandle(hFile));
@@ -701,10 +705,6 @@ BOOL Detoured_WriteFile(HANDLE hFile, LPCVOID lpBuffer, DWORD nNumberOfBytesToWr
 		//return True_WriteFile(trueHandle, lpBuffer, nNumberOfBytesToWrite, lpNumberOfBytesWritten, lpOverlapped);
 	}
 
-	DWORD temp;
-	if (!lpNumberOfBytesWritten)
-		lpNumberOfBytesWritten = &temp;
-
 	TimerScope ts(g_kernelStats.writeFile);
 	BOOL res = True_WriteFile(trueHandle, lpBuffer, nNumberOfBytesToWrite, lpNumberOfBytesWritten, lpOverlapped);
 
