[UBA]
* Added support for "independent" mappings in mapping buffer. They are simply just stand alone mappings but added here just to keep the outside code cleaner

diff --git a/Engine/Source/Programs/UnrealBuildAccelerator/Common/Private/UbaFileMapping.cpp b/Engine/Source/Programs/UnrealBuildAccelerator/Common/Private/UbaFileMapping.cpp
--- a/Engine/Source/Programs/UnrealBuildAccelerator/Common/Private/UbaFileMapping.cpp
+++ b/Engine/Source/Programs/UnrealBuildAccelerator/Common/Private/UbaFileMapping.cpp
@@ -452,8 +452,28 @@ namespace uba
 		CloseMappingStorage(m_storage[MappedView_Persistent]);
 	}
 
-	MappedView FileMappingBuffer::AllocAndMapView(FileMappingType type, u64 size, u64 alignment, const tchar* hint, bool allowShrink)
+	MappedView FileMappingBuffer::AllocAndMapView(FileMappingType type, u64 size, u64 alignment, const tchar* hint, bool allowShrink, bool isIndependent)
 	{
+		if (isIndependent)
+		{
+			UBA_ASSERT(!allowShrink);
+			FileMappingHandle mapping = CreateMemoryMappingW(m_logger, PAGE_READWRITE, size, hint, TC("Independent"));
+			if (!mapping.IsValid())
+				return {};
+			u8* memory = MapViewOfFile(m_logger, mapping, FILE_MAP_WRITE, 0, size);
+			if (!memory)
+			{
+				CloseFileMapping(m_logger, mapping, hint);
+				return {};
+			}
+			MappedView res;
+			res.handle = mapping;
+			res.memory = memory;
+			res.size = size;
+			res.isIndependent = true;
+			return res;
+		}
+
 		MappingStorage& storage = m_storage[type];
 		File* file = PopFile(storage, size, alignment);
 		if (!file)
@@ -638,14 +658,29 @@ namespace uba
 		storage.availableFile = nullptr;
 	}
 
-	MappedView FileMappingBuffer::MapView(FileMappingHandle handle, u64 offset, u64 size, const tchar* hint)
+	MappedView FileMappingBuffer::MapView(FileMappingHandle handle, u64 offset, u64 size, const tchar* hint, bool isIndependent)
 	{
 		UBA_ASSERT(handle.IsValid());// && (handle == m_files[0].handle || handle == m_files[1].handle));
+
+		MappedView res;
+
+		if (isIndependent)
+		{
+			u8* data = MapViewOfFile(m_logger, handle, FILE_MAP_WRITE, offset, size);
+			if (!data)
+				return {};
+			res.handle = handle;
+			res.offset = offset;
+			res.size = size;
+			res.memory = data;
+			res.isIndependent = true;
+			return res;
+		}
+
 		u8 storageIndex = 255;
 		File& file = GetFile(handle, storageIndex);
 
 		u64 alignedOffsetStart = 0;
-		MappedView res;
 
 		u8* data;
 		if (file.mappedMemory)
@@ -699,6 +734,13 @@ namespace uba
 
 		auto unmap = [=, this](const tchar* hint)
 			{
+				if (view.isIndependent)
+				{
+					if (!UnmapViewOfFile(m_logger, view.memory, view.size, hint))
+						m_logger.Error(TC("Failed to unmap view on address %llx (offset %llu) - %s (%s)"), u64(view.memory), view.offset, hint, LastErrorToText().data);
+					return;
+				}
+
 				u8 storageIndex = 255;
 				File& file = GetFile(view.handle, storageIndex);
 
diff --git a/Engine/Source/Programs/UnrealBuildAccelerator/Common/Public/UbaFileMapping.h b/Engine/Source/Programs/UnrealBuildAccelerator/Common/Public/UbaFileMapping.h
--- a/Engine/Source/Programs/UnrealBuildAccelerator/Common/Public/UbaFileMapping.h
+++ b/Engine/Source/Programs/UnrealBuildAccelerator/Common/Public/UbaFileMapping.h
@@ -45,6 +45,7 @@ namespace uba
 		u64 size = 0;
 		u8* memory = nullptr;
 		bool isCompressed = true;
+		bool isIndependent = false;
 	};
 
 	enum FileMappingType : u8
@@ -63,8 +64,8 @@ namespace uba
 		bool AddPersistent(const tchar* name, FileHandle fileHandle, u64 size, u64 capacity);
 		void CloseDatabase();
 
-		MappedView AllocAndMapView(FileMappingType type, u64 size, u64 alignment, const tchar* hint, bool allowShrink = false);
-		MappedView MapView(FileMappingHandle handle, u64 offset, u64 size, const tchar* hint);
+		MappedView AllocAndMapView(FileMappingType type, u64 size, u64 alignment, const tchar* hint, bool allowShrink = false, bool isIndependent = false);
+		MappedView MapView(FileMappingHandle handle, u64 offset, u64 size, const tchar* hint, bool isIndependent = false);
 		MappedView MapView(const StringView& str, u64 size, const tchar* hint);
 		void UnmapView(MappedView view, const tchar* hint, u64 newSize = InvalidValue);
 		
