[UBA]
* Fixed bug where real application path was sent in as first command line argument. This caused newer versions of clang++ to see the real path instead of the virtualized/host-local path when using argv[0]

diff --git a/Engine/Source/Programs/UnrealBuildAccelerator/Common/Private/UbaProcess.cpp b/Engine/Source/Programs/UnrealBuildAccelerator/Common/Private/UbaProcess.cpp
--- a/Engine/Source/Programs/UnrealBuildAccelerator/Common/Private/UbaProcess.cpp
+++ b/Engine/Source/Programs/UnrealBuildAccelerator/Common/Private/UbaProcess.cpp
@@ -1538,7 +1538,7 @@ namespace uba
 			if (!*detoursLib)
 				detoursLib = UBA_DETOURS_LIBRARY_ANSI;
 
-			StringBuffer<> application(m_realApplication);
+			StringBuffer<> application(m_startInfo.applicationStr);
 			m_session.VirtualizePath(application, m_startInfo.rootsHandle);
 
 			TString commandLine;
diff --git a/Engine/Source/Programs/UnrealBuildAccelerator/Detours/Private/Windows/UbaDetoursFunctionsKernelBase.inl b/Engine/Source/Programs/UnrealBuildAccelerator/Detours/Private/Windows/UbaDetoursFunctionsKernelBase.inl
--- a/Engine/Source/Programs/UnrealBuildAccelerator/Detours/Private/Windows/UbaDetoursFunctionsKernelBase.inl
+++ b/Engine/Source/Programs/UnrealBuildAccelerator/Detours/Private/Windows/UbaDetoursFunctionsKernelBase.inl
@@ -4099,6 +4099,14 @@ LPSTR Detoured_GetCommandLineA()
 	return g_virtualCommandLineA;
 }
 
+LPWSTR* Detoured_CommandLineToArgvW(LPCWSTR lpCmdLine, int *pNumArgs)
+{
+	DETOURED_CALL(CommandLineToArgvW);
+	LPWSTR* res = True_CommandLineToArgvW(lpCmdLine, pNumArgs);
+	DEBUG_LOG_TRUE(L"CommandLineToArgvW", L" -> %ls", res[0]);
+	return res;
+}
+
 BOOL Detoured_FreeLibrary(HMODULE hModule)
 {
 	DETOURED_CALL(FreeLibrary);
diff --git a/Engine/Source/Programs/UnrealBuildAccelerator/Detours/Private/Windows/UbaDetoursFunctionsWin.cpp b/Engine/Source/Programs/UnrealBuildAccelerator/Detours/Private/Windows/UbaDetoursFunctionsWin.cpp
--- a/Engine/Source/Programs/UnrealBuildAccelerator/Detours/Private/Windows/UbaDetoursFunctionsWin.cpp
+++ b/Engine/Source/Programs/UnrealBuildAccelerator/Detours/Private/Windows/UbaDetoursFunctionsWin.cpp
@@ -1069,6 +1069,8 @@ void Init(const DetoursPayload& payload, u64 startTime)
 
 	const wchar_t* cmdLine = True_GetCommandLineW();
 
+	// Command line should always have the virtualized proper host path, if it doesn't, then it is CreateProcess call that is wrong
+	#if 0
 	const wchar_t* exePos;
 	if (g_runningRemote && Contains(cmdLine, g_exeDir.data, true, &exePos))
 	{
@@ -1078,6 +1080,7 @@ void Init(const DetoursPayload& payload, u64 startTime)
 		TString realCmdLine(buf.data);
 		realCmdLine += (cmdLine + g_exeDir.count + 1);
 		g_virtualCommandLineW = g_memoryBlock.Strdup(realCmdLine).data;
+		cmdLine = g_virtualCommandLineW;
 	}
 	//else
 	//	g_virtualCommandLineW = g_memoryBlock.Strdup(cmdLine);
@@ -1089,12 +1092,11 @@ void Init(const DetoursPayload& payload, u64 startTime)
 		size_t res;
 		wcstombs_s(&res, g_virtualCommandLineA, len + 1, g_virtualCommandLineW, len);
 	}
+	#endif
 
 	#if UBA_DEBUG_LOG_ENABLED
 	if (isLogging())
 	{
-		if (g_virtualCommandLineW)
-			cmdLine = g_virtualCommandLineW;
 		u64 cmdLineLen = wcslen(cmdLine);
 		wchar_t temp[LogBufSize - 10];
 		if (cmdLineLen > sizeof_array(temp))
diff --git a/Engine/Source/Programs/UnrealBuildAccelerator/Detours/Private/Windows/UbaDetoursFunctionsWin.h b/Engine/Source/Programs/UnrealBuildAccelerator/Detours/Private/Windows/UbaDetoursFunctionsWin.h
--- a/Engine/Source/Programs/UnrealBuildAccelerator/Detours/Private/Windows/UbaDetoursFunctionsWin.h
+++ b/Engine/Source/Programs/UnrealBuildAccelerator/Detours/Private/Windows/UbaDetoursFunctionsWin.h
@@ -230,6 +230,7 @@
 
 #define DETOURED_FUNCTIONS_KERNELBASE_DEBUG \
 	DETOURED_FUNCTION(GetCommandLineA) \
+	DETOURED_FUNCTION(CommandLineToArgvW) \
 	DETOURED_FUNCTION(FreeLibrary) \
 	DETOURED_FUNCTION(RegOpenKeyW) \
 	DETOURED_FUNCTION(RegOpenKeyExW) \
