[UBA]
* Fixed code path involving memory mapped files and using ReadFile with overlapped structure provided

diff --git a/Engine/Source/Programs/UnrealBuildAccelerator/Detours/Private/Windows/UbaDetoursFunctionsKernelBase.inl b/Engine/Source/Programs/UnrealBuildAccelerator/Detours/Private/Windows/UbaDetoursFunctionsKernelBase.inl
--- a/Engine/Source/Programs/UnrealBuildAccelerator/Detours/Private/Windows/UbaDetoursFunctionsKernelBase.inl
+++ b/Engine/Source/Programs/UnrealBuildAccelerator/Detours/Private/Windows/UbaDetoursFunctionsKernelBase.inl
@@ -380,13 +380,20 @@ BOOL Detoured_ReadFile(HANDLE hFile, LPVOID lpBuffer, DWORD nNumberOfBytesToRead
 
 			u8* mem = fi.fileMapMem ? fi.fileMapMem : fi.memoryFile->baseAddress;
 			u64 size = fi.fileMapMem ? fi.fileMapMemSize : fi.memoryFile->writtenSize;
-			UBA_ASSERTF(dh.pos <= size, L"Filepointer is higher than size of file (pointer: %llu, size: %llu) (%ls)", dh.pos, size, HandleToName(hFile));
-			u64 leftToRead = size - dh.pos;
+			u64 readPos = dh.pos;
+			if (lpOverlapped)
+			{
+				UBA_ASSERT(!lpOverlapped->hEvent);
+				readPos = ToLargeInteger(lpOverlapped->OffsetHigh, lpOverlapped->Offset).QuadPart;
+			}
+			UBA_ASSERTF(readPos <= size, L"ReadFile - Read position is higher than size of file (pointer: %llu, size: %llu) (%ls)", readPos, size, HandleToName(hFile));
+			u64 leftToRead = size - readPos;
 			if (nNumberOfBytesToRead > leftToRead)
 				nNumberOfBytesToRead = (DWORD)leftToRead;
 			if (nNumberOfBytesToRead)
-				memcpy(lpBuffer, mem + dh.pos, nNumberOfBytesToRead);
-			dh.pos += nNumberOfBytesToRead;
+				memcpy(lpBuffer, mem + readPos, nNumberOfBytesToRead);
+			if (!lpOverlapped)
+				dh.pos += nNumberOfBytesToRead;
 			if (lpNumberOfBytesRead)
 				*lpNumberOfBytesRead = nNumberOfBytesToRead;
 			SetLastError(ERROR_SUCCESS);
