[UBA]
* Changed Shared_GetFileAttributes to not use KeepInMemory function to solve problems related to writefilestodisk=false. It also now always first check mapped files for information. Need to keep an eye out for file id issues since it is now easier to get files with id 0 since file mapping is not tracking file id.

* Removed shared memory temp file path. Using "outputfile" path instead since we added support for transferring file handle ownership some while ago.
* Fixed race condition in unit test code related to temporary files used by each unit test.

diff --git a/Engine/Source/Programs/UnrealBuildAccelerator/Common/Private/UbaProcess.cpp b/Engine/Source/Programs/UnrealBuildAccelerator/Common/Private/UbaProcess.cpp
--- a/Engine/Source/Programs/UnrealBuildAccelerator/Common/Private/UbaProcess.cpp
+++ b/Engine/Source/Programs/UnrealBuildAccelerator/Common/Private/UbaProcess.cpp
@@ -183,7 +183,6 @@ namespace uba
 			for (auto& pair : m_shared.writtenFiles)
 				if (pair.second.mappingHandle.IsValid())
 					CloseFileMapping(m_session.m_logger, pair.second.mappingHandle, pair.second.name.c_str());
-			ClearTempFiles();
 			delete &m_shared;
 		}
 	}
@@ -588,7 +587,6 @@ namespace uba
 			m_trackedInputs.clear();
 
 			m_shared.writtenFiles.clear();
-			ClearTempFiles();
 		}
 
 		#if PLATFORM_WINDOWS
@@ -633,9 +631,6 @@ namespace uba
 			m_comMemory = {};
 		}
 
-		if (!m_parentProcess)
-			ClearTempFiles();
-
 		if (m_startInfo.exitedFunc)
 		{
 			auto exitedFunc = m_startInfo.exitedFunc;
@@ -956,21 +951,6 @@ namespace uba
 	{
 		writer.WriteU32(m_session.GetDirectoryTableSize());
 		writer.WriteU32(m_session.GetFileMappingSize());
-
-		if (m_tempFilesModified)
-		{
-			// This should be very rare (lld-link.exe uses it for mt.exe).. if we get too many temp files we'll need to rethink this design
-			SCOPED_READ_LOCK(m_shared.tempFilesLock, l);
-			writer.WriteU32(u32(m_shared.tempFiles.size()));
-			for (auto& kv : m_shared.tempFiles)
-			{
-				writer.WriteStringKey(kv.first);
-				writer.WriteU64(kv.second.mappingWritten);
-			}
-		}
-		else
-			writer.WriteU32(0);
-
 		return HandleGetWrittenFiles(reader, writer);
 	}
 
@@ -1157,18 +1137,6 @@ namespace uba
 		return true;
 	}
 
-	bool ProcessImpl::HandleCreateTempFile(BinaryReader& reader, BinaryWriter& writer)
-	{
-		CreateTempFile(reader, m_nativeProcessHandle, m_startInfo.application);
-		return true;
-	}
-
-	bool ProcessImpl::HandleOpenTempFile(BinaryReader& reader, BinaryWriter& writer)
-	{
-		OpenTempFile(reader, writer, m_startInfo.application);
-		return true;
-	}
-
 	bool ProcessImpl::HandleVirtualAllocFailed(BinaryReader& reader, BinaryWriter& writer)
 	{
 		StringBuffer<> allocType;
@@ -1300,8 +1268,6 @@ namespace uba
 		m_trackedOutputs.clear();
 		m_shared.writtenFiles.clear();
 
-		ClearTempFiles();
-		
 		m_processStats = {};
 		m_sessionStats = {};
 		m_storageStats = {};
@@ -1387,58 +1353,6 @@ namespace uba
 		return true;
 	}
 
-	bool ProcessImpl::CreateTempFile(BinaryReader& reader, ProcHandle nativeProcessHandle, const tchar* application)
-	{
-		StringKey key = reader.ReadStringKey();
-		StringBuffer<> fileName;
-		reader.ReadString(fileName);
-		FileMappingHandle mappingHandle = FileMappingHandle::FromU64(reader.ReadU64());
-		u64 mappingHandleSize = reader.ReadU64();
-
-		FileMappingHandle newHandle;
-		if (!DuplicateFileMapping(m_session.m_logger, nativeProcessHandle, mappingHandle, GetCurrentProcessHandle(), newHandle, 0, false, DUPLICATE_SAME_ACCESS, fileName.data))
-		{
-			m_session.m_logger.Error(TC("Failed to duplicate handle for temp file (%s)"), fileName.data);
-			return true;
-		}
-
-		SCOPED_WRITE_LOCK(m_shared.tempFilesLock, tempLock);
-		auto insres = m_shared.tempFiles.try_emplace(key, WrittenFile{ nullptr, StringKeyZero, fileName.data, fileName.data, newHandle, mappingHandleSize, mappingHandle });
-		if (insres.second)
-			return true;
-
-		++m_tempFilesModified;
-		WrittenFile& tempFile = insres.first->second;
-		FileMappingHandle oldMapping = tempFile.mappingHandle;
-		tempFile.mappingHandle = newHandle;
-		tempFile.mappingWritten = mappingHandleSize;
-		tempLock.Leave();
-		CloseFileMapping(m_session.m_logger, oldMapping, fileName.data);
-		return true;
-	}
-
-	bool ProcessImpl::OpenTempFile(BinaryReader& reader, BinaryWriter& writer, const tchar* application)
-	{
-		StringKey fileKey = reader.ReadStringKey();
-		StringBuffer<> fileName;
-		reader.ReadString(fileName);
-					
-		u64 mappingHandle = 0;
-		u64 mappingWritten = 0;
-
-		SCOPED_READ_LOCK(m_shared.tempFilesLock, lock);
-		auto findIt = m_shared.tempFiles.find(fileKey);
-		if (findIt != m_shared.tempFiles.end())
-		{
-			mappingHandle = findIt->second.mappingHandle.ToU64();
-			mappingWritten = findIt->second.mappingWritten;
-		}
-
-		writer.WriteU64(mappingHandle);
-		writer.WriteU64(mappingWritten);
-		return true;
-	}
-
 	bool ProcessImpl::WriteFilesToDisk(bool isExiting)
 	{
 		Vector<WrittenFile*> files;
@@ -2255,14 +2169,6 @@ namespace uba
 	}
 #endif
 
-	void ProcessImpl::ClearTempFiles()
-	{
-		for (auto& pair : m_shared.tempFiles)
-			if (pair.second.mappingHandle.IsValid())
-				CloseFileMapping(m_session.m_logger, pair.second.mappingHandle, pair.second.name.c_str());
-		m_shared.tempFiles.clear();
-	}
-
 	void ProcessImpl::WaitForParent()
 	{
 		u64 startTime = GetTime();
diff --git a/Engine/Source/Programs/UnrealBuildAccelerator/Common/Private/UbaSession.cpp b/Engine/Source/Programs/UnrealBuildAccelerator/Common/Private/UbaSession.cpp
--- a/Engine/Source/Programs/UnrealBuildAccelerator/Common/Private/UbaSession.cpp
+++ b/Engine/Source/Programs/UnrealBuildAccelerator/Common/Private/UbaSession.cpp
@@ -239,6 +239,18 @@ namespace uba
 
 	bool Session::WriteFileToDisk(ProcessImpl& process, WrittenFile& file)
 	{
+		auto& rules = *process.m_startInfo.rules;
+
+		if (StringView(file.name).StartsWith(m_tempPath) && !rules.KeepTempOutputFile(StringView(file.name)))
+		{
+			CloseFileMapping(m_logger, file.mappingHandle, file.backedName.c_str());
+			file.mappingHandle = {};
+			#if UBA_DEBUG
+			m_logger.Info(TC("Skipping writing temp file %s to disk"), file.name.c_str());
+			#endif
+			return true;
+		}
+
 		bool shouldEvictFromMemory = IsRarelyReadAfterWritten(process, file.name) || file.mappingWritten > m_keepOutputFileMemoryMapsThreshold;
 
 		u64 writtenSize = 0;
@@ -246,7 +258,6 @@ namespace uba
 
 		if (shouldWriteToDisk)
 		{
-			auto& rules = *process.m_startInfo.rules;
 			RootsHandle rootsHandle = process.GetStartInfo().rootsHandle;
 
 			bool storeCompressed = m_storeIntermediateFilesCompressed && g_globalRules.FileCanBeCompressed(file.name);
@@ -2960,6 +2971,9 @@ namespace uba
 
 	bool Session::DeleteFile(DeleteFileResponse& out, const DeleteFileMessage& msg)
 	{
+		out.result = true;
+		bool deleteRealFile = true;
+
 		if (msg.closeId != 0)
 		{
 			SCOPED_FUTEX(m_activeFilesLock, lock);
@@ -2971,9 +2985,26 @@ namespace uba
 			m_outputFiles.erase(msg.fileName.data);
 		}
 
+		{
+			auto& shared = msg.process.m_shared;
+			SCOPED_FUTEX(shared.writtenFilesLock, lock);
+			auto findIt = shared.writtenFiles.find(msg.fileNameKey);
+			if (findIt != shared.writtenFiles.end())
+			{
+				WrittenFile& writtenFile = findIt->second;
+				if (writtenFile.mappingHandle.IsValid())
+				{
+					CloseFileMapping(m_logger, writtenFile.mappingHandle, writtenFile.name.c_str());
+					deleteRealFile = false;
+				}
+				shared.writtenFiles.erase(findIt);
+			}
+		}
+
 		RemoveWrittenFile(msg.process, msg.fileNameKey);
 
-		out.result = uba::DeleteFileW(msg.fileName.data);
+		if (deleteRealFile)
+			out.result = uba::DeleteFileW(msg.fileName.data);
 		out.errorCode = GetLastError();
 		out.directoryTableSize = RegisterDeleteFile(msg.fileNameKey, msg.fileName);
 		return true;
@@ -3023,6 +3054,7 @@ namespace uba
 				UBA_ASSERT(insres.second);
 				WrittenFile& newFile = insres.first->second;
 				newFile = oldFile;
+				newFile.key = msg.toKey;
 				newFile.name = msg.toName.data;
 				fs.writtenFiles.erase(findIt);
 
diff --git a/Engine/Source/Programs/UnrealBuildAccelerator/Common/Private/UbaSessionClient.cpp b/Engine/Source/Programs/UnrealBuildAccelerator/Common/Private/UbaSessionClient.cpp
--- a/Engine/Source/Programs/UnrealBuildAccelerator/Common/Private/UbaSessionClient.cpp
+++ b/Engine/Source/Programs/UnrealBuildAccelerator/Common/Private/UbaSessionClient.cpp
@@ -690,17 +690,30 @@ namespace uba
 
 	bool SessionClient::SendFiles(ProcessImpl& process, Timer& sendFiles)
 	{
+		auto& rules = *process.m_startInfo.rules;
+
 		StorageStatsScope storageStatsScope(process.m_storageStats);
 		for (auto& pair : process.m_shared.writtenFiles)
 		{
 			WrittenFile& file = pair.second;
+
+			if (StringView(file.name).StartsWith(m_tempPath) && !rules.KeepTempOutputFile(StringView(file.name)))
+			{
+				CloseFileMapping(m_logger, file.mappingHandle, file.backedName.c_str());
+				file.mappingHandle = {};
+				#if UBA_DEBUG
+				m_logger.Info(TC("Skipping sending back temp file %s"), file.name.c_str());
+				#endif
+				continue;
+			}
+
 			TimerScope timer(sendFiles);
 #ifdef _DEBUG
 			if (!pair.second.mappingHandle.IsValid())
-				m_logger.Warning(TC("%s is not using file mapping"), pair.first.c_str());
+				m_logger.Warning(TC("%s is not using file mapping"), file.name.c_str());
 #endif
 			bool keepMappingInMemory = IsWindows && !IsRarelyReadAfterWritten(process, file.name);
-			bool compressed = process.m_startInfo.rules->SendFileCompressedFromClient(file.name);
+			bool compressed = rules.SendFileCompressedFromClient(file.name);
 			if (!SendFile(file, process.GetId(), keepMappingInMemory, compressed))
 				return false;
 		}
diff --git a/Engine/Source/Programs/UnrealBuildAccelerator/Common/Public/UbaProcess.h b/Engine/Source/Programs/UnrealBuildAccelerator/Common/Public/UbaProcess.h
--- a/Engine/Source/Programs/UnrealBuildAccelerator/Common/Public/UbaProcess.h
+++ b/Engine/Source/Programs/UnrealBuildAccelerator/Common/Public/UbaProcess.h
@@ -96,7 +96,6 @@ namespace uba
 		u32 InternalCreateProcess(void* environment, FileMappingHandle communicationHandle, u64 communicationOffset);
 		u32	InternalExitProcess(bool cancel);
 		void InternalLogLine(bool printInSession, TString&& line, LogEntryType logType);
-		void ClearTempFiles();
 		void WaitForParent();
 		void WaitForChildrenExit();
 		bool CancelWithError(bool dummy = true);
@@ -135,8 +134,6 @@ namespace uba
 		{
 			Futex writtenFilesLock;
 			UnorderedMap<StringKey, WrittenFile> writtenFiles;
-			ReaderWriterLock tempFilesLock;
-			UnorderedMap<StringKey, WrittenFile> tempFiles;
 		};
 
 		u32 m_nativeProcessId = 0;
@@ -161,7 +158,6 @@ namespace uba
 		Vector<u8> m_trackedOutputs;
 		Vector<ProcessHandle> m_childProcesses;
 		Shared& m_shared;
-		Atomic<u32> m_tempFilesModified;
 		SessionStats m_sessionStats;
 		StorageStats m_storageStats;
 		ProcessStats m_processStats;
diff --git a/Engine/Source/Programs/UnrealBuildAccelerator/Core/Private/UbaApplicationRules.cpp b/Engine/Source/Programs/UnrealBuildAccelerator/Core/Private/UbaApplicationRules.cpp
--- a/Engine/Source/Programs/UnrealBuildAccelerator/Core/Private/UbaApplicationRules.cpp
+++ b/Engine/Source/Programs/UnrealBuildAccelerator/Core/Private/UbaApplicationRules.cpp
@@ -60,11 +60,11 @@ namespace uba
 				return 1024 * 1024 * 1024;
 			return Super::FileTypeMaxSize(file, isSystemOrTempFile);
 		}
-		virtual bool IsThrowAway(const StringView& fileName, bool isRunningRemote) const override
+		virtual bool IsThrowAway(StringView fileName, bool isRunningRemote) const override
 		{
 			return fileName.Contains(TCV("vctip_")) || Super::IsThrowAway(fileName, isRunningRemote);
 		}
-		virtual bool KeepInMemory(const StringView& fileName, const tchar* systemTemp, bool isRunningRemote, bool isWrite) const override
+		virtual bool KeepInMemory(StringView fileName, StringView systemTemp, bool isRunningRemote, bool isWrite) const override
 		{
 			if (fileName.Contains(TCV("\\vctip_")))
 				return true;
@@ -90,6 +90,8 @@ namespace uba
 				|| fileName.EndsWith(TCV(".pch"))
 				|| fileName.EndsWith(TCV(".ifc"))
 				|| fileName.EndsWith(TCV(".rc2.res")) // Not really an obj file.. 
+
+				|| fileName.Contains(TCV("\\_cl_")); // This file is needed when cl.exe spawns link.exe
 				;
 		}
 
@@ -109,11 +111,6 @@ namespace uba
 				|| fileName.EndsWith(TCV(".dll"));
 		}
 
-		virtual bool NeedsSharedMemory(const tchar* file) const override
-		{
-			return Contains(file, TC("\\_cl_")); // This file is needed when cl.exe spawns link.exe
-		}
-
 		virtual bool IsCacheable() const override
 		{
 			return true;
@@ -152,10 +149,15 @@ namespace uba
 				|| fileName.EndsWith(TCV(".pdb"))
 				|| fileName.EndsWith(TCV(".dll"))
 				|| fileName.EndsWith(TCV(".exe"))
-				|| fileName.EndsWith(TCV(".rc2.res")); // Not really an obj file.. 
+				|| fileName.EndsWith(TCV(".rc2.res")) // Not really an obj file.. 
+
+				// temp files
+				|| fileName.Contains(TCV("lnk{")) // This file is shared from link.exe to mt.exe and rc.exe so we need to put it shared memory
+			 	|| fileName.Contains(TCV("\\_cl_")) // When link.exe is spawned by cl.exe we might use this which is in shared memory
+			 	|| fileName.EndsWith(TCV(".manifest")); // lld-link.exe is using a different name for files shared with child processes
 		}
 
-		virtual bool IsThrowAway(const StringView& fileName, bool isRunningRemote) const override
+		virtual bool IsThrowAway(StringView fileName, bool isRunningRemote) const override
 		{
 			return fileName.Contains(TCV(".sup.")); // .sup.lib/exp are throw-away files that we don't want created
 		}
@@ -164,12 +166,6 @@ namespace uba
 		{
 			return Contains(file, TC("vctip.exe")) == false; // This is hacky but we don't want to start vctip.exe
 		}
-		virtual bool NeedsSharedMemory(const tchar* file) const override
-		{
-			return Contains(file, TC("lnk{")) // This file is shared from link.exe to mt.exe and rc.exe so we need to put it shared memory
-			 	|| Contains(file, TC("\\_cl_")) // When link.exe is spawned by cl.exe we might use this which is in shared memory
-			 	|| ToView(file).EndsWith(TCV(".manifest")); // lld-link.exe is using a different name for files shared with child processes
-		}
 
 		virtual u64 FileTypeMaxSize(const StringView& file, bool isSystemOrTempFile) const override
 		{
@@ -243,7 +239,7 @@ namespace uba
 	{
 		using Super = ApplicationRulesVcLink;
 
-		virtual bool KeepInMemory(const StringView& fileName, const tchar* systemTemp, bool isRunningRemote, bool isWrite) const override
+		virtual bool KeepInMemory(StringView fileName, StringView systemTemp, bool isRunningRemote, bool isWrite) const override
 		{
 			return fileName.EndsWith(TCV(".manifest")) && fileName.Contains(systemTemp);
 		}
@@ -385,6 +381,16 @@ namespace uba
 		{
 			return DependencyCrawlerType_MsvcCompiler;
 		}
+
+		virtual bool IsOutputFile(const StringView& fileName) const override
+		{
+			return fileName.Contains(TCV("response-")); // file used by internally spawned clang-cl.exe. Note this is only generated when command line is very long
+		}
+
+		virtual bool KeepInMemory(StringView fileName, StringView systemTemp, bool isRunningRemote, bool isWrite) const override
+		{
+			return fileName.StartsWith(systemTemp);
+		}
 	};
 
 	class ApplicationRulesLdLLdExe : public ApplicationRulesClang
@@ -449,33 +455,25 @@ namespace uba
 	{
 		using Super = ApplicationRulesClangPlusPlusExe;
 
-		virtual bool IsThrowAway(const StringView& fileName, bool isRunningRemote) const override
+		virtual bool IsThrowAway(StringView fileName, bool isRunningRemote) const override
 		{
 			return fileName.EndsWith(TCV("-telemetry.json")) || Super::IsThrowAway(fileName, isRunningRemote);
 		}
-
-		//virtual bool NeedsSharedMemory(const tchar* file)
-		//{
-		//	return fileName.Contains(TCV("lto-llvm")); // Used by a clang based platform's link time optimization pass. Shared from lto process back to linker process
-		//}
 	};
 
 	class ApplicationRulesLdExePlatform1 : public ApplicationRules
 	{
 		using Super = ApplicationRules;
 
-		//virtual bool IsOutputFile(const StringView& fileName) const override
-		//{
-		//	return fileName.EndsWith(TCV(".self")) || Equals(file, TC("Symbols.map"));
-		//}
-		virtual bool KeepInMemory(const StringView& fileName, const tchar* systemTemp, bool isRunningRemote, bool isWrite) const override
+		virtual bool IsOutputFile(const StringView& fileName) const override
 		{
-			return Super::KeepInMemory(fileName, systemTemp, isRunningRemote, isWrite)
-				|| fileName.Contains(TCV("thinlto-"));// Used by a clang based platform's link time optimization pass. Shared from lto process back to linker process
+			return fileName.Contains(TCV("thinlto-")); // Used by a clang based platform's link time optimization pass. Shared from lto process back to linker process
 		}
-		virtual bool NeedsSharedMemory(const tchar* file) const override
+
+		virtual bool KeepInMemory(StringView fileName, StringView systemTemp, bool isRunningRemote, bool isWrite) const override
 		{
-			return Contains(file, TC("thinlto-")); // Used by a clang based platform's link time optimization pass. Shared from lto process back to linker process
+			return Super::KeepInMemory(fileName, systemTemp, isRunningRemote, isWrite)
+				|| fileName.Contains(TCV("thinlto-"));// Used by a clang based platform's link time optimization pass. Shared from lto process back to linker process
 		}
 
 		virtual bool CanDependOnCompressedFiles() override
@@ -493,16 +491,16 @@ namespace uba
 			return fileName.Contains(TCV(".self")) || Super::IsOutputFile(fileName);
 		}
 
-		virtual bool IsThrowAway(const StringView& fileName, bool isRunningRemote) const override
+		virtual bool KeepTempOutputFile(const StringView& fileName) const override
+		{
+			return fileName.Contains(TCV(".native."));
+		}
+
+		virtual bool IsThrowAway(StringView fileName, bool isRunningRemote) const override
 		{
 			return Super::IsThrowAway(fileName, isRunningRemote)
 				|| (fileName.EndsWith(TCV("-telemetry.json"))); // Seems like these can collide when running distributed thinlto
 		}
-
-		//virtual bool NeedsSharedMemory(const tchar* file)
-		//{
-		//	return fileName.Contains(TCV("lto-llvm")); // Used by a clang based platform's link time optimization pass. Shared from lto process back to linker process
-		//}
 	};
 
 	class ApplicationRulesLldExePlatform2 : public ApplicationRules
@@ -514,17 +512,12 @@ namespace uba
 			return fileName.Contains(TCV(".self"));
 		}
 
-		virtual bool IsThrowAway(const StringView& fileName, bool isRunningRemote) const override
+		virtual bool IsThrowAway(StringView fileName, bool isRunningRemote) const override
 		{
 			return Super::IsThrowAway(fileName, isRunningRemote)
 				|| (fileName.EndsWith(TCV("-telemetry.json"))); // Seems like these can collide when running distributed thinlto
 		}
 
-		//virtual bool NeedsSharedMemory(const tchar* file)
-		//{
-		//	return fileName.Contains(TCV("lto-llvm")); // Used by a clang based platform's link time optimization pass. Shared from lto process back to linker process
-		//}
-
 		virtual bool CanDependOnCompressedFiles() override
 		{
 			return true;
@@ -550,14 +543,14 @@ namespace uba
 			return !fileName.EndsWith(TCV(".split"));
 		}
 
-		virtual bool KeepInMemory(const StringView& fileName, const tchar* systemTemp, bool isRunningRemote, bool isWrite) const override
+		virtual bool KeepInMemory(StringView fileName, StringView systemTemp, bool isRunningRemote, bool isWrite) const override
 		{
 			if (isWrite && fileName.Contains(TCV("\\Split\\")) && fileName.EndsWith(TCV(".split")))
 				return true;
 			return false;
 		}
 
-		virtual bool IsThrowAway(const StringView& fileName, bool isRunningRemote) const override
+		virtual bool IsThrowAway(StringView fileName, bool isRunningRemote) const override
 		{
 			if (fileName.Contains(TCV("\\Split\\")) && fileName.EndsWith(TCV(".split")))
 				return true;
diff --git a/Engine/Source/Programs/UnrealBuildAccelerator/Core/Public/UbaApplicationRules.h b/Engine/Source/Programs/UnrealBuildAccelerator/Core/Public/UbaApplicationRules.h
--- a/Engine/Source/Programs/UnrealBuildAccelerator/Core/Public/UbaApplicationRules.h
+++ b/Engine/Source/Programs/UnrealBuildAccelerator/Core/Public/UbaApplicationRules.h
@@ -43,25 +43,18 @@ namespace uba
 		}
 
 		// Throw-away means that the file is temporary and will not be used after process exists. (By default these are kept in memory and never touch disk)
-		virtual bool IsThrowAway(const StringView& fileName, bool isRunningRemote) const
+		virtual bool IsThrowAway(StringView fileName, bool isRunningRemote) const
 		{
 			return false;
 		}
 
 		// Keep file in memory
 		// If this returns true it means that file will be kept in memory and never touch disk.
-		virtual bool KeepInMemory(const StringView& fileName, const tchar* systemTemp, bool isRunningRemote, bool isWrite) const
+		virtual bool KeepInMemory(StringView fileName, StringView systemTemp, bool isRunningRemote, bool isWrite) const
 		{
 			return IsThrowAway(fileName, isRunningRemote);
 		}
 
-		// For files that are kept in memory but shared between process (temporary files where one process write and another read)
-		// This only works if file is never read outside of the same process hierarchy
-		virtual bool NeedsSharedMemory(const tchar* file) const
-		{
-			return false;
-		}
-
 		// Max file size if using memory files. Setting this value to a proper size is important because of Wine on linux use this
 		// .. on windows this size can be very big since it is just virtual address space and nothing is committed...
 		// On wine (as of version 9.10) a ftruncate is done on the wine server causing large max sizes to take up actual disk space on the tmp partition
@@ -76,6 +69,12 @@ namespace uba
 			return false;
 		}
 
+		// In general, output files in temp are dropped and never gets written to disk or networked to host.
+		virtual bool KeepTempOutputFile(const StringView& fileName) const
+		{
+			return false;
+		}
+
 		// If returns false this means that all GetFileAttribute etc will return file-not-found on this file
 		virtual bool CanExist(const tchar* file) const
 		{
diff --git a/Engine/Source/Programs/UnrealBuildAccelerator/Core/Public/UbaProtocol.h b/Engine/Source/Programs/UnrealBuildAccelerator/Core/Public/UbaProtocol.h
--- a/Engine/Source/Programs/UnrealBuildAccelerator/Core/Public/UbaProtocol.h
+++ b/Engine/Source/Programs/UnrealBuildAccelerator/Core/Public/UbaProtocol.h
@@ -24,8 +24,6 @@ namespace uba
 		UBA_PROCESS_MESSAGE(CreateProcess) \
 		UBA_PROCESS_MESSAGE(StartProcess) \
 		UBA_PROCESS_MESSAGE(ExitChildProcess) \
-		UBA_PROCESS_MESSAGE(CreateTempFile) \
-		UBA_PROCESS_MESSAGE(OpenTempFile) \
 		UBA_PROCESS_MESSAGE(VirtualAllocFailed) \
 		UBA_PROCESS_MESSAGE(Log) \
 		UBA_PROCESS_MESSAGE(EchoOn) \
@@ -51,7 +49,7 @@ namespace uba
 		#undef UBA_PROCESS_MESSAGE
 	};
 
-	inline constexpr u32 ProcessMessageVersion = 1345;
+	inline constexpr u32 ProcessMessageVersion = 1346;
 
 	inline constexpr u32 CommunicationMemSize = IsWindows ? 64*1024 : 64*1024*2; // Macos expands some commandlines to be crazy long
 
diff --git a/Engine/Source/Programs/UnrealBuildAccelerator/Detours/Private/PosixOS/UbaDetoursFunctionsPosix.cpp b/Engine/Source/Programs/UnrealBuildAccelerator/Detours/Private/PosixOS/UbaDetoursFunctionsPosix.cpp
--- a/Engine/Source/Programs/UnrealBuildAccelerator/Detours/Private/PosixOS/UbaDetoursFunctionsPosix.cpp
+++ b/Engine/Source/Programs/UnrealBuildAccelerator/Detours/Private/PosixOS/UbaDetoursFunctionsPosix.cpp
@@ -27,6 +27,8 @@
 #include <crt_externs.h>
 #endif
 
+#define UBA_DETOUR_DEBUG 0//UBA_DEBUG
+
 namespace uba
 {
 	constexpr bool g_logToScreen = false;
@@ -119,7 +121,7 @@ using namespace uba;
 	DETOURED_FUNCTION_MACOS \
 
 
-#if UBA_DEBUG && PLATFORM_LINUX
+#if UBA_DETOUR_DEBUG && PLATFORM_LINUX
 #define DETOURED_FUNCTION_DEBUG \
 	DETOURED_FUNCTION(write) \
 
@@ -325,10 +327,17 @@ int Shared_open(const char* funcName, const char* file, int flags, int mode, con
 	FixPath(fileName, file);
 	DevirtualizePath(fileName);
 
+	u32 desiredAccess = GetFileAccessFlags(flags);
+	bool isWrite = desiredAccess & AccessFlag_Write;
+
+	#if UBA_DEBUG_LOG_ENABLED
+	const char* isWriteStr = isWrite ? " WRITE" : ""; (void)isWriteStr;
+	#endif
+
 	if (!CanDetour2(fileName) || fileName.Equals("/"))
 	{
 		int res = trueOpen(file, flags, mode);
-		DEBUG_LOG_TRUE(funcName, "NODETOUR (%s) -> %i", file, res);
+		DEBUG_LOG_TRUE(funcName, "NODETOUR%s (%s) -> %i", isWriteStr, file, res);
 		return res;
 	}
 
@@ -350,14 +359,8 @@ int Shared_open(const char* funcName, const char* file, int flags, int mode, con
 
 	u64 size = InvalidValue;
 	u32 closeId = 0;
-	u32 desiredAccess = GetFileAccessFlags(flags);
-	bool isWrite = desiredAccess & AccessFlag_Write;
 	u32 dirTableOffset = ~u32(0);
 
-#if UBA_DEBUG_LOG_ENABLED
-	const char* isWriteStr = isWrite ? " WRITE" : ""; (void)isWriteStr;
-	#endif
-
 	if (g_allowDirectoryCache)
 	{
 		if (!isWrite) // We need to skip SystemTemp.. lots of stuff going on there.
@@ -398,7 +401,7 @@ int Shared_open(const char* funcName, const char* file, int flags, int mode, con
 
 	const char* realFileName = fileName.data;
 
-	SCOPED_WRITE_LOCK(g_mappedFileTable.m_lookupLock, _);
+	SCOPED_WRITE_LOCK(g_mappedFileTable.m_lookupLock, fileTableLock);
 	auto insres = g_mappedFileTable.m_lookup.try_emplace(fileNameKey);
 	FileInfo& info = insres.first->second;
 	FileInfo* fileInfo = &info;
@@ -450,6 +453,7 @@ int Shared_open(const char* funcName, const char* file, int flags, int mode, con
 			realFileName = info.name;
 		}
 	}
+	fileTableLock.Leave();
 
 	if (realFileName[0] == '$')
 	{
@@ -581,8 +585,7 @@ void Shared_close(int fd, const TrueClose& trueClose)
 	g_fileHandles.erase(findIt);
 	lock.Leave();
 
-	trueClose();
-
+	auto closeGuard = MakeGuard([&]() { trueClose(); });
 	if (--fo->refCount)
 		return;
 
@@ -593,6 +596,11 @@ void Shared_close(int fd, const TrueClose& trueClose)
 
 	if (fo->closeId)
 	{
+		struct stat attr;
+		TRUE_WRAPPER(fstat)(fd, &attr); // is this safe? not sure about fflush and user space FILE* streams
+		fi.size = attr.st_size;
+		fi.isWritten = true;
+		closeGuard.Execute();
 		//for (auto& kv : g_fileHandles)
 		//	UBA_ASSERT(&fi != kv.second.fileObject->fileInfo);
 		Rpc_UpdateCloseHandle(path, fo->closeId, fo->deleteOnClose, fo->newName.c_str(), mappingHandle, mappingWritten, true);
@@ -1129,7 +1137,7 @@ UBA_EXPORT int UBA_WRAPPER(open)(const char* file, int flags, ...)
 	return Shared_open("open", file, flags, mode, [](const char* realFile, int flags, int mode) { return TRUE_WRAPPER(open)(realFile, flags, mode); });
 }
 
-#if UBA_DEBUG && PLATFORM_LINUX
+#if UBA_DETOUR_DEBUG && PLATFORM_LINUX
 UBA_EXPORT ssize_t UBA_WRAPPER(write)(int fd, const void* buf, size_t count)
 {
 	UBA_INIT_DETOUR(write, fd, buf, count);
@@ -1689,6 +1697,7 @@ UBA_EXPORT int UBA_WRAPPER(rename)(const char* oldpath, const char* newpath)
 			info = findIt->second;
 		info.originalName = g_memoryBlock.Strdup(fixedNewPath).data;
 		info.name = info.originalName;
+		info.isWritten = true;
 	}
 
 	g_directoryTable.ParseDirectoryTable(directoryTableSize);
@@ -2858,7 +2867,7 @@ namespace uba
 	{
 		int t = errno;
 
-		#if UBA_DEBUG && PLATFORM_LINUX
+		#if UBA_DETOUR_DEBUG && PLATFORM_LINUX
 		TRUE_WRAPPER(write)(g_debugFile, str, strLen);
 		#else
 		write(g_debugFile, str, strLen);
diff --git a/Engine/Source/Programs/UnrealBuildAccelerator/Detours/Private/UbaDetoursFileMappingTable.cpp b/Engine/Source/Programs/UnrealBuildAccelerator/Detours/Private/UbaDetoursFileMappingTable.cpp
--- a/Engine/Source/Programs/UnrealBuildAccelerator/Detours/Private/UbaDetoursFileMappingTable.cpp
+++ b/Engine/Source/Programs/UnrealBuildAccelerator/Detours/Private/UbaDetoursFileMappingTable.cpp
@@ -185,6 +185,10 @@ namespace uba
 			info.fileNameKey = key;
 			info.size = fileSize;
 
+			#if !PLATFORM_WINDOWS
+			info.isWritten = true;
+			#endif
+
 			if (mappingHandle.IsValid())
 				backedName.Clear().Append(':').AppendHex(mappingHandle.ToU64()).Append('-').AppendHex(0);
 
@@ -205,30 +209,6 @@ namespace uba
 			BinaryReader reader;
 			directoryTableSize = reader.ReadU32();
 			fileMappingTableSize = reader.ReadU32();
-
-#if PLATFORM_WINDOWS
-			if (u32 tempFileCount = reader.ReadU32())
-			{
-				SCOPED_WRITE_LOCK(g_mappedFileTable.m_lookupLock, _);
-				while (tempFileCount--)
-				{
-					StringKey fileNameKey = reader.ReadStringKey();
-					u64 fileSize = reader.ReadU64();
-					auto findIt = g_mappedFileTable.m_lookup.find(fileNameKey);
-					if (findIt == g_mappedFileTable.m_lookup.end())
-						continue;
-					FileInfo& info = findIt->second;
-					UBA_ASSERT(info.memoryFile);
-					if (!info.memoryFile)
-						continue;
-					info.memoryFile->writtenSize = fileSize;
-					if (fileSize <= info.memoryFile->committedSize)
-						continue;
-					UBA_ASSERT(info.memoryFile->committedSize == 0);
-					info.memoryFile->EnsureCommitted(DetouredHandle(HandleType_File), fileSize);
-				}
-			}
-#endif
 			UpdateWrittenFiles(reader);
 
 			DEBUG_LOG_PIPE(L"UpdateTables", L"");
diff --git a/Engine/Source/Programs/UnrealBuildAccelerator/Detours/Private/UbaDetoursFileMappingTable.h b/Engine/Source/Programs/UnrealBuildAccelerator/Detours/Private/UbaDetoursFileMappingTable.h
--- a/Engine/Source/Programs/UnrealBuildAccelerator/Detours/Private/UbaDetoursFileMappingTable.h
+++ b/Engine/Source/Programs/UnrealBuildAccelerator/Detours/Private/UbaDetoursFileMappingTable.h
@@ -62,6 +62,8 @@ namespace uba
 		u32 fileMapDesiredAccess = PAGE_READONLY;
 		u32 fileMapViewDesiredAccess = FILE_MAP_READ; // This is the parameter to MapViewOfFile
 		bool mappingChecked = false;
+		#else
+		bool isWritten = false;
 		#endif
 	};
 
diff --git a/Engine/Source/Programs/UnrealBuildAccelerator/Detours/Private/UbaDetoursShared.cpp b/Engine/Source/Programs/UnrealBuildAccelerator/Detours/Private/UbaDetoursShared.cpp
--- a/Engine/Source/Programs/UnrealBuildAccelerator/Detours/Private/UbaDetoursShared.cpp
+++ b/Engine/Source/Programs/UnrealBuildAccelerator/Detours/Private/UbaDetoursShared.cpp
@@ -415,178 +415,135 @@ namespace uba
 
 		memset(&outAttr.data, 0, sizeof(outAttr.data));
 
-		bool isInsideSystemTemp = fileName.StartsWith(g_systemTemp.data);
-		bool keepInMemory = KeepInMemory(fileName, false);
-		if (keepInMemory)
+		#if PLATFORM_WINDOWS
+		if (fileName[1] == ':' && fileName[3] == 0 && (ToLower(fileName[0]) == ToLower(g_virtualWorkingDir[0]) || ToLower(fileName[0]) == g_systemRoot[0]))
+		{
+			// This is the root of the drive.. let's just return it as a directory
+			outAttr.useCache = true;
+			outAttr.exists = true;
+			outAttr.lastError = ErrorSuccess;
+			outAttr.data.dwFileAttributes = FILE_ATTRIBUTE_DIRECTORY;
+			return fileName.data;
+		}
+		#endif
+
 		{
 			SCOPED_READ_LOCK(g_mappedFileTable.m_lookupLock, lock);
 			auto it = g_mappedFileTable.m_lookup.find(fileNameKey);
-			if (it == g_mappedFileTable.m_lookup.end() || it->second.deleted)
+			if (it != g_mappedFileTable.m_lookup.end())
 			{
-				if (isInsideSystemTemp)
+				FileInfo& fi = it->second;
+
+				if (fi.deleted)
 				{
-					outAttr.useCache = false;
+					//if (isInsideSystemTemp)
+					//{
+					//	outAttr.useCache = false;
+					//	return fileName.data;
+					//}
+					outAttr.useCache = true;
+					outAttr.exists = false;
+					outAttr.lastError = ErrorFileNotFound;
 					return fileName.data;
 				}
-				outAttr.useCache = true;
-				outAttr.exists = false;
-				outAttr.lastError = ErrorFileNotFound;
-			}
-			else
-			{
+
+				// TODO: need to implement
+				//outAttr.data.ftLastWriteTime = ;
+				//outAttr.volumeSerial =
+				//outAttr.fileIndex = 
+
+#if PLATFORM_WINDOWS
 				outAttr.useCache = true;
 				outAttr.exists = true;
 				outAttr.lastError = ErrorSuccess;
-#if PLATFORM_WINDOWS
-				LARGE_INTEGER li = ToLargeInteger(it->second.size);
+				LARGE_INTEGER li = ToLargeInteger(fi.size);
 				outAttr.data.nFileSizeLow = li.LowPart;
 				outAttr.data.nFileSizeHigh = li.HighPart;
 				outAttr.data.dwFileAttributes = FILE_ATTRIBUTE_NORMAL;
+				return fileName.data;
 #else
-				UBA_ASSERT(false);
+				if (fi.isWritten)
+				{
+					outAttr.useCache = true;
+					outAttr.exists = true;
+					outAttr.lastError = ErrorSuccess;
+					outAttr.data.st_mode = (mode_t)(S_IFREG | S_IRUSR | S_IWUSR | S_IRGRP | S_IROTH);
+					outAttr.data.st_size = fi.size;
+					return fileName.data;
+				}
 #endif
-				// TODO: Currently only used for waccess... need to implement below
-				//outAttr.data.ftLastWriteTime = ;
-				//outAttr.volumeSerial =
-				//outAttr.fileIndex = 
 			}
 		}
-#if PLATFORM_WINDOWS
-		else if (fileName[1] == ':' && fileName[3] == 0 && (ToLower(fileName[0]) == ToLower(g_virtualWorkingDir[0]) || ToLower(fileName[0]) == g_systemRoot[0]))
+
+		if (!g_allowDirectoryCache || fileName.StartsWith(g_systemTemp.data)) // We need to skip SystemTemp.. lots of stuff going on there.
+		{
+			outAttr.useCache = false;
+			return fileName.data;
+		}
+
+		// This is an optimization where we populate directory table and use that to figure out if file exists or not..
+		// .. in msvc's case it doesn't matter much because these tables are already up to date when msvc use CreateFile.
+		// .. clang otoh is using CreateFile with tooons of different paths trying to open files.. in remote worker case this becomes super expensive
+		u32 dirTableOffset = Rpc_GetEntryOffset(fileNameKey, fileName, checkIfDir);
+
+		if (dirTableOffset == ~u32(0))
 		{
-			// This is the root of the drive.. let's just return it as a directory
 			outAttr.useCache = true;
-			outAttr.exists = true;
-			outAttr.lastError = ErrorSuccess;
-			outAttr.data.dwFileAttributes = FILE_ATTRIBUTE_DIRECTORY;
+			outAttr.exists = false;
+			outAttr.lastError = ErrorFileNotFound;
+			return fileName.data;
 		}
-#else
-		//else if (StartsWith(fileName, g_applicationDir))
-		//{
-		//	outAttr.useCache = false;
-		//	return fileName;
-		//}
-#endif
-		else if (g_allowDirectoryCache)
+
+		DirectoryTable::EntryInformation info;
+		g_directoryTable.GetEntryInformation(info, dirTableOffset);
+
+		if (!info.attributes)
 		{
-			// This is an optimization where we populate directory table and use that to figure out if file exists or not..
-			// .. in msvc's case it doesn't matter much because these tables are already up to date when msvc use CreateFile.
-			// .. clang otoh is using CreateFile with tooons of different paths trying to open files.. in remote worker case this becomes super expensive
-			if (!isInsideSystemTemp) // We need to skip SystemTemp.. lots of stuff going on there.
-			{
-				u32 dirTableOffset = Rpc_GetEntryOffset(fileNameKey, fileName, checkIfDir);
+			// File used to exist but was deleted
+			outAttr.useCache = true;
+			outAttr.exists = false;
+			outAttr.lastError = ErrorFileNotFound;
+			return fileName.data;
+		}
 
-				if (dirTableOffset == ~u32(0))
-				{
-					// This could be a newly written file but process has not fetched latest directory table
-					SCOPED_READ_LOCK(g_mappedFileTable.m_lookupLock, lock);
-					auto findIt = g_mappedFileTable.m_lookup.find(fileNameKey);
-					if (findIt != g_mappedFileTable.m_lookup.end() && !findIt->second.deleted)
-					{
-						outAttr.exists = true;
-						outAttr.lastError = ErrorSuccess;
-						outAttr.useCache = false;
-
-						if (g_runningRemote)
-						{
-							FileInfo& info = findIt->second;
-							outAttr.useCache = true;
-
-							// TODO: This is missing lots of information..
-#if PLATFORM_WINDOWS
-							LARGE_INTEGER li = ToLargeInteger(info.size);
-							outAttr.data.dwFileAttributes = FILE_ATTRIBUTE_NORMAL;
-							outAttr.data.nFileSizeLow = li.LowPart;
-							outAttr.data.nFileSizeHigh = li.HighPart;
-#else
-							outAttr.data.st_mode = (mode_t)(S_IFREG | S_IRUSR | S_IWUSR | S_IRGRP | S_IROTH);
-							outAttr.data.st_size = info.size;
-#endif
-						}
+		u64 fileSize = info.size;
 
-						return findIt->second.name;
-					}
+		// Could be compressed and then directory table size is wrong
+		if (CouldBeCompressedFile(fileName))
+		{
+			// If file is output file we accept wrong size because size is not supposed to be used anyway.
+			// We don't want to trigger unnecessary download/decompress of file
+			if (!g_rules->IsOutputFile(fileName))
+			{
+				StringBuffer<> temp;
+				u32 closeId;
+				Rpc_CreateFileW(fileName, fileNameKey, AccessFlag_Read, temp.data, temp.capacity, fileSize, closeId, false);
+			}
+		}
 
-					outAttr.useCache = true;
-					outAttr.exists = false;
-					outAttr.lastError = ErrorFileNotFound;
-				}
-				else
-				{
-					DirectoryTable::EntryInformation info;
-					g_directoryTable.GetEntryInformation(info, dirTableOffset);
-
-					if (info.attributes)
-					{
-						u64 fileSize = info.size;
-
-						// Could be compressed and then directory table size is wrong
-						if (CouldBeCompressedFile(fileName))
-						{
-							SCOPED_READ_LOCK(g_mappedFileTable.m_lookupLock, lock);
-							auto findIt = g_mappedFileTable.m_lookup.find(fileNameKey);
-							if (findIt == g_mappedFileTable.m_lookup.end())
-							{
-								// If file is output file we accept wrong size because size is not supposed to be used anyway.
-								// We don't want to trigger unnecessary download/decompress of file
-								if (!g_rules->IsOutputFile(fileName))
-								{
-									StringBuffer<> temp;
-									u32 closeId;
-									Rpc_CreateFileW(fileName, fileNameKey, AccessFlag_Read, temp.data, temp.capacity, fileSize, closeId, false);
-								}
-							}
-							else
-							{
-								UBA_ASSERT(!findIt->second.deleted);
-								fileSize = findIt->second.size;
-							}
-						}
-
-						outAttr.useCache = true;
-						outAttr.exists = true;
-						outAttr.lastError = ErrorSuccess;
-
-						UBA_ASSERT(info.fileIndex);
-						outAttr.fileIndex = info.fileIndex;
-						outAttr.volumeSerial = info.volumeSerial;
+		outAttr.useCache = true;
+		outAttr.exists = true;
+		outAttr.lastError = ErrorSuccess;
+
+		UBA_ASSERT(info.fileIndex);
+		outAttr.fileIndex = info.fileIndex;
+		outAttr.volumeSerial = info.volumeSerial;
 
 #if PLATFORM_WINDOWS
-						LARGE_INTEGER li = ToLargeInteger(fileSize);
-						outAttr.data.dwFileAttributes = info.attributes;
-						outAttr.data.nFileSizeLow = li.LowPart;
-						outAttr.data.nFileSizeHigh = li.HighPart;
-						(u64&)outAttr.data.ftCreationTime = info.lastWrite;
-						(u64&)outAttr.data.ftLastAccessTime = info.lastWrite;
-						(u64&)outAttr.data.ftLastWriteTime = info.lastWrite;
+		LARGE_INTEGER li = ToLargeInteger(fileSize);
+		outAttr.data.dwFileAttributes = info.attributes;
+		outAttr.data.nFileSizeLow = li.LowPart;
+		outAttr.data.nFileSizeHigh = li.HighPart;
+		(u64&)outAttr.data.ftCreationTime = info.lastWrite;
+		(u64&)outAttr.data.ftLastAccessTime = info.lastWrite;
+		(u64&)outAttr.data.ftLastWriteTime = info.lastWrite;
 #else
-						outAttr.data.st_mtimespec = ToTimeSpec(info.lastWrite);
-						outAttr.data.st_mode = (mode_t)info.attributes;
-						outAttr.data.st_dev = info.volumeSerial;
-						outAttr.data.st_ino = info.fileIndex;
-						outAttr.data.st_size = fileSize;
+		outAttr.data.st_mtimespec = ToTimeSpec(info.lastWrite);
+		outAttr.data.st_mode = (mode_t)info.attributes;
+		outAttr.data.st_dev = info.volumeSerial;
+		outAttr.data.st_ino = info.fileIndex;
+		outAttr.data.st_size = fileSize;
 #endif
-					}
-					else
-					{
-						// File used to exist but was deleted
-						outAttr.useCache = true;
-						outAttr.exists = false;
-						outAttr.lastError = ErrorFileNotFound;
-					}
-				}
-			}
-			else
-			{
-				outAttr.useCache = false;
-				return fileName.data;
-			}
-		}
-		else
-		{
-			outAttr.useCache = false;
-			return fileName.data;
-		}
 
 #if 0//UBA_DEBUG_VALIDATE
 		if (g_validateFileAccess && !keepInMemory)
diff --git a/Engine/Source/Programs/UnrealBuildAccelerator/Detours/Private/UbaDetoursShared.h b/Engine/Source/Programs/UnrealBuildAccelerator/Detours/Private/UbaDetoursShared.h
--- a/Engine/Source/Programs/UnrealBuildAccelerator/Detours/Private/UbaDetoursShared.h
+++ b/Engine/Source/Programs/UnrealBuildAccelerator/Detours/Private/UbaDetoursShared.h
@@ -121,7 +121,6 @@ namespace uba
 		u32 lastError;
 	};
 
-	inline bool KeepInMemory(const StringView& fileName, bool isWrite) { return g_allowKeepFilesInMemory && g_rules->KeepInMemory(fileName, g_systemTemp.data, g_runningRemote, isWrite); }
 	bool CouldBeCompressedFile(const StringView& fileName);
 
 	void Shared_WriteConsole(const char* chars, u32 charCount, bool isError);
diff --git a/Engine/Source/Programs/UnrealBuildAccelerator/Detours/Private/Windows/UbaDetoursFunctionsKernelBase.inl b/Engine/Source/Programs/UnrealBuildAccelerator/Detours/Private/Windows/UbaDetoursFunctionsKernelBase.inl
--- a/Engine/Source/Programs/UnrealBuildAccelerator/Detours/Private/Windows/UbaDetoursFunctionsKernelBase.inl
+++ b/Engine/Source/Programs/UnrealBuildAccelerator/Detours/Private/Windows/UbaDetoursFunctionsKernelBase.inl
@@ -1106,12 +1106,24 @@ DWORD Detoured_GetFileAttributesW(LPCWSTR lpFileName)
 BOOL Detoured_SetFileAttributesW(LPCWSTR lpFileName, DWORD dwFileAttributes)
 {
 	DETOURED_CALL(SetFileAttributesW);
-	if (KeepInMemory(StringView(lpFileName, u32(wcslen(lpFileName))), true))
+
+	StringBuffer<> fixedFile;
+	FixPath(fixedFile, lpFileName);
+	DevirtualizePath(fixedFile);
+	StringKey fileKey = ToStringKeyLower(fixedFile);
+
+	SCOPED_WRITE_LOCK(g_mappedFileTable.m_lookupLock, lock);
+	auto it = g_mappedFileTable.m_lookup.find(fileKey);
+	bool isMapped = it != g_mappedFileTable.m_lookup.end();
+	lock.Leave();
+
+	if (isMapped || g_rules->IsThrowAway(fixedFile, g_runningRemote))
 	{
-		DEBUG_LOG_DETOURED(L"SetFileAttributesW", L"(%ls) %u", lpFileName, dwFileAttributes);
+		DEBUG_LOG_DETOURED(L"SetFileAttributesW", L"(IGNORED) (%ls) %u", lpFileName, dwFileAttributes);
 		SetLastError(ERROR_SUCCESS);
 		return true;
 	}
+
 	DEBUG_LOG_TRUE(L"SetFileAttributesW", L"(%ls) %u", lpFileName, dwFileAttributes);
 	TimerScope ts(g_kernelStats.setFileInfo);
 	return True_SetFileAttributesW(lpFileName, dwFileAttributes);
@@ -1529,19 +1541,6 @@ BOOL Detoured_DeleteFileW(LPCWSTR lpFileName)
 	StringBuffer<> fixedNameLower(fixedName);
 	fixedNameLower.MakeLower();
 
-	if (KeepInMemory(fixedNameLower, true))
-	{
-		DEBUG_LOG_DETOURED(L"DeleteFileW", L"(INMEMORY) (%ls) -> Success", lpFileName);
-		SetLastError(ERROR_SUCCESS);
-		return TRUE;
-	}
-
-	if (fixedName.StartsWith(g_systemTemp.data))
-	{
-		DEBUG_LOG_TRUE(L"DeleteFileW", L"(%ls)", original);
-		return True_DeleteFileW(original);
-	}
-
 	StringKey fileNameKey = ToStringKey(fixedNameLower);
 
 	u32 directoryTableSize;
@@ -1569,7 +1568,7 @@ BOOL Detoured_DeleteFileW(LPCWSTR lpFileName)
 	return result;
 }
 
-bool Shared_MoveFile(LPCWSTR lpExistingFileName, LPCWSTR lpNewFileName, DWORD dwFlags)
+bool Shared_MoveFile(const tchar* func, LPCWSTR lpExistingFileName, LPCWSTR lpNewFileName, DWORD dwFlags)
 {
 	DETOURED_CALL(MoveFileExW);
 
@@ -1583,8 +1582,11 @@ bool Shared_MoveFile(LPCWSTR lpExistingFileName, LPCWSTR lpNewFileName, DWORD dw
 
 	StringKey sourceKey = ToStringKeyLower(source);
 
+#if 0 // This path needs to be revisited... [hka] I can't figure out how to hit this... but KeepInMemory should not be used outside of NtCreateFile anymore
 	if (KeepInMemory(source, false))
 	{
+		UBA_ASSERT(false);
+
 		SCOPED_WRITE_LOCK(g_mappedFileTable.m_lookupLock, lock);
 		auto it = g_mappedFileTable.m_lookup.find(sourceKey);
 		UBA_ASSERTF(it != g_mappedFileTable.m_lookup.end(), L"Can't find %ls", source.data);
@@ -1617,14 +1619,12 @@ bool Shared_MoveFile(LPCWSTR lpExistingFileName, LPCWSTR lpNewFileName, DWORD dw
 			sourceInfo.memoryFile = nullptr;
 
 			CloseHandle(tempHandle);
-			DEBUG_LOG_DETOURED(L"MoveFileExW", L"(memfile->memfile) %ls to %ls -> Success", lpExistingFileName, lpNewFileName);
+			DEBUG_LOG_DETOURED(func, L"(memfile->memfile) %ls to %ls -> Success", lpExistingFileName, lpNewFileName);
 			SetLastError(ERROR_SUCCESS);
 			return true;
 		}
 
-		UBA_ASSERT(!KeepInMemory(dest, true));
-
-		DEBUG_LOG_DETOURED(L"MoveFileExW", L"(memfile->file) %ls to %ls", lpExistingFileName, lpNewFileName);
+		DEBUG_LOG_DETOURED(func, L"(memfile->file) %ls to %ls", lpExistingFileName, lpNewFileName);
 
 		HANDLE h = CreateFile(lpNewFileName, GENERIC_WRITE, 0, NULL, CREATE_ALWAYS, FILE_ATTRIBUTE_NORMAL, NULL);
 		if (h == INVALID_HANDLE_VALUE)
@@ -1645,6 +1645,7 @@ bool Shared_MoveFile(LPCWSTR lpExistingFileName, LPCWSTR lpNewFileName, DWORD dw
 		sourceInfo.deleted = true;
 		return true;
 	}
+#endif
 
 	StringKey destKey = ToStringKeyLower(dest);
 
@@ -1683,8 +1684,7 @@ bool Shared_MoveFile(LPCWSTR lpExistingFileName, LPCWSTR lpNewFileName, DWORD dw
 
 BOOL Detoured_MoveFileExW(LPCWSTR lpExistingFileName, LPCWSTR lpNewFileName, DWORD dwFlags)
 {
-	//UBA_ASSERT(!Contains(lpNewFileName, L"E:\\dev\\fn\\Engine\\Source\\Programs\\Shared\\EpicGames.UHT\\obj\\project.assets.json"));
-	return Shared_MoveFile(lpExistingFileName, lpNewFileName, dwFlags);
+	return Shared_MoveFile(TC("MoveFileExW"), lpExistingFileName, lpNewFileName, dwFlags);
 }
 
 // MoveFileW ends up here
@@ -1694,7 +1694,7 @@ BOOL Detoured_MoveFileWithProgressW(LPCWSTR lpExistingFileName, LPCWSTR lpNewFil
 	StringBuffer<> source;
 	FixPath(source, lpExistingFileName);
 
-	return Shared_MoveFile(lpExistingFileName, lpNewFileName, dwFlags);
+	return Shared_MoveFile(TC("MoveFileWithProgressW"), lpExistingFileName, lpNewFileName, dwFlags);
 
 	//UBA_ASSERT(!g_runningRemote);
 	//DEBUG_LOG_TRUE(L"MoveFileWithProgressW", L"%ls to %ls", lpExistingFileName, lpNewFileName);
diff --git a/Engine/Source/Programs/UnrealBuildAccelerator/Detours/Private/Windows/UbaDetoursFunctionsNtDll.inl b/Engine/Source/Programs/UnrealBuildAccelerator/Detours/Private/Windows/UbaDetoursFunctionsNtDll.inl
--- a/Engine/Source/Programs/UnrealBuildAccelerator/Detours/Private/Windows/UbaDetoursFunctionsNtDll.inl
+++ b/Engine/Source/Programs/UnrealBuildAccelerator/Detours/Private/Windows/UbaDetoursFunctionsNtDll.inl
@@ -856,7 +856,9 @@ NTSTATUS NTAPI Shared_NtCreateFile(bool IsCreateFunc, PHANDLE hFileHandle, ACCES
 	bool useContent = IsContentUse(DesiredAccess, CreateDisposition);
 	bool isWrite = IsWrite(DesiredAccess, CreateDisposition);
 	bool isThrowAway = g_rules->IsThrowAway(fileName, g_runningRemote);
-	bool keepInMemory = KeepInMemory(fileName, isWrite) || ((isWrite || isDeleteOnClose) && g_rules->IsOutputFile(fileName) && g_allowKeepFilesInMemory) || (isWrite && isThrowAway);
+	bool keepInMemory = g_allowKeepFilesInMemory && g_rules->KeepInMemory(fileName, g_systemTemp, g_runningRemote, isWrite);
+	
+	keepInMemory = keepInMemory || ((isWrite || isDeleteOnClose) && g_rules->IsOutputFile(fileName) && g_allowKeepFilesInMemory) || (isWrite && isThrowAway);
 
 #if UBA_DEBUG_LOG_ENABLED
 	const wchar_t* isWriteStr = isWrite ? L" WRITE" : L""; (void)isWriteStr;
@@ -1279,58 +1281,26 @@ NTSTATUS NTAPI Shared_NtCreateFile(bool IsCreateFunc, PHANDLE hFileHandle, ACCES
 
 		if (!info.memoryFile)
 		{
-			if (NeedsSharedMemory(fileName.data))
+			if (!isThrowAway && CreateDisposition == FILE_OPEN)
 			{
-				#if UBA_DEBUG_LOG_ENABLED
-				memoryType = L"SHAREDMEMORY";
-				#endif
+				if (!info.memoryFile)
 				{
-					RPC_MESSAGE(OpenTempFile, openTempFile)
-					writer.WriteStringKey(fileNameKey);
-					writer.WriteString(fileNameLower);
-					writer.Flush();
-					BinaryReader reader;
-					FileMappingHandle mappingHandle = FileMappingHandle::FromU64(reader.ReadU64());
-					u64 mappingHandleSize = reader.ReadU64();
-					pcs.Leave();
-					if (mappingHandle.IsValid())
-					{
-						TakeFileOwnershipFromSession(mappingHandle, mappingHandleSize);
-					}
-					else if (isWrite)
-					{
-						info.memoryFile = new MemoryFile(false, FileTypeMaxSize(fileName, isSystemOrTempFile), false, fileName.data);
-					}
-					else
-					{
-						DEBUG_LOG_DETOURED(funcName, L"(memory) not found (%ls)", fileName.data);
-						return STATUS_OBJECT_NAME_NOT_FOUND;
-					}
+					*hFileHandle = INVALID_HANDLE_VALUE;
+					DEBUG_LOG_DETOURED(funcName, L"ALREADYEXISTS (%ls) -> Error", fileName.data);
+					return STATUS_OBJECT_NAME_EXISTS;//STATUS_OBJECT_NAME_NOT_FOUND;
 				}
 			}
 			else
 			{
-				if (!isThrowAway && CreateDisposition == FILE_OPEN)
+				bool isOutput = (isWrite || isDeleteOnClose) && g_rules->IsOutputFile(fileName);
+				if (isOutput && (CreateDisposition & FILE_OPEN_IF) == 0)
 				{
-					if (!info.memoryFile)
-					{
-						*hFileHandle = INVALID_HANDLE_VALUE;
-						DEBUG_LOG_DETOURED(funcName, L"ALREADYEXISTS (%ls) -> Error", fileName.data);
-						return STATUS_OBJECT_NAME_EXISTS;//STATUS_OBJECT_NAME_NOT_FOUND;
-					}
+					UBA_ASSERTF(false, TC("Trying to open %s with openif. This is not supported"), fileName.data);
 				}
-				else
-				{
-					bool isOutput = (isWrite || isDeleteOnClose) && g_rules->IsOutputFile(fileName);
-					if (isOutput && (CreateDisposition & FILE_OPEN_IF) == 0)
-					{
-						UBA_ASSERTF(false, TC("Trying to open %s with openif. This is not supported"), fileName.data);
-					}
 
-					bool isLocal = !isOutput;
-					//UBA_ASSERTF(CreateDisposition != FILE_OPEN || Contains(fileName.data, L"vctip_"), TC("Unsupported disposition %u for file %s"), CreateDisposition, fileName.data);
-					info.memoryFile = new MemoryFile(isLocal, FileTypeMaxSize(fileName, isSystemOrTempFile), isThrowAway, fileName.data);
-				}
+				bool isLocal = !isOutput;
+				//UBA_ASSERTF(CreateDisposition != FILE_OPEN || Contains(fileName.data, L"vctip_"), TC("Unsupported disposition %u for file %s"), CreateDisposition, fileName.data);
+				info.memoryFile = new MemoryFile(isLocal, FileTypeMaxSize(fileName, isSystemOrTempFile), isThrowAway, fileName.data);
 			}
 
 			// TODO: Time should be in sync with host machine!
@@ -1611,20 +1581,6 @@ NTSTATUS NTAPI Detoured_NtClose(HANDLE handle)
 				fo->newName.clear();
 			}
 		}
-		else if (NeedsSharedMemory(fi.originalName) && IsWrite(fo->desiredAccess, 0))
-		{
-			UBA_ASSERT(!mf->isLocalOnly);
-			StringBuffer<> fixedName;
-			FixPath(fixedName, path);
-
-			RPC_MESSAGE(CreateTempFile, createTempFile)
-			writer.WriteStringKey(ToStringKeyLower(fixedName));
-			writer.WriteString(fixedName);
-			writer.WriteU64(mappingHandle.ToU64());
-			writer.WriteU64(mappingWritten);
-			writer.Flush();
-			BinaryReader reader;
-		}
 	}
 	else if (fo->deleteOnClose && dh.trueHandle == INVALID_HANDLE_VALUE) // We have used an optimized handle that actually never opens the file so we need to delete it manually
 	{
diff --git a/Engine/Source/Programs/UnrealBuildAccelerator/Detours/Private/Windows/UbaDetoursFunctionsWin.cpp b/Engine/Source/Programs/UnrealBuildAccelerator/Detours/Private/Windows/UbaDetoursFunctionsWin.cpp
--- a/Engine/Source/Programs/UnrealBuildAccelerator/Detours/Private/Windows/UbaDetoursFunctionsWin.cpp
+++ b/Engine/Source/Programs/UnrealBuildAccelerator/Detours/Private/Windows/UbaDetoursFunctionsWin.cpp
@@ -547,7 +547,6 @@ const wchar_t* HandleToName(HANDLE handle)
 }
 
 
-bool NeedsSharedMemory(const wchar_t* file) { return g_allowKeepFilesInMemory && g_rules->NeedsSharedMemory(file); }
 u64 FileTypeMaxSize(const StringBufferBase& file, bool isSystemOrTempFile) { return g_rules->FileTypeMaxSize(file, isSystemOrTempFile); }
 
 bool EnsureMapped(DetouredHandle& handle, DWORD dwFileOffsetHigh = 0, DWORD dwFileOffsetLow = 0, SIZE_T numberOfBytesToMap = 0, void* baseAddress = nullptr)
@@ -1067,7 +1066,7 @@ void Init(const DetoursPayload& payload, u64 startTime)
 			FatalError(4444, L"What the heck: %s (%s)", g_virtualApplication.data, applicationBuffer.data);
 	}
 
-	const wchar_t* cmdLine = True_GetCommandLineW();
+	const wchar_t* cmdLine = True_GetCommandLineW();(void)cmdLine;
 
 	// Command line should always have the virtualized proper host path, if it doesn't, then it is CreateProcess call that is wrong
 	#if 0
diff --git a/Engine/Source/Programs/UnrealBuildAccelerator/Test/Private/UbaTestAll.h b/Engine/Source/Programs/UnrealBuildAccelerator/Test/Private/UbaTestAll.h
--- a/Engine/Source/Programs/UnrealBuildAccelerator/Test/Private/UbaTestAll.h
+++ b/Engine/Source/Programs/UnrealBuildAccelerator/Test/Private/UbaTestAll.h
@@ -118,6 +118,7 @@ namespace uba
 	#define UBA_TEST(x) \
 		if (!filter || Contains(TC(#x), filter)) \
 		{ \
+			DeleteAllFiles(logger, testRootDir.data, false); \
 			logger.Info(TC("Running %s..."), TC(#x)); \
 			if (!x(testLogger, testRootDir)) \
 				return logger.Error(TC("  %s failed"), TC(#x)); \
@@ -141,8 +142,8 @@ namespace uba
 		testRootDir.EnsureEndsWithSlash().Append(L"UbaTest");
 		#else
 		testRootDir.count = GetFullPathNameW("~/UbaTest", testRootDir.capacity, testRootDir.data, nullptr);
+		//testRootDir.Append("/mnt/e/temp/ttt/UbaTest");
 		#endif
-		DeleteAllFiles(logger, testRootDir.data, false);
 		CreateDirectoryW(testRootDir.data);
 		testRootDir.EnsureEndsWithSlash();
 
@@ -157,7 +158,6 @@ namespace uba
 		UBA_TESTS
 
 		logger.Info(TC("Tests finished successfully!"));
-		Sleep(2000);
 
 		return true;
 	}
diff --git a/Engine/Source/Programs/UnrealBuildAccelerator/Test/Private/UbaTestSession.cpp b/Engine/Source/Programs/UnrealBuildAccelerator/Test/Private/UbaTestSession.cpp
--- a/Engine/Source/Programs/UnrealBuildAccelerator/Test/Private/UbaTestSession.cpp
+++ b/Engine/Source/Programs/UnrealBuildAccelerator/Test/Private/UbaTestSession.cpp
@@ -12,7 +12,16 @@
 
 namespace uba
 {
-	bool RunLocal(LoggerWithWriter& logger, const StringBufferBase& testRootDir, const TestSessionFunction& testFunc, bool enableDetour)
+	u32 GetTimeoutTime()
 	{
+		#if PLATFORM_WINDOWS
+		if (IsDebuggerPresent())
+			return 10'000'000;
+		#endif
+		return 10'000;
+	}
+
+	bool RunLocal(LoggerWithWriter& logger, const StringBufferBase& testRootDir, const Config& config, const TestSessionFunction& testFunc, bool enableDetour = true)
+	{
 		LogWriter& logWriter = logger.m_writer;
 
@@ -23,6 +32,7 @@
 
 		StorageCreateInfo storageInfo(rootDir.data, logWriter);
 		storageInfo.casCapacityBytes = 1024ull * 1024 * 1024;
+		storageInfo.Apply(config);
 		StorageImpl storage(storageInfo);
 
 		bool ctorSuccess = true;
@@ -31,6 +41,7 @@
 		SessionServerCreateInfo sessionServerInfo(storage, server, logWriter);
 		sessionServerInfo.checkMemory = false;
 		sessionServerInfo.rootDir = rootDir.data;
+		sessionServerInfo.Apply(config);
 
 		#if UBA_DEBUG
 		sessionServerInfo.logToFile = true;
@@ -51,11 +62,52 @@
 		return testFunc(logger, session, workingDir.data, [&](const ProcessStartInfo& pi) { return session.RunProcess(pi, true, enableDetour); });
 	}
 
+	bool RunLocal(LoggerWithWriter& logger, const StringBufferBase& testRootDir, const TestSessionFunction& testFunc, bool enableDetour)
+	{
+		return RunLocal(logger, testRootDir, {}, testFunc, enableDetour);
+	}
+
+	void GetTestAppPath(LoggerWithWriter& logger, StringBufferBase& out)
+	{
+		GetDirectoryOfCurrentModule(logger, out);
+		out.EnsureEndsWithSlash();
+		out.Append(IsWindows ? TC("UbaTestApp.exe") : TC("UbaTestApp"));
+	}
+
+	bool RunTestApp(LoggerWithWriter& logger, const tchar* workingDir, const RunProcessFunction& runProcess, const tchar* arguments, ProcessHandle* outHandle = nullptr)
+	{
+		StringBuffer<> testApp;
+		GetTestAppPath(logger, testApp);
+
+		ProcessStartInfo processInfo;
+		processInfo.application = testApp.data;
+		processInfo.workingDir = workingDir;
+		processInfo.arguments = arguments;
+		processInfo.logLineFunc = [](void* userData, const tchar* line, u32 length, LogEntryType type)
+			{
+				LoggerWithWriter(g_consoleLogWriter, TC("")).Info(line);
+			};
+
+		ProcessHandle process = runProcess(processInfo);
+		if (!process.WaitForExit(GetTimeoutTime()))
+			return logger.Error(TC("UbaTestApp did not exit in 10 seconds"));
+		if (outHandle)
+			*outHandle = process;
+		u32 exitCode = process.GetExitCode();
+		if (exitCode == 0)
+			return true;
+		for (auto& logLine : process.GetLogLines())
+			logger.Error(logLine.text.c_str());
+		return logger.Error(TC("UbaTestApp returned exit code %u"), exitCode);
+	}
+
 	using TestServerClientSessionFunction = Function<bool(LoggerWithWriter& logger, const StringView& workingDir, SessionServer& sessionServer, SessionClient& sessionClient)>;
-	bool SetupServerClientSession(LoggerWithWriter& logger, const StringBufferBase& testRootDir, bool deleteAll, bool serverShouldListen, const TestServerClientSessionFunction& testFunc)
+	bool SetupServerClientSession(LoggerWithWriter& logger, const StringBufferBase& testRootDir, bool deleteAll, bool serverShouldListen, const Config& serverConfig, const TestServerClientSessionFunction& testFunc)
 	{
 		LogWriter& logWriter = logger.m_writer;
-		NetworkBackendTcp tcpBackend(logWriter);
+		NetworkBackendTcpCreateInfo tcpInfo{logWriter};
+		tcpInfo.Apply(serverConfig);
+		NetworkBackendTcp tcpBackend(tcpInfo);
 
 		bool ctorSuccess = true;
 		NetworkServer server(ctorSuccess, { logWriter });
@@ -68,6 +120,7 @@
 
 		StorageServerCreateInfo storageServerInfo(server, rootDir.data, logWriter);
 		storageServerInfo.casCapacityBytes = 1024ull * 1024 * 1024;
+		storageServerInfo.Apply(serverConfig);
 		auto& storageServer = *new StorageServer(storageServerInfo);
 		auto ssg = MakeGuard([&]() { delete &storageServer; });
 
@@ -79,6 +132,7 @@
 		sessionServerInfo.logToFile = true;
 		sessionServerInfo.remoteLogEnabled = true;
 		#endif
+		sessionServerInfo.Apply(serverConfig);
 
 		auto& sessionServer = *new SessionServer(sessionServerInfo);
 		auto ssg2 = MakeGuard([&]() { delete &sessionServer; });
@@ -129,7 +183,7 @@
 
 		SessionClientCreateInfo sessionClientInfo(storageClient, client, logWriter);
 		sessionClientInfo.rootDir = rootDir.data;
-		sessionClientInfo.allowKeepFilesInMemory = false;
+		//sessionClientInfo.allowKeepFilesInMemory = false;
 
 		#if UBA_DEBUG
 		sessionClientInfo.logToFile = true;
@@ -157,22 +211,19 @@
 		return testFunc(logger, workingDir, sessionServer, sessionClient);
 	}
 
-	bool RunRemote(LoggerWithWriter& logger, const StringBufferBase& testRootDir, const TestSessionFunction& testFunc, bool deleteAll, bool serverShouldListen)
+	bool RunRemote(LoggerWithWriter& logger, const StringBufferBase& testRootDir, const Config& serverConfig, const TestSessionFunction& testFunc, bool deleteAll = true, bool serverShouldListen = true)
 	{
-		return SetupServerClientSession(logger, testRootDir, deleteAll, serverShouldListen,
+		return SetupServerClientSession(logger, testRootDir, deleteAll, serverShouldListen, serverConfig,
 			[&](LoggerWithWriter& logger, const StringView& workingDir, SessionServer& sessionServer, SessionClient& sessionClient)
 			{
 				return testFunc(logger, sessionServer, workingDir.data, [&](const ProcessStartInfo& pi) { return sessionServer.RunProcessRemote(pi); });
 			});
 	}
 
-	void GetTestAppPath(LoggerWithWriter& logger, StringBufferBase& out)
+	bool RunRemote(LoggerWithWriter& logger, const StringBufferBase& testRootDir, const TestSessionFunction& testFunc, bool deleteAll, bool serverShouldListen)
 	{
-		GetDirectoryOfCurrentModule(logger, out);
-		out.EnsureEndsWithSlash();
-		out.Append(IsWindows ? TC("UbaTestApp.exe") : TC("UbaTestApp"));
+		return RunRemote(logger, testRootDir, {}, testFunc, deleteAll, serverShouldListen);
 	}
-
 	bool CreateTextFile(StringBufferBase& outPath, LoggerWithWriter& logger, const tchar* workingDir, const tchar* fileName, const char* text)
 	{
 		outPath.Clear().Append(workingDir).EnsureEndsWithSlash().Append(fileName);
@@ -183,11 +234,8 @@
 		return fr.Close();
 	}
 
-	bool RunTestApp(LoggerWithWriter& logger, SessionServer& session, const tchar* workingDir, const RunProcessFunction& runProcess)
+	bool RunTestAppTests(LoggerWithWriter& logger, SessionServer& session, const tchar* workingDir, const RunProcessFunction& runProcess)
 	{
-		StringBuffer<MaxPath> testApp;
-		GetTestAppPath(logger, testApp);
-
 		StringBuffer<MaxPath> fileR;
 		if (!CreateTextFile(fileR, logger, workingDir, TC("FileR.h"), "Foo"))
 			return false;
@@ -212,25 +260,8 @@
 				return logger.Error(TC("Failed to create dir %s"), dir.data);
 		}
 
-		ProcessStartInfo processInfo;
-		processInfo.application = testApp.data;
-		processInfo.workingDir = workingDir;
-		processInfo.logLineFunc = [](void* userData, const tchar* line, u32 length, LogEntryType type)
-			{
-				LoggerWithWriter(g_consoleLogWriter, TC("")).Info(line);
-			};
-
-		ProcessHandle process = runProcess(processInfo);
-		if (!process.WaitForExit(100000))
-			return logger.Error(TC("UbaTestApp did not exit in 10 seconds"));
-		u32 exitCode = process.GetExitCode();
-
-		if (exitCode != 0)
-		{
-			for (auto& logLine : process.GetLogLines())
-				logger.Error(logLine.text.c_str());
-			return logger.Error(TC("UbaTestApp returned exit code %u"), exitCode);
-		}
+		if (!RunTestApp(logger, workingDir, runProcess, TC("")))
+			return false;
 
 		{
 			StringBuffer<MaxPath> fileW2;
@@ -313,7 +344,7 @@
 		processInfo.arguments = args.data;
 
 		processInfo.workingDir = workingDir;
-		//processInfo.logFile = TC("[REDACTED_USER_HOME]/RunClang.log");
+		//processInfo.logFile = TC("/mnt/e/temp/ttt/RunClang.log");
 		ProcessHandle process = runProcess(processInfo);
 		if (!process.WaitForExit(40000))
 			return logger.Error(TC("clang++ timed out"));
@@ -336,24 +367,10 @@
 				return u32(helloBytes);
 			});
 
-		StringBuffer<> testApp;
-		GetTestAppPath(logger, testApp);
-
-		ProcessStartInfo processInfo;
-		processInfo.application = testApp.data;
-		processInfo.workingDir = workingDir;
-		processInfo.arguments = TC("Whatever");
-		ProcessHandle process = runProcess(processInfo);
-		if (!process.WaitForExit(10000))
-			return logger.Error(TC("UbaTestApp did not exit in 10 seconds"));
-		u32 exitCode = process.GetExitCode();
-
-		if (exitCode != 0)
-			return logger.Error(TC("UbaTestApp returned exit code %u"), exitCode);
-
+		if (!RunTestApp(logger, workingDir, runProcess, TC("Whatever")))
+			return false;
 		if (!gotMessage)
 			return logger.Error(TC("Never got message from UbaTestApp"));
-
 		return true;
 	}
 
@@ -371,12 +388,12 @@
 	// exercise various file functions which will actually go through our detour library.
 	bool TestDetouredTestApp(LoggerWithWriter& logger, const StringBufferBase& testRootDir)
 	{
-		return RunLocal(logger, testRootDir, RunTestApp);
+		return RunLocal(logger, testRootDir, RunTestAppTests);
 	}
 
 	bool TestRemoteDetouredTestApp(LoggerWithWriter& logger, const StringBufferBase& testRootDir)
 	{
-		return RunRemote(logger, testRootDir, RunTestApp);
+		return RunRemote(logger, testRootDir, RunTestAppTests);
 	}
 
 	bool TestCustomService(LoggerWithWriter& logger, const StringBufferBase& testRootDir)
@@ -419,7 +436,7 @@
 				processInfo.arguments = file.data;
 				processInfo.logFile = TC("[REDACTED_USER_HOME]/Touch.log");
 				ProcessHandle process = runProcess(processInfo);
-				if (!process.WaitForExit(10000))
+				if (!process.WaitForExit(GetTimeoutTime()))
 					return logger.Error(TC("UbaTestApp did not exit in 10 seconds"));
 				u32 exitCode = process.GetExitCode();
 				if (exitCode != 0)
@@ -451,7 +468,7 @@
 					};
 
 				ProcessHandle process = runProcess(processInfo);
-				if (!process.WaitForExit(100000))
+				if (!process.WaitForExit(GetTimeoutTime()))
 					return logger.Error(TC("UbaTestApp did not exit in 10 seconds"));
 				u32 exitCode = process.GetExitCode();
 
@@ -513,7 +530,7 @@
 
 				for (auto& process : processes)
 				{
-					if (!process.WaitForExit(10000))
+					if (!process.WaitForExit(GetTimeoutTime()))
 						return logger.Error(TC("UbaTestApp did not exit in 10 seconds"));
 					u32 exitCode = process.GetExitCode();
 					if (exitCode != 0)
@@ -540,7 +557,7 @@
 
 		ProcessHandle process = runProcess(processInfo);
 
-		if (!process.WaitForExit(10000))
+		if (!process.WaitForExit(GetTimeoutTime()))
 			return logger.Error(TC("UbaTestApp did not exit in 10 seconds"));
 		u32 exitCode = process.GetExitCode();
 		if (exitCode != 0)
@@ -578,7 +595,7 @@
 				arg.Append(file);
 				processInfo.arguments = arg.data;
 				ProcessHandle process = runProcess(processInfo);
-				if (!process.WaitForExit(100000))
+				if (!process.WaitForExit(GetTimeoutTime()))
 					return logger.Error(TC("UbaTestApp did not exit in 10 seconds"));
 				u32 exitCode =  process.GetExitCode();
 				return exitCode == 255 ? INVALID_FILE_ATTRIBUTES : exitCode;
@@ -677,7 +694,7 @@
 	bool TestRemoteDirectoryTable(LoggerWithWriter& logger, const StringBufferBase& testRootDir)
 	{
 #if 0
-		return SetupServerClientSession(logger, testRootDir, true, true,
+		return SetupServerClientSession(logger, testRootDir, true, true, {},
 			[&](LoggerWithWriter& logger, const StringView& workingDir, SessionServer& sessionServer, SessionClient& sessionClient)
 			{
 				u32 attributes;
@@ -697,5 +714,4 @@
 		return true;
 #endif
 	}
-
 }
diff --git a/Engine/Source/Programs/UnrealBuildAccelerator/TestApp/Private/PosixOS/UbaTestAppPosix.cpp b/Engine/Source/Programs/UnrealBuildAccelerator/TestApp/Private/PosixOS/UbaTestAppPosix.cpp
--- a/Engine/Source/Programs/UnrealBuildAccelerator/TestApp/Private/PosixOS/UbaTestAppPosix.cpp
+++ b/Engine/Source/Programs/UnrealBuildAccelerator/TestApp/Private/PosixOS/UbaTestAppPosix.cpp
@@ -236,7 +236,7 @@ int main(int argc, char* argv[])
 		struct stat attr;
 		if (stat(str, &attr) == -1)
 			return 255;
-		return attr.st_mode;
+		return 1;
 	}
 	else if (strcmp(argv[1], "-popen") == 0)
 	{
