[UBA]
* Fixed xcode-select detouring logic.. echo process did not echo the string it should have echoed because the posix_spawn call was missing the cmd itself in first arg.
* Fixed bug in readdir where d_namlen was not set on osx.
* Added unit tests for running detoured xcode-select

diff --git a/Engine/Source/Programs/UnrealBuildAccelerator/Detours/Private/PosixOS/UbaDetoursFunctionsPosix.cpp b/Engine/Source/Programs/UnrealBuildAccelerator/Detours/Private/PosixOS/UbaDetoursFunctionsPosix.cpp
--- a/Engine/Source/Programs/UnrealBuildAccelerator/Detours/Private/PosixOS/UbaDetoursFunctionsPosix.cpp
+++ b/Engine/Source/Programs/UnrealBuildAccelerator/Detours/Private/PosixOS/UbaDetoursFunctionsPosix.cpp
@@ -1363,7 +1363,7 @@ UBA_EXPORT dirent* UBA_WRAPPER(readdir)(DIR* dirp)
 			++dirInfo.it;
 			if (dirInfo.it >= dirInfo.fileTableOffsets.size())
 			{
-				DEBUG_LOG_DETOURED("readdir", "(%p) -> nullptr", dirp);
+				DEBUG_LOG_DETOURED("readdir", "(%p) -> nullptr (found %u entries)", dirp, dirInfo.fileTableOffsets.size());
 				return nullptr;
 			}
 			u32 fileTableOffset = dirInfo.fileTableOffsets[dirInfo.it];
@@ -1373,17 +1373,22 @@ UBA_EXPORT dirent* UBA_WRAPPER(readdir)(DIR* dirp)
 			if (info.attributes == 0) // File was deleted
 				continue;
 
+			auto nameLen = u16(strlen(dirInfo.ent.d_name));
+			bool isDir = S_ISDIR(info.attributes);
+
 			dirInfo.ent.d_ino = info.fileIndex;
+			dirInfo.ent.d_reclen = offsetof(struct dirent, d_name) + nameLen + 1;
+			dirInfo.ent.d_type = isDir ? DT_DIR : DT_REG;
+
 			#if PLATFORM_LINUX
 			dirInfo.ent.d_off = 0;
+			#else
+			dirInfo.ent.d_namlen = nameLen;
 			#endif
-			dirInfo.ent.d_reclen = sizeof(dirent);
-			dirInfo.ent.d_type = S_ISDIR(info.attributes) ? DT_DIR : DT_REG;
+
+			//DEBUG_LOG_DETOURED("readdir", "(%p) -> %p (%s)%s", dirp, &dirInfo.ent, dirInfo.ent.d_name, (isDir?" (Dir)":""));
 			return &dirInfo.ent;
 		}
-
-		DEBUG_LOG_DETOURED("readdir", "(%p) -> %p", dirp, &dirInfo.ent);
-		return &dirInfo.ent;
 	}
 
 	auto res = TRUE_WRAPPER(readdir)(dirp);
@@ -1957,7 +1962,7 @@ int SpawnEcho(char* str, pid_t* pid, const posix_spawn_file_actions_t* file_acti
 #else
 	char* const env[] = { nullptr };
 	char cmd[] = "/bin/echo"; 
-	char* const argv2[] = { str, nullptr };
+	char* const argv2[] = { cmd, str, nullptr };
 	int res = TRUE_WRAPPER(posix_spawn)(pid, cmd, file_actions, attrp, argv2, env);
 	DEBUG_LOG_TRUE("posix_spawn", "(ECHO) %s (pid: %u) -> %i", str, *pid, res);
 #endif
diff --git a/Engine/Source/Programs/UnrealBuildAccelerator/Test/Private/UbaTestAll.h b/Engine/Source/Programs/UnrealBuildAccelerator/Test/Private/UbaTestAll.h
--- a/Engine/Source/Programs/UnrealBuildAccelerator/Test/Private/UbaTestAll.h
+++ b/Engine/Source/Programs/UnrealBuildAccelerator/Test/Private/UbaTestAll.h
@@ -84,7 +84,9 @@ namespace uba
 
 
 #define UBA_MAC_TESTS \
-		UBA_POSIX_TESTS
+		UBA_POSIX_TESTS \
+		UBA_TEST(TestXcodeSelect) \
+		UBA_TEST(TestRemoteXcodeSelect) \
 
 
 
diff --git a/Engine/Source/Programs/UnrealBuildAccelerator/Test/Private/UbaTestSession.cpp b/Engine/Source/Programs/UnrealBuildAccelerator/Test/Private/UbaTestSession.cpp
--- a/Engine/Source/Programs/UnrealBuildAccelerator/Test/Private/UbaTestSession.cpp
+++ b/Engine/Source/Programs/UnrealBuildAccelerator/Test/Private/UbaTestSession.cpp
@@ -714,4 +714,19 @@ namespace uba
 		return true;
 #endif
 	}
+
+	#if PLATFORM_MAC
+	bool RunXCodeSelect(LoggerWithWriter& logger, SessionServer& session, const tchar* workingDir, const RunProcessFunction& runProcess)
+	{
+		return RunTestApp(logger, workingDir, runProcess, TC("-xcode-select"));
+	}
+	bool TestXcodeSelect(LoggerWithWriter& logger, const StringBufferBase& testRootDir)
+	{
+		return RunLocal(logger, testRootDir, RunXCodeSelect);
+	}
+	bool TestRemoteXcodeSelect(LoggerWithWriter& logger, const StringBufferBase& testRootDir)
+	{
+		return RunLocal(logger, testRootDir, RunXCodeSelect);
+	}
+	#endif
 }
diff --git a/Engine/Source/Programs/UnrealBuildAccelerator/TestApp/Private/PosixOS/UbaTestAppPosix.cpp b/Engine/Source/Programs/UnrealBuildAccelerator/TestApp/Private/PosixOS/UbaTestAppPosix.cpp
--- a/Engine/Source/Programs/UnrealBuildAccelerator/TestApp/Private/PosixOS/UbaTestAppPosix.cpp
+++ b/Engine/Source/Programs/UnrealBuildAccelerator/TestApp/Private/PosixOS/UbaTestAppPosix.cpp
@@ -300,5 +300,59 @@ int main(int argc, char* argv[])
 
 		return 0;
 	}
+	#if PLATFORM_MAC
+	else if (strcmp(argv[1], "-xcode-select") == 0)
+	{
+		int pipefd[2];
+		if (pipe(pipefd) != 0)
+			return LogError("Failed to create pipe");
+
+		posix_spawn_file_actions_t actions;
+		posix_spawn_file_actions_init(&actions);
+
+		// Close read end in child, redirect stdout to write end
+		posix_spawn_file_actions_addclose(&actions, pipefd[0]);
+		posix_spawn_file_actions_adddup2(&actions, pipefd[1], STDOUT_FILENO);
+		posix_spawn_file_actions_addclose(&actions, pipefd[1]);
+
+		pid_t childPid = 0;
+		const char* execPath = "/usr/bin/xcode-select";
+		const char* args[] { execPath, "--print-path", nullptr };
+		if (posix_spawn(&childPid, execPath, &actions, nullptr, (char**)args, environ) != 0)
+			return LogError("posix_spawn failed");
+		
+		posix_spawn_file_actions_destroy(&actions);
+		close(pipefd[1]);  // Close write end in parent
+
+		bool success = false;
+
+		char buffer[256];
+		int count;
+		while ((count = read(pipefd[0], buffer, sizeof(buffer)-1)) > 0)
+		{
+			buffer[count] = '\0';
+			if (strstr(buffer, "Xcode.app"))
+				success = true;
+		}
+
+		int status = 0;
+		do
+		{
+			if (waitpid(childPid, &status, WUNTRACED | WCONTINUED) == -1)
+				return LogError("waitpid on child process failed (pid %i)", childPid);
+			if (WIFSIGNALED(status))
+				return LogError("Child process killed by signal %d", WTERMSIG(status));
+			if (WIFSTOPPED(status))
+				return LogError("Child process stopped by signal %d", WSTOPSIG(status));
+			if (WIFCONTINUED(status))
+				return LogError("Child process continued");
+		} while (!WIFEXITED(status));
+		if (WEXITSTATUS(status) != 0)
+			return LogError("Child process failed");
+
+		return success ? 0 : -1;
+	}
+	#endif
+
 	return -2;
 }
