[UBA]
* Fixed so mac bin dependency parser adds executable path to loader paths if set

diff --git a/Engine/Source/Programs/UnrealBuildAccelerator/Core/Public/UbaMacBinDependencyParser.h b/Engine/Source/Programs/UnrealBuildAccelerator/Core/Public/UbaMacBinDependencyParser.h
--- a/Engine/Source/Programs/UnrealBuildAccelerator/Core/Public/UbaMacBinDependencyParser.h
+++ b/Engine/Source/Programs/UnrealBuildAccelerator/Core/Public/UbaMacBinDependencyParser.h
@@ -110,7 +110,8 @@ namespace uba
 				else if (name.StartsWith("executable_path/"))
 				{
 					auto executablePath = name.Skip(16).data;
-					UBA_ASSERT(*executablePath == 0);
+					if (*executablePath != 0)
+						loaderPaths[loaderPathsCount++] = executablePath;
 				}
 				else if (name.StartsWith("loader_path/"))
 				{
