[UBA]
* Fixed osx unit tests for machines not having xcode-select

diff --git a/Engine/Source/Programs/UnrealBuildAccelerator/Test/Private/UbaTestSession.cpp b/Engine/Source/Programs/UnrealBuildAccelerator/Test/Private/UbaTestSession.cpp
--- a/Engine/Source/Programs/UnrealBuildAccelerator/Test/Private/UbaTestSession.cpp
+++ b/Engine/Source/Programs/UnrealBuildAccelerator/Test/Private/UbaTestSession.cpp
@@ -284,7 +284,7 @@ namespace uba
 		FILE* fpCommand = popen(command, "r");
 		if (fpCommand == nullptr || fgets(commandOutput.data, commandOutput.capacity, fpCommand) == nullptr || pclose(fpCommand) != 0)
 		{
-			logger.Warning("Failed to get an Xcode from xcode-select");
+			logger.Warning("Failed to run '%s' or get a response");
 			return false;
 		}
 
@@ -718,6 +718,9 @@ namespace uba
 	#if PLATFORM_MAC
 	bool RunXCodeSelect(LoggerWithWriter& logger, SessionServer& session, const tchar* workingDir, const RunProcessFunction& runProcess)
 	{
+		StringBuffer<> xcodeSelect;
+		if (!ExecuteCommand(logger, "which xcode-select", xcodeSelect))
+			return true;
 		return RunTestApp(logger, workingDir, runProcess, TC("-xcode-select"));
 	}
 	bool TestXcodeSelect(LoggerWithWriter& logger, const StringBufferBase& testRootDir)
@@ -726,7 +729,7 @@ namespace uba
 	}
 	bool TestRemoteXcodeSelect(LoggerWithWriter& logger, const StringBufferBase& testRootDir)
 	{
-		return RunLocal(logger, testRootDir, RunXCodeSelect);
+		return RunRemote(logger, testRootDir, RunXCodeSelect);
 	}
 	#endif
 }
