[UBA]
* Added system error callback that can be used deep down in callstacks to report a system error like "no space left on device" etc
* Added so UbaAgent registers a callback for system errors. When error occurs it is reported to the host and then the helper exits

diff --git a/Engine/Source/Programs/UnrealBuildAccelerator/Agent/Private/UbaAgent.cpp b/Engine/Source/Programs/UnrealBuildAccelerator/Agent/Private/UbaAgent.cpp
--- a/Engine/Source/Programs/UnrealBuildAccelerator/Agent/Private/UbaAgent.cpp
+++ b/Engine/Source/Programs/UnrealBuildAccelerator/Agent/Private/UbaAgent.cpp
@@ -1319,8 +1319,13 @@ namespace uba
 
 			g_sessionClient = sessionClient;
 
+			u32 systemError = 0;
+			SetSystemErrorCallback([&](u32 error) { systemError = error; wakeupSessionWait.Set(); });
+
 			auto disconnectAndStopLoggingThread = MakeGuard([&]()
 			{
+				SetSystemErrorCallback({});
+
 				g_exitLock->Enter();
 				g_sessionClient = nullptr;
 				g_exitLock->Leave();
@@ -1411,6 +1416,12 @@ namespace uba
 					}
 				}
 
+				if (systemError)
+				{
+					isTerminating = true;
+					sessionClient->SetIsTerminating(TStrdup(LastErrorToText(systemError).data), 0); // Leak allocation
+				}
+
 				if (sessionClient->Wait(sessionWaitTimeout, &wakeupSessionWait))
 				{
 					// We got version mismatch and have the cas keys for the needed Agent/Detours binaries
diff --git a/Engine/Source/Programs/UnrealBuildAccelerator/Common/Private/UbaFile.cpp b/Engine/Source/Programs/UnrealBuildAccelerator/Common/Private/UbaFile.cpp
--- a/Engine/Source/Programs/UnrealBuildAccelerator/Common/Private/UbaFile.cpp
+++ b/Engine/Source/Programs/UnrealBuildAccelerator/Common/Private/UbaFile.cpp
@@ -1032,6 +1032,19 @@ namespace uba
 		return true;
 	};
 
+	Function<void(u32)> g_systemErrorCallback;
+
+	void ReportSystemError(u32 error)
+	{
+		if (g_systemErrorCallback)
+			g_systemErrorCallback(error);
+	}
+
+	void SetSystemErrorCallback(const Function<void(u32)>& callback)
+	{
+		g_systemErrorCallback = callback;
+	}
+
 	VolumeCache::~VolumeCache()
 	{
 		#if PLATFORM_WINDOWS
diff --git a/Engine/Source/Programs/UnrealBuildAccelerator/Common/Private/UbaFileAccessor.cpp b/Engine/Source/Programs/UnrealBuildAccelerator/Common/Private/UbaFileAccessor.cpp
--- a/Engine/Source/Programs/UnrealBuildAccelerator/Common/Private/UbaFileAccessor.cpp
+++ b/Engine/Source/Programs/UnrealBuildAccelerator/Common/Private/UbaFileAccessor.cpp
@@ -501,6 +501,9 @@ namespace uba
 			{
 				if (rename(realFileName, m_fileName) == -1)
 				{
+					if (errno == ENOSPC)
+						ReportSystemError(errno);
+
 					if (errno != EXDEV)
 						return m_logger.Error(TC("Failed to rename temporary file %s to %s (%s)"), realFileName, m_fileName, strerror(errno));
 
diff --git a/Engine/Source/Programs/UnrealBuildAccelerator/Common/Public/UbaFile.h b/Engine/Source/Programs/UnrealBuildAccelerator/Common/Public/UbaFile.h
--- a/Engine/Source/Programs/UnrealBuildAccelerator/Common/Public/UbaFile.h
+++ b/Engine/Source/Programs/UnrealBuildAccelerator/Common/Public/UbaFile.h
@@ -91,6 +91,8 @@ namespace uba
 	bool GetCurrentDirectoryW(StringBufferBase& out);
 	bool GetAlternativeUbaPath(Logger& logger, StringBufferBase& out, StringView firstPath, bool isWindowsArm);
 
+	void ReportSystemError(u32 error);
+	void SetSystemErrorCallback(const Function<void(u32)>& callback);
 
 	class DirectoryCache
 	{
