[UBA]
* Added code to handle EAGAIN error in TraverseDir for posix

diff --git a/Engine/Source/Programs/UnrealBuildAccelerator/Common/Public/UbaDirectoryIterator.h b/Engine/Source/Programs/UnrealBuildAccelerator/Common/Public/UbaDirectoryIterator.h
--- a/Engine/Source/Programs/UnrealBuildAccelerator/Common/Public/UbaDirectoryIterator.h
+++ b/Engine/Source/Programs/UnrealBuildAccelerator/Common/Public/UbaDirectoryIterator.h
@@ -270,20 +270,31 @@ namespace uba
 			strcpy(dirPath2 + dirPathLen, fileName);
 
 			struct stat attr;
-			int res = stat(dirPath2, &attr);
-			if (res == -1)
+			bool skipEntry = false;
+
+			while (true)
 			{
-				// In the case of symlinks we may need to try lstat
-				// just in case.
-				if (errno == ENOENT) {
+				int res = stat(dirPath2, &attr);
+				if (res == 0)
+					break;
+
+				// In the case of symlinks we may need to try lstat just in case.
+				if (errno == ENOENT)
 					res = lstat(dirPath2, &attr);
-				}
+				if (res == 0)
+					break;
 
-				if (errno == ENODEV || errno == EACCES || errno == EPERM)
+				if (errno == EAGAIN)
 					continue;
-				UBA_ASSERTF(res == 0, "TraverseDir:stat error handling not added for file %s (error %s)", dirPath2, strerror(errno));
+
+				skipEntry = true;
+				UBA_ASSERTF(errno != ENODEV && errno != EACCES && errno != EPERM, "TraverseDir:stat error handling not added for file %s (error %s)", dirPath2, strerror(errno));
+				break;
 			}
 
+			if (skipEntry)
+				continue;
+
 			DirectoryEntry entry;
 			entry.name = fileName;
 			entry.nameLen = strlen(fileName);
