[UBA]
* Fixed inverted assert in TraverseDir

diff --git a/Engine/Source/Programs/UnrealBuildAccelerator/Common/Public/UbaDirectoryIterator.h b/Engine/Source/Programs/UnrealBuildAccelerator/Common/Public/UbaDirectoryIterator.h
--- a/Engine/Source/Programs/UnrealBuildAccelerator/Common/Public/UbaDirectoryIterator.h
+++ b/Engine/Source/Programs/UnrealBuildAccelerator/Common/Public/UbaDirectoryIterator.h
@@ -288,7 +288,7 @@ namespace uba
 					continue;
 
 				skipEntry = true;
-				UBA_ASSERTF(errno != ENODEV && errno != EACCES && errno != EPERM, "TraverseDir:stat error handling not added for file %s (error %s)", dirPath2, strerror(errno));
+				UBA_ASSERTF(errno == ENODEV || errno == EACCES || errno == EPERM, "TraverseDir:stat error handling not added for file %s (error %s)", dirPath2, strerror(errno));
 				break;
 			}
 
