[UBA]
* Bumped trace version to be able to track more kernel call types
* Added timers around kernel call rename and rename fallback
* Fixed bug in visualizer where if messagebox pops up because of trace stream parsing issues everything hangs because message pump reenters WM_TIMER

diff --git a/Engine/Source/Programs/UnrealBuildAccelerator/Common/Private/UbaFileAccessor.cpp b/Engine/Source/Programs/UnrealBuildAccelerator/Common/Private/UbaFileAccessor.cpp
--- a/Engine/Source/Programs/UnrealBuildAccelerator/Common/Private/UbaFileAccessor.cpp
+++ b/Engine/Source/Programs/UnrealBuildAccelerator/Common/Private/UbaFileAccessor.cpp
@@ -499,7 +499,12 @@ namespace uba
 			#else
 			if (m_tempPath)
 			{
-				if (rename(realFileName, m_fileName) == -1)
+				int res;
+				{
+					ExtendedTimerScope ts(KernelStats::GetCurrent().renameFile);
+					res = rename(realFileName, m_fileName);
+				}
+				if (res == -1)
 				{
 					if (errno == ENOSPC)
 						ReportSystemError(errno);
@@ -507,6 +512,8 @@ namespace uba
 					if (errno != EXDEV)
 						return m_logger.Error(TC("Failed to rename temporary file %s to %s (%s)"), realFileName, m_fileName, strerror(errno));
 
+					ExtendedTimerScope ts(KernelStats::GetCurrent().renameFileFallback);
+
 					// Need to copy, can't rename over devices
 					int targetFd = open(m_fileName, O_CREAT | O_TRUNC | O_WRONLY | O_CLOEXEC, m_flagsAndAttributes);
 					auto closeTargetFile = MakeGuard([&]() { close(targetFd); });
diff --git a/Engine/Source/Programs/UnrealBuildAccelerator/Common/Private/UbaStats.cpp b/Engine/Source/Programs/UnrealBuildAccelerator/Common/Private/UbaStats.cpp
--- a/Engine/Source/Programs/UnrealBuildAccelerator/Common/Private/UbaStats.cpp
+++ b/Engine/Source/Programs/UnrealBuildAccelerator/Common/Private/UbaStats.cpp
@@ -206,7 +206,12 @@ namespace uba
 			return;
 		}
 
-		u16 bits = reader.ReadU16();
+		u64 bits;
+		if (version < 43)
+			bits = reader.ReadU16();
+		else
+			bits = reader.Read7BitEncoded();
+
 		#define UBA_KERNEL_STAT(type, var, ver) if (bits & (1 << Bit_##var)) uba::Read(reader, var, version);
 		UBA_KERNEL_STATS
 		#undef UBA_KERNEL_STAT
diff --git a/Engine/Source/Programs/UnrealBuildAccelerator/Common/Public/UbaTrace.h b/Engine/Source/Programs/UnrealBuildAccelerator/Common/Public/UbaTrace.h
--- a/Engine/Source/Programs/UnrealBuildAccelerator/Common/Public/UbaTrace.h
+++ b/Engine/Source/Programs/UnrealBuildAccelerator/Common/Public/UbaTrace.h
@@ -69,7 +69,7 @@ namespace uba
 		#undef UBA_TRACE_TYPE
 	};
 
-	static constexpr u32 TraceVersion = 42;
+	static constexpr u32 TraceVersion = 43;
 	static constexpr u32 TraceReadCompatibilityVersion = 6;
 
 	class Trace : public WorkTracker
diff --git a/Engine/Source/Programs/UnrealBuildAccelerator/Core/Public/UbaProcessStats.h b/Engine/Source/Programs/UnrealBuildAccelerator/Core/Public/UbaProcessStats.h
--- a/Engine/Source/Programs/UnrealBuildAccelerator/Core/Public/UbaProcessStats.h
+++ b/Engine/Source/Programs/UnrealBuildAccelerator/Core/Public/UbaProcessStats.h
@@ -134,6 +134,11 @@ namespace uba
 			UBA_KERNEL_STAT(ExtendedTimer, traverseDir, 27) \
 			UBA_KERNEL_STAT(ExtendedTimer, virtualAlloc, 30) \
 			UBA_KERNEL_STAT(TimeAndBytes, memoryCompress, 41) \
+			UBA_KERNEL_STAT(ExtendedTimer, renameFile, 43) \
+			UBA_KERNEL_STAT(ExtendedTimer, renameFileFallback, 43) \
+			UBA_KERNEL_STAT(ExtendedTimer, copyFile, 43) \
+			UBA_KERNEL_STAT(ExtendedTimer, moveFile, 43) \
+			UBA_KERNEL_STAT(ExtendedTimer, deleteFile, 43) \
 
 		#define UBA_KERNEL_STAT(type, var, ver) type var;
 		UBA_KERNEL_STATS
@@ -155,12 +160,12 @@ namespace uba
 
 		void Write(BinaryWriter& writer)
 		{
-			u16 bits = 0;
+			u64 bits = 0;
 			#define UBA_KERNEL_STAT(type, var, ver) if (var.count) bits |= (1 << Bit_##var);
 			UBA_KERNEL_STATS
 			#undef UBA_KERNEL_STAT
 
-			writer.WriteU16(bits);
+			writer.Write7BitEncoded(bits);
 
 			#define UBA_KERNEL_STAT(type, var, ver) if (var.count) uba::Write(writer, var);
 			UBA_KERNEL_STATS
diff --git a/Engine/Source/Programs/UnrealBuildAccelerator/Visualizer/Private/UbaVisualizer.cpp b/Engine/Source/Programs/UnrealBuildAccelerator/Visualizer/Private/UbaVisualizer.cpp
--- a/Engine/Source/Programs/UnrealBuildAccelerator/Visualizer/Private/UbaVisualizer.cpp
+++ b/Engine/Source/Programs/UnrealBuildAccelerator/Visualizer/Private/UbaVisualizer.cpp
@@ -3662,6 +3662,10 @@ namespace uba
 		}
 		case WM_TIMER:
 		{
+			if (m_handlingTimer) // This is to prevent hangs caused by message boxes from asserts
+				break;
+			m_handlingTimer = true;
+
 			bool changed = false;
 			if (!m_paused)
 			{
@@ -3702,6 +3706,8 @@ namespace uba
 				if (!m_traceView.finished)
 					SetTimer(m_hwnd, 0, waitTime, NULL);
 			}
+
+			m_handlingTimer = false;
 			break;
 		}
 		case WM_MOUSEWHEEL:
diff --git a/Engine/Source/Programs/UnrealBuildAccelerator/Visualizer/Private/UbaVisualizer.h b/Engine/Source/Programs/UnrealBuildAccelerator/Visualizer/Private/UbaVisualizer.h
--- a/Engine/Source/Programs/UnrealBuildAccelerator/Visualizer/Private/UbaVisualizer.h
+++ b/Engine/Source/Programs/UnrealBuildAccelerator/Visualizer/Private/UbaVisualizer.h
@@ -285,6 +285,7 @@ namespace uba
 		bool m_mouseOverWindow = false;
 		bool m_showPopup = false;
 		bool m_locked = false;
+		bool m_handlingTimer = false;
 
 		HBITMAP m_cachedBitmap = 0;
 		RECT m_cachedBitmapRect = { INT_MIN, INT_MIN, INT_MIN, INT_MIN };
