[UBA]
* Added detouring of readdir64 on linux
* Moved process id from payload to init message so it can be tracked on posix as well
* Centralized detoured log header and added so posix writes out the same things as windows
* NetworkBackendTcp::DeleteConnection set isready event before destroying connection to prevent hang
* Reduced lock in NetworkServer for when adding connections
* Made NetworkServer::RemoveDisconnectedConnections public so it can be used in unit tests to flush out client disconnects before backend goes out of scope

diff --git a/Engine/Source/Programs/UnrealBuildAccelerator/Common/Private/UbaNetworkBackendTcp.cpp b/Engine/Source/Programs/UnrealBuildAccelerator/Common/Private/UbaNetworkBackendTcp.cpp
--- a/Engine/Source/Programs/UnrealBuildAccelerator/Common/Private/UbaNetworkBackendTcp.cpp
+++ b/Engine/Source/Programs/UnrealBuildAccelerator/Common/Private/UbaNetworkBackendTcp.cpp
@@ -831,6 +831,7 @@ namespace uba
 			Connection& c = *it;
 			if (&c != connection)
 				continue;
+			c.ready.Set();
 			it = m_connections.erase(it);
 			break;
 		}
diff --git a/Engine/Source/Programs/UnrealBuildAccelerator/Common/Private/UbaNetworkServer.cpp b/Engine/Source/Programs/UnrealBuildAccelerator/Common/Private/UbaNetworkServer.cpp
--- a/Engine/Source/Programs/UnrealBuildAccelerator/Common/Private/UbaNetworkServer.cpp
+++ b/Engine/Source/Programs/UnrealBuildAccelerator/Common/Private/UbaNetworkServer.cpp
@@ -842,7 +842,8 @@ namespace uba
 				done.Set();
 				return 0;
 			});
-
+		
+		lock.Leave();
 		done.IsSet();
 		return success;
 	}
diff --git a/Engine/Source/Programs/UnrealBuildAccelerator/Common/Private/UbaProcess.cpp b/Engine/Source/Programs/UnrealBuildAccelerator/Common/Private/UbaProcess.cpp
--- a/Engine/Source/Programs/UnrealBuildAccelerator/Common/Private/UbaProcess.cpp
+++ b/Engine/Source/Programs/UnrealBuildAccelerator/Common/Private/UbaProcess.cpp
@@ -762,6 +762,7 @@ namespace uba
 		InitMessage msg { *this };
 		InitResponse response;
 		m_messageSuccess = m_session.GetInitResponse(response, msg) && m_messageSuccess;
+		writer.WriteU32(m_id);
 		writer.WriteBool(m_echoOn);
 		writer.WriteBool(m_parentProcess != nullptr);
 		writer.WriteString(m_startInfo.application);
@@ -1663,7 +1664,6 @@ namespace uba
 			}
 
 			DetoursPayload payload;
-			payload.processId = m_id;
 			payload.hostProcess = hostProcess;
 			payload.cancelEvent = m_cancelEvent.GetHandle();
 			payload.writeEvent = m_writeEvent.GetHandle();
diff --git a/Engine/Source/Programs/UnrealBuildAccelerator/Core/Public/UbaDetoursPayload.h b/Engine/Source/Programs/UnrealBuildAccelerator/Core/Public/UbaDetoursPayload.h
--- a/Engine/Source/Programs/UnrealBuildAccelerator/Core/Public/UbaDetoursPayload.h
+++ b/Engine/Source/Programs/UnrealBuildAccelerator/Core/Public/UbaDetoursPayload.h
@@ -5,11 +5,12 @@
 #include "UbaStringBuffer.h"
 #include "UbaPlatform.h"
 
+#if PLATFORM_WINDOWS
+
 namespace uba
 {
 	struct __declspec(uuid("{a1041c70-bf7c-4e94-8f8b-bacf0f31ba9a}")) DetoursPayload
 	{
-		u32 processId = 0;
 		HANDLE hostProcess = 0;
 		HANDLE cancelEvent = 0;
 		HANDLE writeEvent = 0;
@@ -34,3 +35,5 @@ namespace uba
 
 	static const _GUID DetoursPayloadGuid = __uuidof(DetoursPayload);
 }
+
+#endif
\ No newline at end of file
diff --git a/Engine/Source/Programs/UnrealBuildAccelerator/Detours/Private/PosixOS/UbaDetoursFunctionsPosix.cpp b/Engine/Source/Programs/UnrealBuildAccelerator/Detours/Private/PosixOS/UbaDetoursFunctionsPosix.cpp
--- a/Engine/Source/Programs/UnrealBuildAccelerator/Detours/Private/PosixOS/UbaDetoursFunctionsPosix.cpp
+++ b/Engine/Source/Programs/UnrealBuildAccelerator/Detours/Private/PosixOS/UbaDetoursFunctionsPosix.cpp
@@ -36,6 +36,7 @@ namespace uba
 	bool g_isDetouring;
 	bool g_isInitialized;
 	bool g_isCancelled;
+	u32 g_processId;
 	pid_t g_pid;
 	void Deinit();
 }
@@ -143,6 +144,7 @@ using namespace uba;
 	DETOURED_FUNCTION(__fxstatat) \
 	DETOURED_FUNCTION(__fxstatat64) \
 	DETOURED_FUNCTION(open64) \
+	DETOURED_FUNCTION(readdir64) \
 	DETOURED_FUNCTION(fstatat64) \
 	DETOURED_FUNCTION(fpathconf) \
 	DETOURED_FUNCTION(pathconf) \
@@ -286,7 +288,6 @@ namespace uba
 	};
 
 	using FileHandles = UnorderedMap<int, DetouredHandle>;
-	VARIABLE_MEM(StringBuffer<MaxPath>, g_exeDir);
 	VARIABLE_MEM(FileHandles, g_fileHandles);
 	VARIABLE_MEM(ReaderWriterLock, g_fileHandlesLock);
 
@@ -1351,51 +1352,64 @@ UBA_EXPORT int UBA_WRAPPER(dirfd)(DIR* dirp)
 	return res;
 }
 
-UBA_EXPORT dirent* UBA_WRAPPER(readdir)(DIR* dirp)
+dirent* Shared_readdir(const char* func, DIR* dirp)
 {
-	UBA_INIT_DETOUR(readdir, dirp);
-
-	if (IsDirInfo(dirp))
+	auto& dirInfo = *AsDirInfo(dirp);
+	while (true)
 	{
-		auto& dirInfo = *AsDirInfo(dirp);
-		while (true)
+		++dirInfo.it;
+		if (dirInfo.it >= dirInfo.fileTableOffsets.size())
 		{
-			++dirInfo.it;
-			if (dirInfo.it >= dirInfo.fileTableOffsets.size())
-			{
-				DEBUG_LOG_DETOURED("readdir", "(%p) -> nullptr (found %u entries)", dirp, dirInfo.fileTableOffsets.size());
-				return nullptr;
-			}
-			u32 fileTableOffset = dirInfo.fileTableOffsets[dirInfo.it];
+			DEBUG_LOG_DETOURED(func, "(%p) -> nullptr (found %u entries)", dirp, dirInfo.fileTableOffsets.size());
+			return nullptr;
+		}
+		u32 fileTableOffset = dirInfo.fileTableOffsets[dirInfo.it];
 
-			DirectoryTable::EntryInformation info;
-			g_directoryTable.GetEntryInformation(info, fileTableOffset, dirInfo.ent.d_name, 256);
-			if (info.attributes == 0) // File was deleted
-				continue;
+		DirectoryTable::EntryInformation info;
+		g_directoryTable.GetEntryInformation(info, fileTableOffset, dirInfo.ent.d_name, 256);
+		if (info.attributes == 0) // File was deleted
+			continue;
 
-			auto nameLen = u16(strlen(dirInfo.ent.d_name));
-			bool isDir = S_ISDIR(info.attributes);
+		auto nameLen = u16(strlen(dirInfo.ent.d_name));
+		bool isDir = S_ISDIR(info.attributes);
 
-			dirInfo.ent.d_ino = info.fileIndex;
-			dirInfo.ent.d_reclen = offsetof(struct dirent, d_name) + nameLen + 1;
-			dirInfo.ent.d_type = isDir ? DT_DIR : DT_REG;
+		dirInfo.ent.d_ino = info.fileIndex;
+		dirInfo.ent.d_reclen = offsetof(struct dirent, d_name) + nameLen + 1;
+		dirInfo.ent.d_type = isDir ? DT_DIR : DT_REG;
 
-			#if PLATFORM_LINUX
-			dirInfo.ent.d_off = 0;
-			#else
-			dirInfo.ent.d_namlen = nameLen;
-			#endif
+		#if PLATFORM_LINUX
+		dirInfo.ent.d_off = 0;
+		#else
+		dirInfo.ent.d_namlen = nameLen;
+		#endif
 
-			//DEBUG_LOG_DETOURED("readdir", "(%p) -> %p (%s)%s", dirp, &dirInfo.ent, dirInfo.ent.d_name, (isDir?" (Dir)":""));
-			return &dirInfo.ent;
-		}
+		//DEBUG_LOG_DETOURED(func, "(%p) -> %p (%s)%s", dirp, &dirInfo.ent, dirInfo.ent.d_name, (isDir?" (Dir)":""));
+		return &dirInfo.ent;
 	}
+}
 
+UBA_EXPORT dirent* UBA_WRAPPER(readdir)(DIR* dirp)
+{
+	UBA_INIT_DETOUR(readdir, dirp);
+	if (IsDirInfo(dirp))
+		return Shared_readdir("readdir", dirp);
 	auto res = TRUE_WRAPPER(readdir)(dirp);
 	DEBUG_LOG_TRUE("readdir", "(%p) -> %p", dirp, res);
 	return res;
 }
 
+#if PLATFORM_LINUX
+UBA_EXPORT dirent64* UBA_WRAPPER(readdir64)(DIR* dirp)
+{
+	UBA_INIT_DETOUR(readdir64, dirp);
+	if (IsDirInfo(dirp))
+		return (dirent64*)Shared_readdir("readdir64", dirp);
+	auto res = TRUE_WRAPPER(readdir64)(dirp);
+	DEBUG_LOG_TRUE("readdir64", "(%p) -> %p", dirp, res);
+	return res;
+}
+#endif
+
 UBA_EXPORT void UBA_WRAPPER(rewinddir)(DIR* dirp)
 {
 	UBA_INIT_DETOUR(rewinddir, dirp);
@@ -2641,7 +2655,6 @@ namespace uba
 
 	void PreInit(const char* logFile)
 	{
-		g_exeDirMem.Create();
 		g_fileHandlesMem.Create();
 		g_fileHandlesLockMem.Create();
 
@@ -2665,49 +2678,6 @@ namespace uba
 			g_debugFile = (FileHandle)open("/dev/tty", O_WRONLY);
 		else if (logFile)
 			g_debugFile = (FileHandle)open(logFile, O_CREAT | O_TRUNC | O_WRONLY, S_IRUSR | S_IWUSR);
-		if (g_debugFile != InvalidFileHandle)
-		{
-			char buf[32 * 1024];
-#if PLATFORM_LINUX
-			int fd = TRUE_WRAPPER(open)("/proc/self/cmdline", O_RDONLY);
-			if (fd != -1)
-			{
-				auto bufSize = read(fd, buf, sizeof(buf));
-				TRUE_WRAPPER(close)(fd);
-				if (bufSize != -1)
-				{
-					buf[bufSize] = 0;
-					const char* it = buf;
-					while (*it)
-					{
-						if (it == buf)
-						{
-							DEBUG_LOG("Executable: %s", it);
-						}
-						else
-						{
-							DEBUG_LOG("               %s", it);
-						}
-						it += strlen(it) + 1;
-					}
-					DEBUG_LOG("", it);
-				}
-			}
-#else
-			int ArgC = *_NSGetArgc();
-			char** ArgV = *_NSGetArgv();
-			DEBUG_LOG("Executable: %s", ArgV[0]);
-			for (int Arg = 1; Arg < ArgC; Arg++)
-			{
-				DEBUG_LOG("               %s", ArgV[Arg]);
-			}
-			DEBUG_LOG("");
-#endif
-			DEBUG_LOG("SystemTemp: %s", g_systemTemp.data);
-			if (g_runningRemote)
-				DEBUG_LOG("RunningRemote: true");
-			DEBUG_LOG("");
-		}
 		#endif
 
 		StringBuffer<> exePath;
@@ -2720,6 +2690,41 @@ namespace uba
 		g_exeDir.EnsureEndsWithSlash();
 	}
 
+	void LogHeader()
+	{
+		#if UBA_DEBUG_LOG_ENABLED
+		StringBuffer<32 * 1024> buf;
+		#if PLATFORM_LINUX
+		int fd = TRUE_WRAPPER(open)("/proc/self/cmdline", O_RDONLY);
+		if (fd != -1)
+		{
+			auto bufSize = read(fd, buf.data, buf.capacity);
+			if (bufSize != -1)
+			{
+				char* it = buf.data;
+				while (*it)
+				{
+					it += strlen(it);
+					*it = ' ';
+					++it;
+				}
+				buf.Resize(it - buf.data - 1);
+			}
+			TRUE_WRAPPER(close)(fd);
+		}
+		#else
+		char** argv = *_NSGetArgv();
+		for (int i=0, e=*_NSGetArgc(); i<e; ++i)
+		{
+			if (i != 0)
+				buf.Append(' ');
+			buf.Append(argv[i]);
+		}
+		#endif
+		LogHeader(buf);
+		#endif
+	}
+
 	void Init()
 	{
 		UBA_ASSERTF(!g_isInitialized, "Already initialized");
@@ -2740,6 +2745,7 @@ namespace uba
 			writer.Flush();
 			BinaryReader reader;
 
+			g_processId = reader.ReadU32();
 			bool echoOn = reader.ReadBool();
 			g_isChild = reader.ReadBool();
 
@@ -2766,10 +2772,6 @@ namespace uba
 		VirtualizePath(g_virtualWorkingDir);
 		VirtualizePath(g_exeDir);
 
-		DEBUG_LOG("Application: %s", g_virtualApplication.data);
-		DEBUG_LOG("WorkingDir: %s", g_virtualWorkingDir.data);
-		DEBUG_LOG("ExeDir: %s", g_exeDir.data);
-
 		UBA_ASSERTF(g_virtualApplicationDir.capacity > 0, "g_virtualApplicationDir.capacity > 0");
 
 		const char* lastSlash = strrchr(g_virtualApplication.data, '/');
@@ -2779,6 +2781,8 @@ namespace uba
 		setenv("PWD", g_virtualWorkingDir.data, 1);
 		g_virtualWorkingDir.EnsureEndsWithSlash();
 
+		LogHeader();
+
 		StringBuffer<128> mappedFileTableUidName;
 		GetMappingHandleName(mappedFileTableUidName, mappedFileTableHandle);
 		int mappedFileTableFd = shm_open(mappedFileTableUidName.data, O_RDONLY, S_IRUSR | S_IWUSR);
@@ -2868,14 +2872,14 @@ namespace uba
 	}
 
 	#if UBA_DEBUG_LOG_ENABLED
-	void WriteDebug(const char* str, u32 strLen)
+	void WriteDebug(const void* data, u32 dataLen)
 	{
 		int t = errno;
 
 		#if UBA_DETOUR_DEBUG && PLATFORM_LINUX
-		TRUE_WRAPPER(write)(g_debugFile, str, strLen);
+		TRUE_WRAPPER(write)(g_debugFile, data, dataLen);
 		#else
-		write(g_debugFile, str, strLen);
+		write(g_debugFile, data, dataLen);
 		#endif
 		//fsync(g_debugFile);
 		errno = t;
@@ -2955,6 +2959,7 @@ extern "C"
 			close(debugFile);
 			debugFile = open(g_logName.data, O_CREAT | O_TRUNC | O_WRONLY, S_IRUSR | S_IWUSR);
 			g_debugFile = (FileHandle)debugFile;
+			LogHeader(ToView(outArguments));
 			#endif
 		}
 
diff --git a/Engine/Source/Programs/UnrealBuildAccelerator/Detours/Private/UbaDetoursShared.cpp b/Engine/Source/Programs/UnrealBuildAccelerator/Detours/Private/UbaDetoursShared.cpp
--- a/Engine/Source/Programs/UnrealBuildAccelerator/Detours/Private/UbaDetoursShared.cpp
+++ b/Engine/Source/Programs/UnrealBuildAccelerator/Detours/Private/UbaDetoursShared.cpp
@@ -16,6 +16,7 @@ namespace uba
 	VARIABLE_MEM(ReaderWriterLock, g_communicationLock);
 	VARIABLE_MEM(StringBuffer<256>, g_logName);
 	VARIABLE_MEM(StringBuffer<512>, g_virtualWorkingDir);
+	VARIABLE_MEM(StringBuffer<256>, g_exeDir);
 	VARIABLE_MEM(StringBuffer<128>, g_systemRoot);
 	VARIABLE_MEM(StringBuffer<128>, g_systemTemp);
 	VARIABLE_MEM(MemoryBlock, g_memoryBlock);
@@ -41,6 +42,7 @@ namespace uba
 		g_communicationLockMem.Create();
 		g_logNameMem.Create();
 		g_virtualWorkingDirMem.Create();
+		g_exeDirMem.Create();
 		g_systemRootMem.Create();
 		g_systemTempMem.Create();
 
@@ -53,7 +55,7 @@ namespace uba
 
 #if UBA_DEBUG_LOG_ENABLED
 	FileHandle g_debugFile = InvalidFileHandle;
-	void WriteDebug(const char* str, u32 strLen);
+	void WriteDebug(const void* data, u32 dataLen);
 	constexpr const char g_emptyString[] = "                                                     ";
 	constexpr const char* g_emptyStringEnd = ((const char*)g_emptyString) + sizeof_array(g_emptyString) - 1;
 	thread_local StringBuffer<LogBufSize> t_a;
@@ -264,6 +266,52 @@ namespace uba
 		return false;
 	}
 
+	void LogHeader(StringView cmdLine)
+	{
+		#if UBA_DEBUG_LOG_ENABLED
+		if (g_debugFile == InvalidFileHandle)
+			return;
+		FlushDebug();
+		WriteDebugLog(TC("ProcessId: %u"), g_processId);
+		WriteDebug("CmdLine: ", 9);
+
+		#if PLATFORM_WINDOWS
+		u32 left = cmdLine.count;
+		const tchar* read = cmdLine.data;
+		while (left)
+		{
+			char buf[1024];
+			u32 written = 0;
+			while (*read && written < sizeof(buf))
+				buf[written++] = (char)*read++;
+			WriteDebug(buf, written);
+			left -= written;
+		}
+		#else
+		WriteDebug(cmdLine.data, cmdLine.count*sizeof(tchar));
+		#endif
+
+		WriteDebug("\n", 1);
+		WriteDebugLog(TC("WorkingDir: %s"), g_virtualWorkingDir.data);
+		WriteDebugLog(TC("ExeDir: %s"), g_virtualApplicationDir.data);
+		WriteDebugLog(TC("ExeDir (actual): %s"), g_exeDir.data);
+		WriteDebugLog(TC("SystemTemp: %s"), g_systemTemp.data);
+		WriteDebugLog(TC("Rules: %u (%u)"), g_rules->index, GetApplicationRules()[g_rules->index].hash);
+		if (g_runningRemote)
+		{
+			StringBuffer<256> computerName;
+			GetComputerNameW(computerName);
+			WriteDebugLog(TC("Remote: %s"), computerName.data);
+		}
+		static u32 reuseCounter;
+		if (reuseCounter)
+			WriteDebugLog(TC("ProcessReuseIndex: %u"), reuseCounter);
+		++reuseCounter;
+		WriteDebugLog(TC(""));
+		FlushDebug();
+		#endif
+	}
+
 	void LogVfsInfo()
 	{
 		for (u32 i=0; i!=g_vfsEntryCount; ++i)
diff --git a/Engine/Source/Programs/UnrealBuildAccelerator/Detours/Private/UbaDetoursShared.h b/Engine/Source/Programs/UnrealBuildAccelerator/Detours/Private/UbaDetoursShared.h
--- a/Engine/Source/Programs/UnrealBuildAccelerator/Detours/Private/UbaDetoursShared.h
+++ b/Engine/Source/Programs/UnrealBuildAccelerator/Detours/Private/UbaDetoursShared.h
@@ -28,6 +28,7 @@ namespace uba
 	bool IsVfsEnabled();
 	bool DevirtualizePath(StringBufferBase& path);
 	bool VirtualizePath(StringBufferBase& path);
+	void LogHeader(StringView cmdLine);
 	void LogVfsInfo();
 
 	#if UBA_DEBUG_LOG_ENABLED
@@ -53,6 +54,7 @@ namespace uba
 
 	extern StringBuffer<512>& g_virtualApplication;
 	extern StringBuffer<512>& g_virtualApplicationDir;
+	extern StringBuffer<256>& g_exeDir;
 	extern StringBuffer<256>& g_logName;
 
 	extern StringBuffer<512>& g_virtualWorkingDir;
@@ -92,6 +94,7 @@ namespace uba
 	inline constexpr bool g_allowFindFileDetour = true;
 	inline constexpr bool g_allowListDirectoryHandle = true;
 
+	extern u32 g_processId;
 	extern u32 g_rulesIndex;
 	extern ApplicationRules* g_rules;
 	extern bool g_runningRemote;
diff --git a/Engine/Source/Programs/UnrealBuildAccelerator/Detours/Private/Windows/UbaDetoursFunctionsKernelBase.inl b/Engine/Source/Programs/UnrealBuildAccelerator/Detours/Private/Windows/UbaDetoursFunctionsKernelBase.inl
--- a/Engine/Source/Programs/UnrealBuildAccelerator/Detours/Private/Windows/UbaDetoursFunctionsKernelBase.inl
+++ b/Engine/Source/Programs/UnrealBuildAccelerator/Detours/Private/Windows/UbaDetoursFunctionsKernelBase.inl
@@ -76,14 +76,9 @@ DWORD Local_GetLongPathNameW(LPCWSTR lpszShortPath, LPWSTR lpszLongPath, DWORD c
 LPWSTR Detoured_GetCommandLineW()
 {
 	DETOURED_CALL(GetCommandLineW);
-	if (!g_virtualCommandLineW)
-	{
-		LPWSTR str = True_GetCommandLineW();
-		DEBUG_LOG_TRUE(L"GetCommandLineW", L"");// str);
-		return str;
-	}
-	DEBUG_LOG_DETOURED(L"GetCommandLineW", L"");
-	return (LPWSTR)g_virtualCommandLineW;
+	LPWSTR str = True_GetCommandLineW();
+	DEBUG_LOG_TRUE(L"GetCommandLineW", L"");// str);
+	return str;
 }
 
 DWORD Detoured_GetCurrentDirectoryW(DWORD nBufferLength, LPWSTR lpBuffer)
@@ -4096,14 +4091,9 @@ BOOL Detoured_UnmapViewOfFile2(HANDLE Process, PVOID BaseAddress, ULONG UnmapFla
 LPSTR Detoured_GetCommandLineA()
 {
 	DETOURED_CALL(GetCommandLineA);
-	if (!g_virtualCommandLineA)
-	{
-		auto str = True_GetCommandLineA();
-		DEBUG_LOG_TRUE(L"GetCommandLineA", L"");// str);
-		return str;
-	}
-	DEBUG_LOG_DETOURED(L"GetCommandLineA", L"");
-	return g_virtualCommandLineA;
+	auto str = True_GetCommandLineA();
+	DEBUG_LOG_TRUE(L"GetCommandLineA", L"");// str);
+	return str;
 }
 
 LPWSTR* Detoured_CommandLineToArgvW(LPCWSTR lpCmdLine, int *pNumArgs)
diff --git a/Engine/Source/Programs/UnrealBuildAccelerator/Detours/Private/Windows/UbaDetoursFunctionsWin.cpp b/Engine/Source/Programs/UnrealBuildAccelerator/Detours/Private/Windows/UbaDetoursFunctionsWin.cpp
--- a/Engine/Source/Programs/UnrealBuildAccelerator/Detours/Private/Windows/UbaDetoursFunctionsWin.cpp
+++ b/Engine/Source/Programs/UnrealBuildAccelerator/Detours/Private/Windows/UbaDetoursFunctionsWin.cpp
@@ -71,18 +71,19 @@ constexpr int StdOutFd = -2;
 #if UBA_DEBUG_LOG_ENABLED
 bool g_debugFileFlushOnWrite = false;
 
-void WriteDebug(const char* s, u32 strLen)
+void WriteDebug(const void* data, u32 dataLen)
 {
-	DWORD toWrite = (DWORD)strLen;
+	auto readMem = (const u8*)data;
+	DWORD toWrite = (DWORD)dataLen;
 	DWORD lastError = GetLastError();
 	while (true)
 	{
 		DWORD written;
-		if (True_WriteFile((HANDLE)g_debugFile, s, toWrite, &written, NULL))
+		if (True_WriteFile((HANDLE)g_debugFile, readMem, toWrite, &written, NULL))
 			break;
 		if (GetLastError() != ERROR_IO_PENDING)
 			break;// ExitProcess(1340); // During shutdown this might actually cause error and we just ignore that and break out
-		s += written;
+		readMem += written;
 		toWrite -= written;
 	}
 	if (g_debugFileFlushOnWrite)
@@ -114,11 +115,6 @@ bool g_isRunningWine;
 int g_uiLanguage;
 u32 g_processId;
 
-StringBuffer<256> g_exeDir;
-
-const wchar_t* g_virtualCommandLineW;
-char* g_virtualCommandLineA;
-
 constexpr u32 TrackInputsMemCapacity = 512 * 1024;
 u8* g_trackInputsMem;
 u32 g_trackInputsBufPos;
@@ -1001,8 +997,6 @@ void Init(const DetoursPayload& payload, u64 startTime)
 		g_stdHandle[2] = makeDetouredHandle(new DetouredHandle(HandleType_StdIn)); // STD_IN
 	}
 
-	g_processId = payload.processId;
-
 	if (payload.trackInputs)
 		g_trackInputsMem = (u8*)g_memoryBlock.Allocate(TrackInputsMemCapacity, 1, L"TrackInputs");
 	
@@ -1028,6 +1022,7 @@ void Init(const DetoursPayload& payload, u64 startTime)
 		writer.Flush();
 		BinaryReader reader;
 
+		g_processId = reader.ReadU32();
 		g_echoOn = reader.ReadBool();
 		g_isChild = reader.ReadBool();
 
@@ -1066,56 +1061,11 @@ void Init(const DetoursPayload& payload, u64 startTime)
 			FatalError(4444, L"What the heck: %s (%s)", g_virtualApplication.data, applicationBuffer.data);
 	}
 
-	const wchar_t* cmdLine = True_GetCommandLineW();(void)cmdLine;
-
-	// Command line should always have the virtualized proper host path, if it doesn't, then it is CreateProcess call that is wrong
-	#if 0
-	const wchar_t* exePos;
-	if (g_runningRemote && Contains(cmdLine, g_exeDir.data, true, &exePos))
-	{
-		StringBuffer<> buf;
-		buf.Append(cmdLine, exePos - cmdLine);
-		buf.Append(g_virtualApplicationDir);
-		TString realCmdLine(buf.data);
-		realCmdLine += (cmdLine + g_exeDir.count + 1);
-		g_virtualCommandLineW = g_memoryBlock.Strdup(realCmdLine).data;
-		cmdLine = g_virtualCommandLineW;
-	}
-	//else
-	//	g_virtualCommandLineW = g_memoryBlock.Strdup(cmdLine);
-
-	if (g_virtualCommandLineW)
-	{
-		u64 len = wcslen(g_virtualCommandLineW);
-		g_virtualCommandLineA = (char*)g_memoryBlock.Allocate(len + 1, 1, L"");
-		size_t res;
-		wcstombs_s(&res, g_virtualCommandLineA, len + 1, g_virtualCommandLineW, len);
-	}
-	#endif
-
 	#if UBA_DEBUG_LOG_ENABLED
 	if (isLogging())
 	{
-		u64 cmdLineLen = wcslen(cmdLine);
-		wchar_t temp[LogBufSize - 10];
-		if (cmdLineLen > sizeof_array(temp))
-		{
-			memcpy(temp, cmdLine, sizeof(temp));
-			temp[sizeof_array(temp)-1] = 0;
-			cmdLine = temp;
-		}
-		DEBUG_LOG(L"ProcessId: %u", payload.processId);
-		DEBUG_LOG(L"Cmdline: %ls", cmdLine);
-		DEBUG_LOG(L"WorkingDir: %ls", g_virtualWorkingDir.data);
-		DEBUG_LOG(L"ExeDir: %ls", g_virtualApplicationDir.data);
-		DEBUG_LOG(L"ExeDir (true): %ls", g_exeDir.data);
-		DEBUG_LOG(L"Rules: %u (%u)", g_rules->index, GetApplicationRules()[g_rules->index].hash);
-		if (g_runningRemote)
-		{
-			StringBuffer<256> computerName;
-			GetComputerNameW(computerName);
-			DEBUG_LOG(L"Remote: %s", computerName.data);
-		}
+		StringView sv(ToView(True_GetCommandLineW()));
+		LogHeader(sv);
 		LogVfsInfo();
 	}
 	#endif
@@ -1281,6 +1231,7 @@ extern "C"
 			CloseHandle(debugFile);
 			debugFile = CreateFileW(g_logName.data, GENERIC_WRITE, FILE_SHARE_READ, NULL, CREATE_ALWAYS, FILE_ATTRIBUTE_NORMAL, NULL);
 			g_debugFile = (FileHandle)(u64)debugFile;
+			LogHeader(ToView(outArguments));
 			#endif
 		}
 
