[UBA]
* Added storage.ReportFileWrite before writing files sent from remotes to host. This is needed to flush out potential deferred cas entries.

diff --git a/Engine/Source/Programs/UnrealBuildAccelerator/Common/Private/UbaSessionServer.cpp b/Engine/Source/Programs/UnrealBuildAccelerator/Common/Private/UbaSessionServer.cpp
--- a/Engine/Source/Programs/UnrealBuildAccelerator/Common/Private/UbaSessionServer.cpp
+++ b/Engine/Source/Programs/UnrealBuildAccelerator/Common/Private/UbaSessionServer.cpp
@@ -1134,6 +1134,9 @@ namespace uba
 				return true;
 			}
 
+			// We need to report before writing to flush out potential 
+			m_storage.ReportFileWrite(destinationKey, destination.data);
+
 			if (shouldWriteToDisk)
 			{
 				bool writeCompressed = false;
