[UBA]
* Added validation of deferred cas writes to make sure hash created when write was deferred still matches the file that is being written

diff --git a/Engine/Source/Programs/UnrealBuildAccelerator/Common/Private/UbaSession.cpp b/Engine/Source/Programs/UnrealBuildAccelerator/Common/Private/UbaSession.cpp
--- a/Engine/Source/Programs/UnrealBuildAccelerator/Common/Private/UbaSession.cpp
+++ b/Engine/Source/Programs/UnrealBuildAccelerator/Common/Private/UbaSession.cpp
@@ -453,6 +453,7 @@ namespace uba
 			entry.canBeFreed = true;
 			entry.usedCount = 0;
 			entry.usedCountBeforeFree = g_globalRules.GetUsedCountBeforeFree(file.name);
+			entry.writtenSize = writtenSize;
 
 			#if UBA_DEBUG_TRACK_MAPPING
 			entry.name = file.name;
@@ -1694,6 +1695,7 @@ namespace uba
 				out.offset = entry.mappingOffset;
 				out.size = entry.size;
 				out.lastWriteTime = entry.lastWriteTime;
+				out.writtenSize = entry.writtenSize;
 				return true;
 			});
 
diff --git a/Engine/Source/Programs/UnrealBuildAccelerator/Common/Private/UbaStorage.cpp b/Engine/Source/Programs/UnrealBuildAccelerator/Common/Private/UbaStorage.cpp
--- a/Engine/Source/Programs/UnrealBuildAccelerator/Common/Private/UbaStorage.cpp
+++ b/Engine/Source/Programs/UnrealBuildAccelerator/Common/Private/UbaStorage.cpp
@@ -89,6 +89,7 @@ namespace uba
 		table.GetValueAsBool(asyncUnmapViewOfFile, TC("AsyncUnmapViewOfFile"));
 		table.GetValueAsU64(casCapacityBytes, TC("CasCapacityBytes"));
 		table.GetValueAsBool(keepTransientDataMapped, TC("KeepTransientDataMapped"));
+		table.GetValueAsBool(storeCompressed, TC("StoreCompressed"));
 
 		const tchar* compressor;
 		if (table.GetValueAsString(compressor, TC("Compressor")))
@@ -586,6 +587,7 @@ namespace uba
 				u8* fileMem = scope.MapView(TC("WriteCompressed"), externalMapping.handle, externalMapping.offset, externalMapping.size, fileName);
 				if (!fileMem)
 					return false;
+				out.readSize = externalMapping.writtenSize;
 				return StorageImpl::WriteCompressed(out, fileName, InvalidFileHandle, fileMem, externalMapping.size, casFile, nullptr, 0);
 			}
 		}
@@ -598,6 +600,7 @@ namespace uba
 		u64 readFileSize = 0;
 		if (!uba::GetFileSizeEx(readFileSize, readHandle))
 			return m_logger.Error(TC("[WriteCasFileNoCheck] GetFileSize failed for %s (%s)"), fileName, LastErrorToText().data);
+		out.readSize = readFileSize;
 
 		bool readIsCompressed = false;
 		if (readFileSize >= sizeof(CompressedFileHeader) && g_globalRules.FileCanBeCompressed(ToView(fileName)))
@@ -940,7 +943,7 @@ namespace uba
 	}
 
 
-	bool StorageImpl::AddCasFile(StringKey fileNameKey, const tchar* fileName, const CasKey& casKey, bool deferCreation)
+	bool StorageImpl::AddCasFile(StringKey fileNameKey, const tchar* fileName, const CasKey& casKey, bool deferCreation, u64 fileSize, u64 lastWriteTime)
 	{
 		UBA_ASSERTF(IsCompressed(casKey) == m_storeCompressed, TC("CasKey compress mode must match storage compress mode. Compress mode Key: %i, Store: %i Cas: %s (%s)"), (IsCompressed(casKey) ? 1 : 0), (m_storeCompressed ? 1 : 0), CasKeyString(casKey).str, fileName);
 		CasEntry& casEntry = GetOrCreateCasEntry(casKey);
@@ -987,6 +990,8 @@ namespace uba
 			{
 				DeferedCasCreation& dcc = res.first->second;
 				dcc.fileName = fileName;
+				dcc.fileSize = fileSize;
+				dcc.lastWriteTime = lastWriteTime;
 			}
 			auto res2 = m_deferredCasCreationLookupByName.try_emplace(fileNameKey, casKey);
 			if (res2.second)
@@ -995,6 +1000,14 @@ namespace uba
 				return m_logger.Error(TC("Same file %s is registered with different caskey (%s vs %s)"), fileName, CasKeyString(res2.first->second).str, CasKeyString(casKey).str);
 			return true;
 		}
+		#if UBA_DEBUG
+		else
+		{
+			SCOPED_READ_LOCK(m_deferredCasCreationLookupLock, deferredLock);
+			auto findIt = m_deferredCasCreationLookupByName.find(fileNameKey);
+			UBA_ASSERTF(findIt == m_deferredCasCreationLookupByName.end() || findIt->second == casKey, TC("Found deferred entry for %s with different cas key while writing cas. A ReportFileWrite is missing somewhere"), fileName);
+		}
+		#endif
 
 		casEntry.verified = true;
 		casEntry.exists = false;
@@ -1928,7 +1941,7 @@ namespace uba
 			if (fileEntry.casKey != CasKeyZero)
 			{
 				UBA_ASSERT(casKeyOverride == CasKeyZero || casKeyOverride == fileEntry.casKey);
- 				if (!AddCasFile(fileNameKey, fileName, fileEntry.casKey, deferCreation))
+ 				if (!AddCasFile(fileNameKey, fileName, fileEntry.casKey, deferCreation, fileEntry.size, fileEntry.lastWritten))
 					return false;
 			}
 			out = fileEntry.casKey;
@@ -1940,6 +1953,7 @@ namespace uba
 		ExternalFileMapping externalMapping;
 
 		u64 fileSize = 0;
+		u64 writtenSize = 0;
 		u64 lastWriteTime = 0;
 
 		FileHandle fileHandle = InvalidFileHandle;
@@ -1948,6 +1962,7 @@ namespace uba
 		if (m_externalFileMappingsProvider && m_externalFileMappingsProvider(externalMapping, fileNameKey, fileName))
 		{
 			fileSize = externalMapping.size;
+			writtenSize = externalMapping.writtenSize;
 			lastWriteTime = externalMapping.lastWriteTime;
 		}
 		else
@@ -1967,6 +1982,7 @@ namespace uba
 				return m_logger.Error(TC("GetFileInformationByHandle failed on %s"), fileName);
 			}
 			fileSize = info.size;
+			writtenSize = fileSize;
 			lastWriteTime = info.lastWriteTime;
 		}
 
@@ -1975,14 +1991,14 @@ namespace uba
 			if (casKeyOverride != CasKeyZero && casKeyOverride != fileEntry.casKey)
 			{
 				fileEntry.casKey = casKeyOverride;
-				if (!AddCasFile(fileNameKey, fileName, fileEntry.casKey, deferCreation))
+				if (!AddCasFile(fileNameKey, fileName, fileEntry.casKey, deferCreation, writtenSize, lastWriteTime))
 					return false;
 				out = fileEntry.casKey;
 				return true;
 			}
 			if (fileSize == fileEntry.size && (lastWriteTime && lastWriteTime == fileEntry.lastWritten) && fileEntry.casKey != CasKeyInvalid)
 			{
-				if (!AddCasFile(fileNameKey, fileName, fileEntry.casKey, deferCreation))
+				if (!AddCasFile(fileNameKey, fileName, fileEntry.casKey, deferCreation, writtenSize, lastWriteTime))
 					return false;
 				out = fileEntry.casKey;
 				return true;
@@ -2034,7 +2050,7 @@ namespace uba
 		if (fileEntry.casKey == CasKeyZero)
 			return false;
 
-		if (!AddCasFile(fileNameKey, fileName, fileEntry.casKey, deferCreation))
+		if (!AddCasFile(fileNameKey, fileName, fileEntry.casKey, deferCreation, writtenSize, lastWriteTime))
 			return false;
 
 		out = fileEntry.casKey;
@@ -2197,6 +2213,7 @@ namespace uba
 			return false;
 		StringKey fileNameKey = findIt->second.fileNameKey;
 		StringBuffer<> deferredCreation(findIt->second.fileName);
+		u64 deferredFileSize = findIt->second.fileSize;
 		for (auto& it2 : findIt->second.names)
 			m_deferredCasCreationLookupByName.erase(it2);
 		m_deferredCasCreationLookup.erase(findIt);
@@ -2204,12 +2221,20 @@ namespace uba
 		if (casEntry.beingWritten)
 			m_logger.Warning(TC("Deferred cas %s is being written from network. This should never happen (%s)"), CasKeyString(casKey).str, deferredCreation.data);
 
+		// TODO: We need to validate that the file is still the same file.
+		// It seems like there can be a race condition with ispc and .generated.h where file content changes
+		// while .generated.h is in the deferred creation lookup.
+		// .. we need to add verification for this and throw an error
+
 		casEntry.verified = true;
 		deferredLock.Leave();
 		WriteResult res;
 		if (!WriteCasFile(res, deferredCreation.data, casKey))
 			return m_logger.Error(TC("Failed to write deferred cas %s (%s)"), CasKeyString(casKey).str, deferredCreation.data);
 
+		if (res.readSize != deferredFileSize) // TODO: Should also check last written time
+			m_logger.Warning(TC("Deferred cas %s was deferred with file size %llu but was read with size %llu (%s)"), CasKeyString(casKey).str, deferredFileSize, res.readSize, deferredCreation.data);
+
 		casEntry.exists = true;
 		entryLock.Leave();
 
@@ -2680,6 +2705,9 @@ namespace uba
 		auto findIt = m_deferredCasCreationLookupByName.find(fileNameKey);
 		if (findIt == m_deferredCasCreationLookupByName.end())
 			return;
+
+		m_logger.Logf(IsDebug ? LogEntryType_Warning : LogEntryType_Info, TC("Flush cas creation of file %s (This should be rare)"), fileName);
+
 		CasKey key = findIt->second;
 		deferredLock.Leave();
 		HasCasFile(key);
diff --git a/Engine/Source/Programs/UnrealBuildAccelerator/Common/Private/UbaStorageServer.cpp b/Engine/Source/Programs/UnrealBuildAccelerator/Common/Private/UbaStorageServer.cpp
--- a/Engine/Source/Programs/UnrealBuildAccelerator/Common/Private/UbaStorageServer.cpp
+++ b/Engine/Source/Programs/UnrealBuildAccelerator/Common/Private/UbaStorageServer.cpp
@@ -91,8 +91,10 @@ namespace uba
 		out.size = InvalidValue;
 
 		CasKey actualKey = casKey;
-		//if (m_storeCompressed)
-		//	actualKey = AsCompressed(casKey, true);
+
+		// This must be true because code handling storing files from remotes does not have decompression support
+		// .. causing mix of compressed and non-compressed cas files with too many permutations
+		UBA_ASSERT(m_storeCompressed);
 
 		SCOPED_FUTEX(m_waitEntriesLock, waitLock);
 		WaitEntry& waitEntry = m_waitEntries[actualKey];
diff --git a/Engine/Source/Programs/UnrealBuildAccelerator/Common/Public/UbaSession.h b/Engine/Source/Programs/UnrealBuildAccelerator/Common/Public/UbaSession.h
--- a/Engine/Source/Programs/UnrealBuildAccelerator/Common/Public/UbaSession.h
+++ b/Engine/Source/Programs/UnrealBuildAccelerator/Common/Public/UbaSession.h
@@ -264,6 +264,7 @@ namespace uba
 			u64 mappingOffset = 0;
 			u64 size = 0;
 			u64 lastWriteTime = 0;
+			u64 writtenSize = InvalidValue;
 			u16 refCount = 0;
 			u8 usedCount = 0;
 			u8 usedCountBeforeFree = 0;
diff --git a/Engine/Source/Programs/UnrealBuildAccelerator/Common/Public/UbaStorage.h b/Engine/Source/Programs/UnrealBuildAccelerator/Common/Public/UbaStorage.h
--- a/Engine/Source/Programs/UnrealBuildAccelerator/Common/Public/UbaStorage.h
+++ b/Engine/Source/Programs/UnrealBuildAccelerator/Common/Public/UbaStorage.h
@@ -68,10 +68,10 @@ namespace uba
 		virtual void SetTrace(Trace* trace, bool detailed) {}
 		virtual void Ping() {}
 
-		struct WriteResult { FileMappingHandle mappingHandle; u64 size = InvalidValue; u64 offset = InvalidValue; };
+		struct WriteResult { FileMappingHandle mappingHandle; u64 readSize = InvalidValue; u64 size = InvalidValue; u64 offset = InvalidValue; };
 		virtual bool WriteCompressed(WriteResult& out, const tchar* from, FileHandle readHandle, u8* readMem, u64 fileSize, const tchar* toFile, const void* header, u64 headerSize, u64 lastWriteTime = 0) = 0;
 
-		struct ExternalFileMapping { FileMappingHandle handle; u64 offset = 0; u64 size = 0; u64 lastWriteTime = 0; };
+		struct ExternalFileMapping { FileMappingHandle handle; u64 offset = 0; u64 size = 0; u64 lastWriteTime = 0; u64 writtenSize = InvalidValue; };
 		using ExternalFileMappingsProvider = Function<bool(ExternalFileMapping& out, StringKey fileNameKey, const tchar* fileName)>;
 		virtual void RegisterExternalFileMappingsProvider(ExternalFileMappingsProvider&& provider) {}
 	};
@@ -173,7 +173,7 @@ namespace uba
 		bool WriteCasFileNoCheck(WriteResult& out, const StringKey& fileNameKey, const tchar* fileName, const CasKey& casKey, const tchar* casFile, bool storeCompressed);
 		bool WriteCasFile(WriteResult& out, const tchar* fileName, const CasKey& casKey);
 		bool VerifyExisting(bool& outReturnValue, ScopedWriteLock& entryLock, const CasKey& casKey, CasEntry& casEntry, StringBufferBase& casFile);
-		bool AddCasFile(StringKey fileNameKey, const tchar* fileName, const CasKey& casKey, bool deferCreation);
+		bool AddCasFile(StringKey fileNameKey, const tchar* fileName, const CasKey& casKey, bool deferCreation, u64 fileSize = InvalidValue, u64 lastWriteTime = 0);
 		void CasEntryAccessed(const CasKey& casKey);
 		virtual bool IsDisallowedPath(const tchar* fileName);
 		virtual bool DecompressMemoryToMemory(const u8* compressedData, u64 compressedSize, u8* writeData, u64 decompressedSize, const tchar* readHint, const tchar* writeHint) override;
@@ -269,7 +269,7 @@ namespace uba
 
 		ReaderWriterLock m_deferredCasCreationLookupLock;
 		using NameToCas = UnorderedMap<StringKey, CasKey>;
-		struct DeferedCasCreation { StringKey fileNameKey; TString fileName; List<NameToCas::iterator> names; };
+		struct DeferedCasCreation { StringKey fileNameKey; TString fileName; List<NameToCas::iterator> names; u64 fileSize = InvalidValue; u64 lastWriteTime = 0;};
 		UnorderedMap<CasKey, DeferedCasCreation> m_deferredCasCreationLookup;
 		NameToCas m_deferredCasCreationLookupByName;
 
diff --git a/Engine/Source/Programs/UnrealBuildAccelerator/Core/Private/UbaApplicationRules.cpp b/Engine/Source/Programs/UnrealBuildAccelerator/Core/Private/UbaApplicationRules.cpp
--- a/Engine/Source/Programs/UnrealBuildAccelerator/Core/Private/UbaApplicationRules.cpp
+++ b/Engine/Source/Programs/UnrealBuildAccelerator/Core/Private/UbaApplicationRules.cpp
@@ -376,6 +376,7 @@ namespace uba
 
 	class ApplicationRulesClangClExe : public ApplicationRulesClangPlusPlusExe
 	{
+		using Super = ApplicationRulesClangPlusPlusExe;
 	public:
 		virtual DependencyCrawlerType GetDependencyCrawlerType() const override
 		{
@@ -384,7 +385,7 @@ namespace uba
 
 		virtual bool IsOutputFile(const StringView& fileName) const override
 		{
-			return fileName.Contains(TCV("response-")); // file used by internally spawned clang-cl.exe. Note this is only generated when command line is very long
+			return Super::IsOutputFile(fileName) || fileName.Contains(TCV("response-")); // file used by internally spawned clang-cl.exe. Note this is only generated when command line is very long
 		}
 
 		virtual bool KeepInMemory(StringView fileName, StringView systemTemp, bool isRunningRemote, bool isWrite) const override
diff --git a/Engine/Source/Programs/UnrealBuildAccelerator/Core/Public/UbaPlatform.h b/Engine/Source/Programs/UnrealBuildAccelerator/Core/Public/UbaPlatform.h
--- a/Engine/Source/Programs/UnrealBuildAccelerator/Core/Public/UbaPlatform.h
+++ b/Engine/Source/Programs/UnrealBuildAccelerator/Core/Public/UbaPlatform.h
@@ -72,10 +72,12 @@ namespace uba
 	#define UBA_ASSERT(x) do { if (x) break; uba::UbaAssert(TC(""), __FILE__, __LINE__, #x, true, 543221, nullptr, 0); } while(false)
 	#define UBA_ASSERTF(x, fmt, ...) do { UBA_FMT_CHECK(fmt, ##__VA_ARGS__); if (x) break; uba::StringBuffer<> _buf; _buf.Appendf(fmt, ##__VA_ARGS__); uba::UbaAssert(_buf.data, __FILE__, __LINE__, #x, true, 543221, nullptr, 0); } while(false)
 	#define UBA_ASSERT_MESSAGEBOX 0
+	inline constexpr bool IsDebug = true;
 	#else
 	#define UBA_ASSERT(x) do { } while(false)
 	#define UBA_ASSERTF(x, ...) do { } while(false)
 	#define UBA_ASSERT_MESSAGEBOX 0
+	inline constexpr bool IsDebug = false;
 	#endif
 
 	#define UBA_NOT_IMPLEMENTED(x) UBA_ASSERTF(false, TC("%s not implemented!"), #x);
diff --git a/Engine/Source/Programs/UnrealBuildAccelerator/Test/Private/UbaTest.cpp b/Engine/Source/Programs/UnrealBuildAccelerator/Test/Private/UbaTest.cpp
--- a/Engine/Source/Programs/UnrealBuildAccelerator/Test/Private/UbaTest.cpp
+++ b/Engine/Source/Programs/UnrealBuildAccelerator/Test/Private/UbaTest.cpp
@@ -1,9 +1,36 @@
 // Copyright Epic Games, Inc. All Rights Reserved.
 
 #include "UbaTestAll.h"
+#include "UbaFileAccessor.h"
 
 namespace uba
 {
+	bool CreateTestFile(LoggerWithWriter& logger, StringView testRootDir, StringView fileName, StringView content)
+	{
+		StringBuffer<> testFileName;
+		testFileName.Append(testRootDir).EnsureEndsWithSlash().Append(fileName).FixPathSeparators();
+
+		if (fileName.Contains(PathSeparator) || fileName.Contains(NonPathSeparator))
+		{
+			StringBuffer<> testFileDir;
+			testFileDir.AppendDir(testFileName);
+			DirectoryCache().CreateDirectory(logger, testFileDir.data);
+		}
+
+		FileAccessor file(logger, testFileName.data);
+		if (!file.CreateMemoryWrite(false, DefaultAttributes(), content.count))
+			return logger.Error(TC("Failed to create file for write"));
+		memcpy(file.GetData(), content.data, content.count);
+		return file.Close();
+	}
+
+	bool DeleteTestFile(LoggerWithWriter& logger, StringView testRootDir, StringView fileName)
+	{
+		StringBuffer<> testFileName;
+		testFileName.Append(testRootDir).EnsureEndsWithSlash().Append(fileName).FixPathSeparators();
+		return DeleteFileW(testFileName.data);
+	}
+
 	bool WrappedMain(int argc, tchar* argv[])
 	{
 		AddExceptionHandler();
diff --git a/Engine/Source/Programs/UnrealBuildAccelerator/Test/Private/UbaTest.h b/Engine/Source/Programs/UnrealBuildAccelerator/Test/Private/UbaTest.h
--- a/Engine/Source/Programs/UnrealBuildAccelerator/Test/Private/UbaTest.h
+++ b/Engine/Source/Programs/UnrealBuildAccelerator/Test/Private/UbaTest.h
@@ -16,6 +16,9 @@ namespace uba
 	class SessionServer;
 	struct ProcessStartInfo;
 
+	bool CreateTestFile(LoggerWithWriter& logger, StringView testRootDir, StringView fileName, StringView content);
+	bool DeleteTestFile(LoggerWithWriter& logger, StringView testRootDir, StringView fileName);
+
 	using RunProcessFunction = Function<ProcessHandle(const ProcessStartInfo&)>;
 	using TestSessionFunction = Function<bool(LoggerWithWriter& logger, SessionServer& session, const tchar* workingDir, const RunProcessFunction& runProcess)>;
 	bool RunLocal(LoggerWithWriter& logger, const StringBufferBase& testRootDir, const TestSessionFunction& testFunc, bool enableDetour = true);
diff --git a/Engine/Source/Programs/UnrealBuildAccelerator/Test/Private/UbaTestAll.h b/Engine/Source/Programs/UnrealBuildAccelerator/Test/Private/UbaTestAll.h
--- a/Engine/Source/Programs/UnrealBuildAccelerator/Test/Private/UbaTestAll.h
+++ b/Engine/Source/Programs/UnrealBuildAccelerator/Test/Private/UbaTestAll.h
@@ -74,6 +74,7 @@ namespace uba
 		UBA_TEST(TestStdOutViaCmd) \
 		UBA_TEST(TestVolumeCache) \
 		UBA_TEST(TestDependencyCrawler) \
+		UBA_TEST(TestRemoteProcessSpecialCase1) \
 
 
 #define UBA_LINUX_TESTS \
diff --git a/Engine/Source/Programs/UnrealBuildAccelerator/Test/Private/UbaTestSession.cpp b/Engine/Source/Programs/UnrealBuildAccelerator/Test/Private/UbaTestSession.cpp
--- a/Engine/Source/Programs/UnrealBuildAccelerator/Test/Private/UbaTestSession.cpp
+++ b/Engine/Source/Programs/UnrealBuildAccelerator/Test/Private/UbaTestSession.cpp
@@ -101,8 +101,8 @@ namespace uba
 		return logger.Error(TC("UbaTestApp returned exit code %u"), exitCode);
 	}
 
-	using TestServerClientSessionFunction = Function<bool(LoggerWithWriter& logger, const StringView& workingDir, SessionServer& sessionServer, SessionClient& sessionClient)>;
-	bool SetupServerClientSession(LoggerWithWriter& logger, const StringBufferBase& testRootDir, bool deleteAll, bool serverShouldListen, const Config& serverConfig, const TestServerClientSessionFunction& testFunc)
+	using TestServerSessionFunction = Function<bool(LoggerWithWriter& logger, const StringView& workingDir, SessionServer& sessionServer)>;
+	bool SetupServerSession(LoggerWithWriter& logger, const StringBufferBase& testRootDir, bool deleteAll, bool serverShouldListen, const Config& serverConfig, const TestServerSessionFunction& testFunc)
 	{
 		LogWriter& logWriter = logger.m_writer;
 		NetworkBackendTcpCreateInfo tcpInfo{logWriter};
@@ -111,11 +111,14 @@ namespace uba
 
 		bool ctorSuccess = true;
 		NetworkServer server(ctorSuccess, { logWriter });
-		NetworkClient client(ctorSuccess, { logWriter });
 
 		StringBuffer<MaxPath> rootDir;
 		rootDir.Append(testRootDir).Append(TCV("Uba"));
-		if (deleteAll && !DeleteAllFiles(logger, rootDir.data))
+
+		StringBuffer<MaxPath> toDelete(rootDir);
+		if (!deleteAll)
+			toDelete.Append(PathSeparator).Append(TCV("sessions"));
+		if (!DeleteAllFiles(logger, toDelete.data))
 			return false;
 
 		StorageServerCreateInfo storageServerInfo(server, rootDir.data, logWriter);
@@ -127,6 +130,7 @@ namespace uba
 		SessionServerCreateInfo sessionServerInfo(storageServer, server, logWriter);
 		sessionServerInfo.checkMemory = false;
 		sessionServerInfo.rootDir = rootDir.data;
+		sessionServerInfo.useUniqueId = false;
 
 		#if UBA_DEBUG
 		sessionServerInfo.logToFile = true;
@@ -145,12 +149,37 @@ namespace uba
 		clientConfig.AddTable(TC("Storage")).AddValue(TC("CheckExistsOnServer"), true);
 		server.SetClientsConfig(clientConfig);
 
-		u16 port = 1356;
+		StringBuffer<MaxPath> workingDir;
+		workingDir.Append(testRootDir).Append(TCV("WorkingDir"));
+		if (deleteAll && !DeleteAllFiles(logger, workingDir.data))
+			return false;
+		if (!storageServer.CreateDirectory(workingDir.data))
+			return false;
+		if (deleteAll && !DeleteAllFiles(logger, workingDir.data, false))
+			return false;
+
+		workingDir.EnsureEndsWithSlash();
+		return testFunc(logger, workingDir, sessionServer);
+	}
+
+	using TestClientSessionFunction = Function<bool(LoggerWithWriter& logger, SessionClient& sessionClient)>;
+	bool SetupClientSession(LoggerWithWriter& logger, const StringBufferBase& testRootDir, bool deleteAll, bool serverShouldListen, NetworkServer& server, u16 port, const TestClientSessionFunction& testFunc)
+	{
+		Config serverConfig;
+
+
+		LogWriter& logWriter = logger.m_writer;
+		NetworkBackendTcpCreateInfo tcpInfo{logWriter};
+		tcpInfo.Apply(serverConfig);
+		NetworkBackendTcp tcpBackend(tcpInfo);
+
+		bool ctorSuccess = true;
+		NetworkClient client(ctorSuccess, { logWriter });
 
 		if (serverShouldListen)
 		{
 			if (!server.StartListen(tcpBackend, port))
-				return logger.Error(TC("Failed to listen"));
+				return false;
 			if (!client.Connect(tcpBackend, TC("127.0.0.1"), port))
 				return logger.Error(TC("Failed to connect"));
 			if (!client.Connect(tcpBackend, TC("127.0.0.1"), port))
@@ -165,15 +194,19 @@ namespace uba
 			while (server.HasConnectInProgress())
 				Sleep(1);
 		}
-
-		auto disconnectGuard = MakeGuard([&]() { client.Disconnect(); });
+		auto disconnectGuard = MakeGuard([&]() { tcpBackend.StopListen(); client.Disconnect(); server.RemoveDisconnectedConnections(); });
 
 		Config config;
 		if (!client.FetchConfig(config))
 			return false;
 
-		rootDir.Append(TCV("Client"));
-		if (deleteAll && !DeleteAllFiles(logger, rootDir.data))
+		StringBuffer<MaxPath> rootDir;
+		rootDir.Append(testRootDir).Append(TCV("UbaClient")).AppendValue(port);
+
+		StringBuffer<MaxPath> toDelete(rootDir);
+		if (!deleteAll)
+			toDelete.Append(PathSeparator).Append(TCV("sessions"));
+		if (!DeleteAllFiles(logger, toDelete.data))
 			return false;
 
 		StorageClientCreateInfo storageClientInfo(client, rootDir.data);
@@ -183,6 +216,7 @@ namespace uba
 
 		SessionClientCreateInfo sessionClientInfo(storageClient, client, logWriter);
 		sessionClientInfo.rootDir = rootDir.data;
+		sessionClientInfo.useUniqueId = false;
 		//sessionClientInfo.allowKeepFilesInMemory = false;
 
 		#if UBA_DEBUG
@@ -194,21 +228,21 @@ namespace uba
 
 		auto cg = MakeGuard([&]() { sessionClient.Stop(); disconnectGuard.Execute(); });
 
-		StringBuffer<MaxPath> workingDir;
-		workingDir.Append(testRootDir).Append(TCV("WorkingDir"));
-		if (deleteAll && !DeleteAllFiles(logger, workingDir.data))
-			return false;
-		if (!storageServer.CreateDirectory(workingDir.data))
-			return false;
-		if (deleteAll && !DeleteAllFiles(logger, workingDir.data, false))
-			return false;
-
 		storageClient.Start();
 		sessionClient.Start();
+		return testFunc(logger, sessionClient);
+	}
 
-
-		workingDir.EnsureEndsWithSlash();
-		return testFunc(logger, workingDir, sessionServer, sessionClient);
+	using TestServerClientSessionFunction = Function<bool(LoggerWithWriter& logger, const StringView& workingDir, SessionServer& sessionServer, SessionClient& sessionClient)>;
+	bool SetupServerClientSession(LoggerWithWriter& logger, const StringBufferBase& testRootDir, bool deleteAll, bool serverShouldListen, const Config& serverConfig, const TestServerClientSessionFunction& testFunc)
+	{
+		return SetupServerSession(logger, testRootDir, deleteAll, serverShouldListen, serverConfig, [&](LoggerWithWriter& logger, const StringView& workingDir, SessionServer& sessionServer)
+			{
+				return SetupClientSession(logger, testRootDir, deleteAll, serverShouldListen, sessionServer.GetServer(), 1356, [&](LoggerWithWriter& logger, SessionClient& sessionClient)
+					{
+						return testFunc(logger, workingDir, sessionServer, sessionClient);
+					});
+			});
 	}
 
 	bool RunRemote(LoggerWithWriter& logger, const StringBufferBase& testRootDir, const Config& serverConfig, const TestSessionFunction& testFunc, bool deleteAll = true, bool serverShouldListen = true)
@@ -346,7 +380,7 @@ namespace uba
 		processInfo.workingDir = workingDir;
 		//processInfo.logFile = TC("/mnt/e/temp/ttt/RunClang.log");
 		ProcessHandle process = runProcess(processInfo);
-		if (!process.WaitForExit(40000))
+		if (!process.WaitForExit(GetTimeoutTime()))
 			return logger.Error(TC("clang++ timed out"));
 		u32 exitCode = process.GetExitCode();
 		if (exitCode != 0)
@@ -732,4 +766,40 @@ namespace uba
 		return RunRemote(logger, testRootDir, RunXCodeSelect);
 	}
 	#endif
+
+	bool TestRemoteProcessSpecialCase1(LoggerWithWriter& logger, const StringBufferBase& testRootDir)
+	{
+		Config config;
+		return SetupServerSession(logger, testRootDir, true, true, config,
+			[&](LoggerWithWriter& logger, const StringView& workingDir, SessionServer& sessionServer)
+			{
+				auto runProcess = [&](const ProcessStartInfo& pi) { return sessionServer.RunProcessRemote(pi); };
+				if (!CreateTestFile(logger, workingDir, TCV("SpecialFile1"), TCV("0")))
+					return false;
+
+				const StringView casFile(TCV("UbaClient1357/cas/4d/4d067153ac729a4a7e8220c97935ffba67487800"));
+
+				if (!SetupClientSession(logger, testRootDir, true, false, sessionServer.GetServer(), 1357, [&](LoggerWithWriter& logger, SessionClient& sessionClient)
+					{
+						if (!CreateTestFile(logger, testRootDir, casFile, TCV("0")))
+							return false;
+						return RunTestApp(logger, workingDir.data, runProcess, TC("-readwrite=0"));
+					}))
+					return false;
+
+				if (!DeleteTestFile(logger, testRootDir, casFile))
+					return false;
+
+				if (!SetupClientSession(logger, testRootDir, false, false, sessionServer.GetServer(), 1357, [&](LoggerWithWriter& logger, SessionClient& sessionClient)
+					{
+						if (!RunTestApp(logger, workingDir.data, runProcess, TC("-readwrite=1")))
+							return false;
+						return true;
+					}))
+					return false;
+
+
+				return true;
+			});
+	}
 }
diff --git a/Engine/Source/Programs/UnrealBuildAccelerator/TestApp/Private/Windows/UbaTestAppWin.cpp b/Engine/Source/Programs/UnrealBuildAccelerator/TestApp/Private/Windows/UbaTestAppWin.cpp
--- a/Engine/Source/Programs/UnrealBuildAccelerator/TestApp/Private/Windows/UbaTestAppWin.cpp
+++ b/Engine/Source/Programs/UnrealBuildAccelerator/TestApp/Private/Windows/UbaTestAppWin.cpp
@@ -17,6 +17,30 @@ int LogError(const wchar_t* format, ...)
 	return -1;
 }
 
+int ReadTestFile(void* outData, int capacity, const wchar_t* fileName)
+{
+	HANDLE fh = CreateFileW(fileName, GENERIC_READ, FILE_SHARE_READ, NULL, OPEN_EXISTING, 0, NULL);
+	if (fh == INVALID_HANDLE_VALUE)
+		return LogError(L"Failed to open %s for read", fileName);
+	DWORD bytesRead;
+	if (!ReadFile(fh, outData, capacity, &bytesRead, NULL))
+		return LogError(L"Failed to read from file %s", fileName);
+	CloseHandle(fh);
+	return bytesRead;
+}
+
+int WriteTestFile(const void* data, int size, const wchar_t* fileName)
+{
+	HANDLE fh = CreateFileW(fileName, GENERIC_WRITE, FILE_SHARE_WRITE, NULL, CREATE_ALWAYS, 0, NULL);
+	if (fh == INVALID_HANDLE_VALUE)
+		return LogError(L"Failed to open %s for write", fileName);
+	DWORD bytesWritten;
+	if (!WriteFile(fh, data, size, &bytesWritten, NULL) || size != int(bytesWritten))
+		return LogError(L"Failed to write %i bytes to file %s", size, fileName);
+	CloseHandle(fh);
+	return 0;
+}
+
 int wmain(int argc, wchar_t* argv[])
 {
 	HMODULE detoursHandle = GetModuleHandleW(L"UbaDetours.dll");
@@ -298,6 +322,23 @@ int wmain(int argc, wchar_t* argv[])
 		}
 		wprintf(L"%s\n", str);
 	}
+	else if (wcsncmp(argv[1], L"-readwrite=", 11) == 0)
+	{
+		int callIndex = _wtoi(argv[1] + 11);(void)callIndex;
+
+		wchar_t buf[32];
+		int read = ReadTestFile(buf, 32*2, L"SpecialFile1");
+		if (read == -1)
+			return -1;
+		buf[read*2] = 0;
+
+		int readIndex = _wtoi(buf);
+		if (callIndex != readIndex)
+			return LogError(L"callIndex different from readIndex");
+
+		buf[0] = wchar_t('1' + readIndex);
+		return WriteTestFile(buf, 2, L"SpecialFile1");
+	}
 	else
 	{
 		if (!detoursHandle)
