[UBA]
* Moved compression code out of Storage and into its own type inside UbaStorageUtils.cpp
* Changed Storage::m_workManager to be reference instead of pointer since it should never be optional

diff --git a/Engine/Source/Programs/UnrealBuildAccelerator/Agent/Private/UbaAgent.cpp b/Engine/Source/Programs/UnrealBuildAccelerator/Agent/Private/UbaAgent.cpp
--- a/Engine/Source/Programs/UnrealBuildAccelerator/Agent/Private/UbaAgent.cpp
+++ b/Engine/Source/Programs/UnrealBuildAccelerator/Agent/Private/UbaAgent.cpp
@@ -858,7 +858,8 @@ namespace uba
 		if (command.empty() && useStorage)
 		{
 			// Create a uba storage quickly just to fix non-graceful shutdowns
-			StorageCreateInfo info(g_rootDir.data, logWriter);
+			WorkManagerImpl workManager(maxProcessCount);
+			StorageCreateInfo info(g_rootDir.data, logWriter, workManager);
 			info.Apply(config);
 			info.rootDir = g_rootDir.data;
 			info.casCapacityBytes = storageCapacity;
@@ -1242,7 +1243,6 @@ namespace uba
 			storageInfo.storeCompressed = storeCompressed;
 			storageInfo.sendCompressed = sendCompressed;
 			storageInfo.resendCas = resendCas;
-			storageInfo.workManager = client;
 			storageInfo.getProxyBackendCallback = getProxyBackend;
 			storageInfo.getProxyBackendUserData = &backends;
 			storageInfo.allowProxy = allowProxy;
diff --git a/Engine/Source/Programs/UnrealBuildAccelerator/Cli/Private/UbaCli.cpp b/Engine/Source/Programs/UnrealBuildAccelerator/Cli/Private/UbaCli.cpp
--- a/Engine/Source/Programs/UnrealBuildAccelerator/Cli/Private/UbaCli.cpp
+++ b/Engine/Source/Programs/UnrealBuildAccelerator/Cli/Private/UbaCli.cpp
@@ -546,19 +546,21 @@ namespace uba
 		bool exit = false;
 		if (deleteCas)
 		{
-			StorageImpl(StorageCreateInfo(g_rootDir.data, logWriter)).DeleteAllCas();
+			WorkManagerImpl workManager(0);
+			StorageImpl(StorageCreateInfo(g_rootDir.data, logWriter, workManager)).DeleteAllCas();
 			for (u32 i=0; i!=agentCount; ++i)
 			{
 				StringBuffer<> clientRootDir;
 				clientRootDir.Append(g_rootDir).Append("Agent").AppendValue(i);
-				StorageImpl(StorageCreateInfo(clientRootDir.data, logWriter)).DeleteAllCas();
+				StorageImpl(StorageCreateInfo(clientRootDir.data, logWriter, workManager)).DeleteAllCas();
 			}
 			exit = true;
 		}
 
 		if (!addCas.empty())
 		{
-			StorageCreateInfo storageInfo(g_rootDir.data, logWriter);
+			WorkManagerImpl workManager(0);
+			StorageCreateInfo storageInfo(g_rootDir.data, logWriter, workManager);
 			storageInfo.casCapacityBytes = 0;
 			storageInfo.storeCompressed = storeCompressed;
 			StorageImpl storage(storageInfo);
@@ -570,7 +572,8 @@ namespace uba
 
 		if (checkCas)
 		{
-			StorageCreateInfo storageInfo(g_rootDir.data, logWriter);
+			WorkManagerImpl workManager(0);
+			StorageCreateInfo storageInfo(g_rootDir.data, logWriter, workManager);
 			storageInfo.casCapacityBytes = 0;
 			storageInfo.storeCompressed = storeCompressed;
 			StorageImpl storage(storageInfo);
@@ -581,7 +584,8 @@ namespace uba
 
 		if (!checkFileTable.empty())
 		{
-			StorageCreateInfo storageInfo(g_rootDir.data, logWriter);
+			WorkManagerImpl workManager(0);
+			StorageCreateInfo storageInfo(g_rootDir.data, logWriter, workManager);
 			storageInfo.casCapacityBytes = 0;
 			storageInfo.storeCompressed = storeCompressed;
 			StorageImpl storage(storageInfo);
@@ -679,10 +683,9 @@ namespace uba
 
 			testCompressOriginalSize = fileSize;
 
-			StorageCreateInfo storageInfo(g_rootDir.data, logWriter);
+			StorageCreateInfo storageInfo(g_rootDir.data, logWriter, workManager);
 			storageInfo.casCapacityBytes = 0;
 			storageInfo.storeCompressed = storeCompressed;
-			storageInfo.workManager = &workManager;
 			StorageImpl storage(storageInfo);
 
 			Storage::WriteResult res;
@@ -690,7 +693,7 @@ namespace uba
 
 			StringBuffer<> dest;
 			dest.Append(storage.GetTempPath()).Append(TCV("TestCompress.tmp"));
-			if (!storage.WriteCompressed(res, TC("MemoryMap"), InvalidFileHandle, mem, fileSize, dest.data, &header, sizeof(header), 0))
+			if (!storage.WriteCompressedFile(res, TC("MemoryMap"), InvalidFileHandle, mem, fileSize, dest.data, &header, sizeof(header), 0))
 				return false;
 			if (testDecompress.empty())
 				return true;
@@ -711,10 +714,9 @@ namespace uba
 			if (fileSize < 16)
 				return logger.Error(TC("File %s is too small to be compressed. Requires at least 16 bytes"), testDecompress.c_str());
 
-			StorageCreateInfo storageInfo(g_rootDir.data, logWriter);
+			StorageCreateInfo storageInfo(g_rootDir.data, logWriter, workManager);
 			storageInfo.casCapacityBytes = 0;
 			storageInfo.storeCompressed = storeCompressed;
-			storageInfo.workManager = &workManager;
 			StorageImpl storage(storageInfo);
 
 			BinaryReader reader(mem, 0, fileSize);
diff --git a/Engine/Source/Programs/UnrealBuildAccelerator/Cli/Private/UbaClient.h b/Engine/Source/Programs/UnrealBuildAccelerator/Cli/Private/UbaClient.h
--- a/Engine/Source/Programs/UnrealBuildAccelerator/Cli/Private/UbaClient.h
+++ b/Engine/Source/Programs/UnrealBuildAccelerator/Cli/Private/UbaClient.h
@@ -42,7 +42,6 @@ namespace uba
 				clientRootDir.Append("Agent").AppendValue(info.index);
 
 			StorageClientCreateInfo storageClientInfo(*networkClient, clientRootDir.data);
-			storageClientInfo.workManager = networkClient;
 			storageClientInfo.zone = info.zone;
 			storageClientInfo.getProxyBackendCallback = [](void* ud, const tchar* h) -> NetworkBackend& { return ((Client*)ud)->GetProxyBackend(h); };
 			storageClientInfo.getProxyBackendUserData = this;
diff --git a/Engine/Source/Programs/UnrealBuildAccelerator/Common/Private/UbaCacheClient.cpp b/Engine/Source/Programs/UnrealBuildAccelerator/Common/Private/UbaCacheClient.cpp
--- a/Engine/Source/Programs/UnrealBuildAccelerator/Common/Private/UbaCacheClient.cpp
+++ b/Engine/Source/Programs/UnrealBuildAccelerator/Common/Private/UbaCacheClient.cpp
@@ -799,7 +799,7 @@ namespace uba
 				if (!casKeyOffsets.empty())
 				{
 					Atomic<bool> fetchSuccess = true;
-					m_storage.m_workManager->ParallelFor(u32(casKeyOffsets.size() - 1), casKeyOffsets, [&](const WorkContext&, auto& it)
+					m_storage.m_workManager.ParallelFor(u32(casKeyOffsets.size() - 1), casKeyOffsets, [&](const WorkContext&, auto& it)
 						{
 							u32 casKeyOffset = *it;
 							if (!FetchCasTable(tws, bucket, cacheStats, casKeyOffset))
diff --git a/Engine/Source/Programs/UnrealBuildAccelerator/Common/Private/UbaSession.cpp b/Engine/Source/Programs/UnrealBuildAccelerator/Common/Private/UbaSession.cpp
--- a/Engine/Source/Programs/UnrealBuildAccelerator/Common/Private/UbaSession.cpp
+++ b/Engine/Source/Programs/UnrealBuildAccelerator/Common/Private/UbaSession.cpp
@@ -336,10 +336,10 @@ namespace uba
 
 			if (storeCompressed)
 			{
-				Storage::WriteResult res;
+				CompressWriteResult res;
 				CompressedFileHeader header { CalculateCasKey(mem, fileSize, true, &m_workManager, file.name.c_str()) };
 
-				if (!m_storage.WriteCompressed(res, TC("MemoryMap"), InvalidFileHandle, mem, fileSize, file.name.c_str(), &header, sizeof(header), file.lastWriteTime))
+				if (!m_storage.WriteCompressedFile(res, TC("MemoryMap"), InvalidFileHandle, mem, fileSize, file.name.c_str(), &header, sizeof(header), file.lastWriteTime))
 					return false;
 
 				shouldEvictFromMemory = false; // Can't evict without properly update filemappingtable.. the file on disk does now not match what was registered for write
diff --git a/Engine/Source/Programs/UnrealBuildAccelerator/Common/Private/UbaStorage.cpp b/Engine/Source/Programs/UnrealBuildAccelerator/Common/Private/UbaStorage.cpp
--- a/Engine/Source/Programs/UnrealBuildAccelerator/Common/Private/UbaStorage.cpp
+++ b/Engine/Source/Programs/UnrealBuildAccelerator/Common/Private/UbaStorage.cpp
@@ -71,8 +71,9 @@ namespace uba
 #endif
 	}
 
-	StorageCreateInfo::StorageCreateInfo(const tchar* rootDir_, LogWriter& w)
+	StorageCreateInfo::StorageCreateInfo(const tchar* rootDir_, LogWriter& w, WorkManager& wm)
 	:	writer(w)
+	,	workManager(wm)
 	,	rootDir(rootDir_)
 	,	casCompressor(DefaultCompressor)
 	,	casCompressionLevel(DefaultCompressionLevel)
@@ -323,257 +324,13 @@ namespace uba
 		m_externalFileMappingsProvider = std::move(provider);
 	}
 
-	bool StorageImpl::WriteCompressed(WriteResult& out, const tchar* from, FileHandle readHandle, u8* readMem, u64 fileSize, const tchar* to, const void* header, u64 headerSize, u64 lastWriteTime)
+	bool StorageImpl::WriteCompressedFile(CompressWriteResult& out, const tchar* from, FileHandle readHandle, u8* readMem, u64 fileSize, const tchar* to, const void* header, u64 headerSize, u64 lastWriteTime)
 	{
-		StorageStats& stats = Stats();
-
-		u64 totalWritten = 0;
-
-		u64 diff = (u64)OodleLZ_GetCompressedBufferSizeNeeded((OodleLZ_Compressor)m_casCompressor, BufferSlotHalfSize) - BufferSlotHalfSize;
-		u64 maxUncompressedBlock = BufferSlotHalfSize - diff - 8; // 8 bytes for the little header
-		u32 workCount = u32((fileSize + maxUncompressedBlock - 1) / maxUncompressedBlock);
-
-		FileAccessor destinationFile(m_logger, to);
-		if (!destinationFile.CreateWrite(false, DefaultAttributes(), 0, m_tempPath.data))
-			return false;
-		if (headerSize)
-			if (!destinationFile.Write(header, headerSize))
-				return false;
-		if (!destinationFile.Write(&fileSize, sizeof(u64))) // Store file size first in compressed file
-			return false;
-
-		totalWritten += sizeof(u64);
-
-		u64 left = fileSize;
-
-		if (m_workManager && workCount > 1)
-		{
-			if (!readMem)
-			{
-				FileMappingHandle fileMapping = uba::CreateFileMappingW(m_logger, readHandle, PAGE_READONLY, fileSize, from);
-				if (!fileMapping.IsValid())
-					return m_logger.Error(TC("Failed to create file mapping for %s (%s)"), from, LastErrorToText().data);
-
-				auto fmg = MakeGuard([&]() { CloseFileMapping(m_logger, fileMapping, from); });
-				u8* uncompressedData = MapViewOfFile(m_logger, fileMapping, FILE_MAP_READ, 0, fileSize);
-				if (!uncompressedData)
-					return m_logger.Error(TC("Failed to map view of file mapping for %s (%s)"), from, LastErrorToText().data);
-
-				auto udg = MakeGuard([&]()
-					{
-						if (m_asyncUnmapViewOfFile)
-							m_workManager->AddWork([this, uncompressedData, fileSize, f = TString(from)](const WorkContext&) { UnmapViewOfFile(m_logger, uncompressedData, fileSize, f.c_str()); }, 1, TC("UnmapFile"));
-						else
-							UnmapViewOfFile(m_logger, uncompressedData, fileSize, from);
-					});
-
-				if (!WriteMemToCompressedFile(destinationFile, workCount, uncompressedData, fileSize, maxUncompressedBlock, totalWritten))
-					return false;
-			}
-			else
-			{
-				if (!WriteMemToCompressedFile(destinationFile, workCount, readMem, fileSize, maxUncompressedBlock, totalWritten))
-					return false;
-			}
-		}
-		else
-		{
-			u8* slot = m_bufferSlots.Pop();
-			auto _ = MakeGuard([&](){ m_bufferSlots.Push(slot); });
-			u8* uncompressedData = slot;
-			u8* compressBuffer = slot + BufferSlotHalfSize;
-
-			auto& memoryCompressTime = KernelStats::GetCurrent().memoryCompress;
-
-			TimerScope cts(stats.compressWrite);
-			while (left)
-			{
-				u64 uncompressedBlockSize = Min(left, maxUncompressedBlock);
-				
-				void* scratchMem = nullptr;
-				u64 scratchSize = 0;
-
-				if (readMem)
-				{
-					scratchMem = uncompressedData;
-					scratchSize = BufferSlotHalfSize;
-					uncompressedData = readMem + fileSize - left;
-				}
-				else
-				{
-					if (!ReadFile(m_logger, from, readHandle, uncompressedData, uncompressedBlockSize))
-						return false;
-					scratchMem = uncompressedData + uncompressedBlockSize;
-					scratchSize = BufferSlotHalfSize - uncompressedBlockSize;
-				}
-				u8* destBuf = compressBuffer;
-				OO_SINTa compressedBlockSize;
-				{
-					TimerScope kts(memoryCompressTime);
-					compressedBlockSize = OodleLZ_Compress((OodleLZ_Compressor)m_casCompressor, uncompressedData, (OO_SINTa)uncompressedBlockSize, destBuf + 8, (OodleLZ_CompressionLevel)m_casCompressionLevel, NULL, NULL, NULL, scratchMem, scratchSize);
-					if (compressedBlockSize == OODLELZ_FAILED)
-						return m_logger.Error(TC("Failed to compress %llu bytes for %s"), uncompressedBlockSize, from);
-					memoryCompressTime.bytes += compressedBlockSize;
-				}
-
-				*(u32*)destBuf =  u32(compressedBlockSize);
-				*(u32*)(destBuf+4) =  u32(uncompressedBlockSize);
-
-				u32 writeBytes = u32(compressedBlockSize) + 8;
-				if (!destinationFile.Write(destBuf, writeBytes))
-					return false;
-
-				totalWritten += writeBytes;
-
-				left -= uncompressedBlockSize;
-			}
-		}
-
-		if (lastWriteTime)
-			if (!SetFileLastWriteTime(destinationFile.GetHandle(), lastWriteTime))
-				return m_logger.Error(TC("Failed to set file time on filehandle for %s"), to);
-		if (!destinationFile.Close())
-			return false;
-
-		stats.createCasBytesRaw += fileSize;
-		stats.createCasBytesComp += totalWritten;
-
-		out.size = totalWritten;
-		return true;
+		CompressWriter w(m_logger, m_bufferSlots, m_workManager, Stats(), m_casCompressor, m_casCompressionLevel, m_asyncUnmapViewOfFile);
+		return w.CompressToFile(out, from, readHandle, readMem, fileSize, to, header, headerSize, lastWriteTime, m_tempPath.data);
 	}
 
-	bool StorageImpl::WriteMemToCompressedFile(FileAccessor& destination, u32 workCount, const u8* uncompressedData, u64 fileSize, u64 maxUncompressedBlock, u64& totalWritten)
-	{
-		#ifndef __clang_analyzer__
-
-		struct WorkRec
-		{
-			WorkRec() = delete;
-			WorkRec(const WorkRec&) = delete;
-			WorkRec(u32 wc)
-			{
-				workCount = wc;
-
-				// For some very unknown reason ASAN on linux triggers on the "new[]" call when doing delete[]
-				// while doing it manually works properly. Will stop investigating this and move on
-				#if PLATFORM_WINDOWS
-				events = new Event[workCount];
-				#else
-				events = (Event*)aligned_alloc(alignof(Event), sizeof(Event)*workCount);
-				for (auto i=0;i!=workCount; ++i)
-					new (events + i) Event();
-				#endif
-			}
-			~WorkRec()
-			{
-				#if PLATFORM_WINDOWS
-				delete[] events;
-				#else
-				for (auto i=0;i!=workCount; ++i)
-					events[i].~Event();
-				aligned_free(events);
-				#endif
-			}
-			Atomic<u64> refCount;
-			Atomic<u64> compressCounter;
-			Event* events;
-			const u8* uncompressedData = nullptr;
-			FileAccessor* destination = nullptr;
-			u64 written = 0;
-			u64 workCount = 0;
-			u64 maxUncompressedBlock = 0;
-			u64 fileSize = 0;
-			bool error = false;
-		};
-
-		WorkRec* rec = new WorkRec(workCount);
-		rec->uncompressedData = uncompressedData;
-		rec->maxUncompressedBlock = maxUncompressedBlock;
-		rec->fileSize = fileSize;
-		rec->destination = &destination;
-
-		for (u32 i=0; i!=workCount; ++i)
-			rec->events[i].Create(true);
-
-		StorageStats& stats = Stats();
-		TimerScope cts(stats.compressWrite);
-
-		auto& kernelStats = KernelStats::GetCurrent();
-
-		auto work = [this, rec, &stats, &kernelStats](const WorkContext&)
-		{
-			KernelStatsScope kss(kernelStats);
-
-			u8* slot = m_bufferSlots.Pop();
-			auto _ = MakeGuard([&](){ m_bufferSlots.Push(slot); });
-			u8* compressSlotBuffer = slot + BufferSlotHalfSize;
-			while (true)
-			{
-				u64 index = rec->compressCounter++;
-				if (index >= rec->workCount)
-				{
-					if (!--rec->refCount)
-						delete rec;
-					return;
-				}
-				u64 startOffset = rec->maxUncompressedBlock*index;
-				const u8* uncompressedDataSlot = rec->uncompressedData + startOffset;
-				OO_SINTa uncompressedBlockSize = (OO_SINTa)Min(rec->maxUncompressedBlock, rec->fileSize - startOffset);
-				OO_SINTa compressedBlockSize;
-				{
-					void* scratchMem = slot;
-					u64 scratchSize = BufferSlotHalfSize;
-					TimerScope kts(kernelStats.memoryCompress);
-					compressedBlockSize = OodleLZ_Compress((OodleLZ_Compressor)m_casCompressor, uncompressedDataSlot, uncompressedBlockSize, compressSlotBuffer + 8, (OodleLZ_CompressionLevel)m_casCompressionLevel, NULL, NULL, NULL, scratchMem, scratchSize);
-					if (compressedBlockSize == OODLELZ_FAILED)
-					{
-						m_logger.Error(TC("Failed to compress %llu bytes for %s"), u64(uncompressedBlockSize), rec->destination->GetFileName());
-						rec->error = true;
-						return;
-					}
-					kernelStats.memoryCompress.bytes += compressedBlockSize;
-				}
-				*(u32*)compressSlotBuffer =  u32(compressedBlockSize);
-				*(u32*)(compressSlotBuffer+4) =  u32(uncompressedBlockSize);
-
-				if (index)
-					rec->events[index-1].IsSet();
-
-				u32 writeBytes = u32(compressedBlockSize) + 8;
-
-				if (!rec->destination->Write(compressSlotBuffer, writeBytes))
-					rec->error = true;
-
-				rec->written += writeBytes;
-				if (index < rec->workCount)
-					rec->events[index].Set();
-			}
-		};
-
-		u32 workerCount = Min(workCount, m_workManager->GetWorkerCount());
-		workerCount = Min(workerCount, MaxWorkItemsPerAction);
-
-		rec->refCount = workerCount + 1; // We need to keep refcount up 1 to make sure it is not deleted before we read rec->written
-		m_workManager->AddWork(work, workerCount-1, TC("Compress")); // We are a worker ourselves
-		{
-			TrackWorkScope tws;
-			work({tws});
-		}
-		rec->events[rec->workCount - 1].IsSet();
-
-		totalWritten += rec->written;
-		bool error = rec->error;
-
-		if (!--rec->refCount)
-			delete rec;
-
-		if (error)
-			return false;
-		#endif // __clang_analyzer__
-		return true;
-	}
-
-	bool StorageImpl::WriteCasFileNoCheck(WriteResult& out, const StringKey& fileNameKey, const tchar* fileName, const CasKey& casKey, const tchar* casFile, bool storeCompressed)
+	bool StorageImpl::WriteCasFileNoCheck(CompressWriteResult& out, const StringKey& fileNameKey, const tchar* fileName, const CasKey& casKey, const tchar* casFile, bool storeCompressed)
 	{
 		StorageStats& stats = Stats();
 		TimerScope ts(stats.createCas);
@@ -588,7 +345,7 @@ namespace uba
 				if (!fileMem)
 					return false;
 				out.readSize = externalMapping.writtenSize;
-				return StorageImpl::WriteCompressed(out, fileName, InvalidFileHandle, fileMem, externalMapping.size, casFile, nullptr, 0);
+				return WriteCompressedFile(out, fileName, InvalidFileHandle, fileMem, externalMapping.size, casFile, nullptr, 0);
 			}
 		}
 
@@ -623,7 +380,7 @@ namespace uba
 		}
 
 		if (!readIsCompressed && storeCompressed)
-			return WriteCompressed(out, fileName, readHandle, 0, readFileSize, casFile, nullptr, 0);
+			return WriteCompressedFile(out, fileName, readHandle, 0, readFileSize, casFile, nullptr, 0);
 
 		if (readIsCompressed && !storeCompressed)
 			return m_logger.Error(TC("[WriteCasFileNoCheck] Writing compressed file to non-compressed storage not supported (%s)"), fileName);
@@ -654,7 +411,7 @@ namespace uba
 		return true;
 	}
 
-	bool StorageImpl::WriteCasFile(WriteResult& out, const tchar* fileName, const CasKey& casKey)
+	bool StorageImpl::WriteCasFile(CompressWriteResult& out, const tchar* fileName, const CasKey& casKey)
 	{
 		UBA_ASSERT(IsCompressed(casKey) == m_storeCompressed);
 		StringBuffer<> casFile;
@@ -777,14 +534,11 @@ namespace uba
 					}
 				};
 
-			if (m_workManager)
-			{
-				u32 workCount = u32(decompressedSize / BufferSlotSize) + 1;
-				u32 workerCount = Min(workCount, m_workManager->GetWorkerCount() - 1); // We are a worker ourselves
-				workerCount = Min(workerCount, MaxWorkItemsPerAction); // Cap this to not starve other things
-				rec->refCount += workerCount;
-				m_workManager->AddWork(work, workerCount, TC("DecompressMemToMem"));
-			}
+			u32 workCount = u32(decompressedSize / BufferSlotSize) + 1;
+			u32 workerCount = Min(workCount, m_workManager.GetWorkerCount() - 1); // We are a worker ourselves
+			workerCount = Min(workerCount, MaxWorkItemsPerAction); // Cap this to not starve other things
+			rec->refCount += workerCount;
+			m_workManager.AddWork(work, workerCount, TC("DecompressMemToMem"));
 
 			TimerScope ts(stats.decompressToMem);
 			{
@@ -953,7 +707,7 @@ namespace uba
 
 		if (casEntry.beingWritten)
 		{
-			m_logger.Warning(TC("Cas %s is being written from network while being added by host (%s)"), CasKeyString(casKey).str, fileName);
+			m_logger.Info(TC("Cas %s is being written from network while being added by host (%s)"), CasKeyString(casKey).str, fileName);
 			int waitCount = 0;
 			while (casEntry.beingWritten)
 			{
@@ -1012,7 +766,7 @@ namespace uba
 		casEntry.verified = true;
 		casEntry.exists = false;
 
-		WriteResult res;
+		CompressWriteResult res;
 		if (!WriteCasFileNoCheck(res, fileNameKey, fileName, casKey, casFile.data, IsCompressed(casKey)))
 			return false;
 
@@ -1048,9 +802,9 @@ namespace uba
 				fullPath.EnsureEndsWithSlash().Append(e.name);
 				++workLeft;
 
-				if (allowParallel && m_workManager)
+				if (allowParallel)
 				{
-					m_workManager->AddWork([&traverseCasFileDir, p = fullPath.ToString()](const WorkContext&) { traverseCasFileDir(p); }, 1, TC("TraverseCasFiles"));
+					m_workManager.AddWork([&traverseCasFileDir, p = fullPath.ToString()](const WorkContext&) { traverseCasFileDir(p); }, 1, TC("TraverseCasFiles"));
 				}
 				else
 				{
@@ -1059,7 +813,7 @@ namespace uba
 			});
 
 		while (workLeft)
-			m_workManager->DoWork();
+			m_workManager.DoWork();
 	}
 
 	void StorageImpl::TraverseAllCasFiles(const Function<void(const CasKey& key, u64 size)>& func, bool allowParallel)
@@ -1275,12 +1029,12 @@ namespace uba
 	:	m_workManager(info.workManager)
 	,	m_logger(info.writer, logPrefix)
 	,	m_maxParallelCopyOrLinkBottleneck(info.maxParallelCopyOrLink)
-	,	m_casDataBuffer(m_logger, info.workManager)
+	,	m_casDataBuffer(m_logger, &info.workManager)
 	{
 		m_casCapacityBytes = info.casCapacityBytes;
 		m_storeCompressed = info.storeCompressed;
 		m_manuallyHandleOverflow = info.manuallyHandleOverflow;
-		m_asyncUnmapViewOfFile = info.asyncUnmapViewOfFile && m_workManager;
+		m_asyncUnmapViewOfFile = info.asyncUnmapViewOfFile;
 		m_allowDeleteVerified = info.allowDeleteVerified;
 		m_writeToDisk = info.writeToDisk;
 		if (!m_writeToDisk)
@@ -2228,7 +1982,7 @@ namespace uba
 
 		casEntry.verified = true;
 		deferredLock.Leave();
-		WriteResult res;
+		CompressWriteResult res;
 		if (!WriteCasFile(res, deferredCreation.data, casKey))
 			return m_logger.Error(TC("Failed to write deferred cas %s (%s)"), CasKeyString(casKey).str, deferredCreation.data);
 
@@ -2280,7 +2034,7 @@ namespace uba
 		casEntry.verified = true;
 		if (!fileName)
 			return false;
-		WriteResult res;
+		CompressWriteResult res;
 		if (!WriteCasFile(res, fileName, casKey))
 			return false;
 		casEntry.exists = true;
@@ -2784,7 +2538,7 @@ namespace uba
 	{
 		StorageStats& stats = Stats();
 		TimerScope ts(stats.calculateCasKey);
-		return uba::CalculateCasKey(fileMem, fileSize, storeCompressed, m_workManager, nullptr);
+		return uba::CalculateCasKey(fileMem, fileSize, storeCompressed, &m_workManager, nullptr);
 	}
 
 
@@ -2810,12 +2564,12 @@ namespace uba
 			}
 			auto udg = MakeGuard([&]() { 
 					if (m_asyncUnmapViewOfFile)
-						m_workManager->AddWork([this, fileData, fileSize, fn = TString(fileName)](const WorkContext&) { UnmapViewOfFile(m_logger, fileData, fileSize, fn.c_str()); }, 1, TC("UnmapFile"));
+						m_workManager.AddWork([this, fileData, fileSize, fn = TString(fileName)](const WorkContext&) { UnmapViewOfFile(m_logger, fileData, fileSize, fn.c_str()); }, 1, TC("UnmapFile"));
 					else
 						UnmapViewOfFile(m_logger, fileData, fileSize, fileName);
 				});
 
-			return uba::CalculateCasKey(fileData, fileSize, storeCompressed, m_workManager, fileName);
+			return uba::CalculateCasKey(fileData, fileSize, storeCompressed, &m_workManager, fileName);
 		}
 
 		CasKeyHasher hasher;
@@ -2836,7 +2590,7 @@ namespace uba
 
 	bool StorageImpl::DecompressFileToMemory(const tchar* fileName, FileHandle fileHandle, u8* dest, u64 decompressedSize, const tchar* writeHint, u64 fileStartOffset)
 	{
-		if (m_workManager && decompressedSize > BufferSlotSize*4) // Arbitrary size threshold. We want to at least catch the pch here
+		if (decompressedSize > BufferSlotSize*4) // Arbitrary size threshold. We want to at least catch the pch here
 		{
 			u64 compressedSize;
 			if (!uba::GetFileSizeEx(compressedSize, fileHandle))
diff --git a/Engine/Source/Programs/UnrealBuildAccelerator/Common/Private/UbaStorageUtils.cpp b/Engine/Source/Programs/UnrealBuildAccelerator/Common/Private/UbaStorageUtils.cpp
--- a/Engine/Source/Programs/UnrealBuildAccelerator/Common/Private/UbaStorageUtils.cpp
+++ b/Engine/Source/Programs/UnrealBuildAccelerator/Common/Private/UbaStorageUtils.cpp
@@ -808,5 +808,311 @@ namespace uba
 		return true;
 	}
 
+	CompressWriter::CompressWriter(Logger& l, BufferSlots& bs, WorkManager& wm, StorageStats& s, u8 cc, u8 ccl, bool avof)
+	:	m_logger(l)
+	,	m_bufferSlots(bs)
+	,	m_workManager(wm)
+	,	m_stats(s)
+	,	m_casCompressor(cc)
+	,	m_casCompressionLevel(ccl)
+	,	m_asyncUnmapViewOfFile(avof)
+	{
+	}
+
+	bool CompressWriter::CompressToFile(CompressWriteResult& out, const tchar* from, FileHandle readHandle, u8* readMem, u64 fileSize, const tchar* toFile, const void* header, u64 headerSize, u64 lastWriteTime, const tchar* tempPath)
+	{
+		FileAccessor destinationFile(m_logger, toFile);
+		if (!destinationFile.CreateWrite(false, DefaultAttributes(), 0, tempPath))
+			return false;
+		if (headerSize)
+			if (!destinationFile.Write(header, headerSize))
+				return false;
+
+		LinearWriterFile destination(destinationFile);
+		if (!CompressFromMemOrFile(out, destination, from, readHandle, readMem, fileSize))
+			return false;
+		if (lastWriteTime)
+			if (!SetFileLastWriteTime(destinationFile.GetHandle(), lastWriteTime))
+				return m_logger.Error(TC("Failed to set file time on filehandle for %s"), toFile);
+		if (!destinationFile.Close())
+			return false;
+		return true;
+	}
+
+	bool CompressWriter::CompressedToMapping(CompressWriteResult& out, u8* readMem, u64 fileSize, const tchar* hint)
+	{
+		u64 mappingSize = fileSize + 1024;
+		FileMappingHandle destMapping = CreateMemoryMappingW(m_logger, PAGE_READWRITE, mappingSize, nullptr, hint);
+		if (!destMapping.IsValid())
+			return false;
+		auto destGuard = MakeGuard([&]() { CloseFileMapping(m_logger, destMapping, hint); });
+
+		u8* compressedData = MapViewOfFile(m_logger, destMapping, FILE_MAP_WRITE, 0, mappingSize);
+		if (!compressedData)
+			return false;
+		auto unmapGuard = MakeGuard([&]() { UnmapViewOfFile(m_logger, compressedData, mappingSize, hint); });
+
+		LinearWriterMem destination(compressedData, mappingSize);
+		if (!CompressFromMemOrFile(out, destination, hint, InvalidFileHandle, readMem, fileSize))
+			return false;
+		destGuard.Cancel();
+		out.mappingHandle = destMapping;
+		out.offset = 0;
+		return true;
+	}
+
+	bool CompressWriter::CompressFromMem(LinearWriter& destination, u32 workCount, const u8* uncompressedData, u64 fileSize, u64 maxUncompressedBlock, u64& totalWritten)
+	{
+		constexpr u32 MaxWorkItemsPerAction = 128; // Cap this to not starve other things
+
+		#ifndef __clang_analyzer__
+
+		struct WorkRec
+		{
+			WorkRec() = delete;
+			WorkRec(const WorkRec&) = delete;
+			WorkRec(u32 wc)
+			{
+				workCount = wc;
+
+				// For some very unknown reason ASAN on linux triggers on the "new[]" call when doing delete[]
+				// while doing it manually works properly. Will stop investigating this and move on
+				#if PLATFORM_WINDOWS
+				events = new Event[workCount];
+				#else
+				events = (Event*)aligned_alloc(alignof(Event), sizeof(Event)*workCount);
+				for (auto i=0;i!=workCount; ++i)
+					new (events + i) Event();
+				#endif
+			}
+			~WorkRec()
+			{
+				#if PLATFORM_WINDOWS
+				delete[] events;
+				#else
+				for (auto i=0;i!=workCount; ++i)
+					events[i].~Event();
+				aligned_free(events);
+				#endif
+			}
+			LinearWriter* destination;
+			Atomic<u64> refCount;
+			Atomic<u64> compressCounter;
+			Event* events;
+			const u8* uncompressedData = nullptr;
+			u64 written = 0;
+			u64 workCount = 0;
+			u64 maxUncompressedBlock = 0;
+			u64 fileSize = 0;
+			bool error = false;
+		};
+
+		WorkRec* recPtr = new WorkRec(workCount);
+		WorkRec& rec = *recPtr;
+		rec.destination = &destination;
+		rec.uncompressedData = uncompressedData;
+		rec.maxUncompressedBlock = maxUncompressedBlock;
+		rec.fileSize = fileSize;
+
+		for (u32 i=0; i!=workCount; ++i)
+			rec.events[i].Create(true);
+
+		TimerScope cts(m_stats.compressWrite);
+
+		auto& kernelStats = KernelStats::GetCurrent();
+
+		auto work = [recPtr, this, &kernelStats](const WorkContext&)
+		{
+			auto& rec = *recPtr;
+			KernelStatsScope kss(kernelStats);
+
+			u8* slot = m_bufferSlots.Pop();
+			auto _ = MakeGuard([&](){ m_bufferSlots.Push(slot); });
+			u8* compressSlotBuffer = slot + BufferSlotHalfSize;
+			while (true)
+			{
+				u64 index = rec.compressCounter++;
+				if (index >= rec.workCount)
+				{
+					if (!--rec.refCount)
+						delete recPtr;
+					return;
+				}
+				u64 startOffset = rec.maxUncompressedBlock*index;
+				const u8* uncompressedDataSlot = rec.uncompressedData + startOffset;
+				OO_SINTa uncompressedBlockSize = (OO_SINTa)Min(rec.maxUncompressedBlock, rec.fileSize - startOffset);
+				OO_SINTa compressedBlockSize;
+				{
+					void* scratchMem = slot;
+					u64 scratchSize = BufferSlotHalfSize;
+					TimerScope kts(kernelStats.memoryCompress);
+					compressedBlockSize = OodleLZ_Compress((OodleLZ_Compressor)m_casCompressor, uncompressedDataSlot, uncompressedBlockSize, compressSlotBuffer + 8, (OodleLZ_CompressionLevel)m_casCompressionLevel, NULL, NULL, NULL, scratchMem, scratchSize);
+					if (compressedBlockSize == OODLELZ_FAILED)
+					{
+						m_logger.Error(TC("Failed to compress %llu bytes for %s"), u64(uncompressedBlockSize), rec.destination->GetHint());
+						rec.error = true;
+						return;
+					}
+					kernelStats.memoryCompress.bytes += compressedBlockSize;
+				}
+				*(u32*)compressSlotBuffer =  u32(compressedBlockSize);
+				*(u32*)(compressSlotBuffer+4) =  u32(uncompressedBlockSize);
+
+				if (index)
+					rec.events[index-1].IsSet();
+
+				u32 writeBytes = u32(compressedBlockSize) + 8;
+
+				if (!rec.destination->Write(compressSlotBuffer, writeBytes))
+					rec.error = true;
+
+				rec.written += writeBytes;
+				if (index < rec.workCount)
+					rec.events[index].Set();
+			}
+		};
+
+		u32 workerCount = Min(workCount, m_workManager.GetWorkerCount());
+		workerCount = Min(workerCount, MaxWorkItemsPerAction);
+
+		rec.refCount = workerCount + 1; // We need to keep refcount up 1 to make sure it is not deleted before we read rec->written
+		m_workManager.AddWork(work, workerCount-1, TC("Compress")); // We are a worker ourselves
+		{
+			TrackWorkScope tws;
+			work({tws});
+		}
+		rec.events[rec.workCount - 1].IsSet();
+
+		totalWritten += rec.written;
+		bool error = rec.error;
+
+		if (!--rec.refCount)
+			delete recPtr;
+
+		if (error)
+			return false;
+		#endif // __clang_analyzer__
+		return true;
+	}
+
+	bool CompressWriter::CompressFromMemOrFile(CompressWriteResult& result, LinearWriter& destination, const tchar* from, FileHandle readHandle, u8* readMem, u64 fileSize)
+	{
+		if (!destination.Write(&fileSize, sizeof(u64))) // Store file size first in compressed file
+			return false;
+		u64 totalWritten = sizeof(u64);
+
+		StorageStats& stats = m_stats;
+
+		u64 diff = (u64)OodleLZ_GetCompressedBufferSizeNeeded((OodleLZ_Compressor)m_casCompressor, BufferSlotHalfSize) - BufferSlotHalfSize;
+		u64 maxUncompressedBlock = BufferSlotHalfSize - diff - 8; // 8 bytes for the little header
+		u32 workCount = u32((fileSize + maxUncompressedBlock - 1) / maxUncompressedBlock);
+
+		u64 left = fileSize;
+
+		if (workCount > 1)
+		{
+			if (!readMem)
+			{
+				FileMappingHandle fileMapping = uba::CreateFileMappingW(m_logger, readHandle, PAGE_READONLY, fileSize, from);
+				if (!fileMapping.IsValid())
+					return m_logger.Error(TC("Failed to create file mapping for %s (%s)"), from, LastErrorToText().data);
+
+				auto fmg = MakeGuard([&]() { CloseFileMapping(m_logger, fileMapping, from); });
+				u8* uncompressedData = MapViewOfFile(m_logger, fileMapping, FILE_MAP_READ, 0, fileSize);
+				if (!uncompressedData)
+					return m_logger.Error(TC("Failed to map view of file mapping for %s (%s)"), from, LastErrorToText().data);
+
+				auto udg = MakeGuard([&]()
+					{
+						if (m_asyncUnmapViewOfFile)
+							m_workManager.AddWork([this, uncompressedData, fileSize, f = TString(from)](const WorkContext&) { UnmapViewOfFile(m_logger, uncompressedData, fileSize, f.c_str()); }, 1, TC("UnmapFile"));
+						else
+							UnmapViewOfFile(m_logger, uncompressedData, fileSize, from);
+					});
+
+				if (!CompressFromMem(destination, workCount, uncompressedData, fileSize, maxUncompressedBlock, totalWritten))
+					return false;
+			}
+			else
+			{
+				if (!CompressFromMem(destination, workCount, readMem, fileSize, maxUncompressedBlock, totalWritten))
+					return false;
+			}
+		}
+		else
+		{
+			u8* slot = m_bufferSlots.Pop();
+			auto _ = MakeGuard([&](){ m_bufferSlots.Push(slot); });
+			u8* uncompressedData = slot;
+			u8* compressBuffer = slot + BufferSlotHalfSize;
+
+			auto& memoryCompressTime = KernelStats::GetCurrent().memoryCompress;
+
+			TimerScope cts(stats.compressWrite);
+			while (left)
+			{
+				u64 uncompressedBlockSize = Min(left, maxUncompressedBlock);
+				
+				void* scratchMem = nullptr;
+				u64 scratchSize = 0;
+
+				if (readMem)
+				{
+					scratchMem = uncompressedData;
+					scratchSize = BufferSlotHalfSize;
+					uncompressedData = readMem + fileSize - left;
+				}
+				else
+				{
+					if (!ReadFile(m_logger, from, readHandle, uncompressedData, uncompressedBlockSize))
+						return false;
+					scratchMem = uncompressedData + uncompressedBlockSize;
+					scratchSize = BufferSlotHalfSize - uncompressedBlockSize;
+				}
+				u8* destBuf = compressBuffer;
+				OO_SINTa compressedBlockSize;
+				{
+					TimerScope kts(memoryCompressTime);
+					compressedBlockSize = OodleLZ_Compress((OodleLZ_Compressor)m_casCompressor, uncompressedData, (OO_SINTa)uncompressedBlockSize, destBuf + 8, (OodleLZ_CompressionLevel)m_casCompressionLevel, NULL, NULL, NULL, scratchMem, scratchSize);
+					if (compressedBlockSize == OODLELZ_FAILED)
+						return m_logger.Error(TC("Failed to compress %llu bytes for %s"), uncompressedBlockSize, from);
+					memoryCompressTime.bytes += compressedBlockSize;
+				}
+
+				*(u32*)destBuf =  u32(compressedBlockSize);
+				*(u32*)(destBuf+4) =  u32(uncompressedBlockSize);
+
+				u32 writeBytes = u32(compressedBlockSize) + 8;
+				if (!destination.Write(destBuf, writeBytes))
+					return false;
 
+				totalWritten += writeBytes;
+
+				left -= uncompressedBlockSize;
+			}
+		}
+
+		stats.createCasBytesRaw += fileSize;
+		stats.createCasBytesComp += totalWritten;
+
+		result.size = totalWritten;
+		return true;
+	}
+
+	LinearWriterFile::LinearWriterFile(FileAccessor& f) : file(f) {}
+	bool LinearWriterFile::Write(const void* data, u64 dataLen) { return file.Write(data, dataLen); }
+	const tchar* LinearWriterFile::GetHint() { return file.GetFileName(); }
+
+	LinearWriterMem::LinearWriterMem(u8* d, u64 c) : pos(d), size(0), capacity(c) {}
+	bool LinearWriterMem::Write(const void* data, u64 dataLen)
+	{
+		u64 newSize = size + dataLen;
+		if (newSize > capacity)
+			return false;
+		memcpy(pos, data, dataLen);
+		pos += dataLen;
+		size = newSize;
+		return true;
+	}
+	const tchar* LinearWriterMem::GetHint() { return TC(""); }
 }
diff --git a/Engine/Source/Programs/UnrealBuildAccelerator/Common/Public/UbaStorage.h b/Engine/Source/Programs/UnrealBuildAccelerator/Common/Public/UbaStorage.h
--- a/Engine/Source/Programs/UnrealBuildAccelerator/Common/Public/UbaStorage.h
+++ b/Engine/Source/Programs/UnrealBuildAccelerator/Common/Public/UbaStorage.h
@@ -16,6 +16,7 @@ namespace uba
 	class FileAccessor;
 	class Trace;
 	class WorkManager;
+	struct CompressWriteResult;
 	struct DirectoryEntry;
 	extern CasKey EmptyFileKey;
 	
@@ -68,8 +69,7 @@ namespace uba
 		virtual void SetTrace(Trace* trace, bool detailed) {}
 		virtual void Ping() {}
 
-		struct WriteResult { FileMappingHandle mappingHandle; u64 readSize = InvalidValue; u64 size = InvalidValue; u64 offset = InvalidValue; };
-		virtual bool WriteCompressed(WriteResult& out, const tchar* from, FileHandle readHandle, u8* readMem, u64 fileSize, const tchar* toFile, const void* header, u64 headerSize, u64 lastWriteTime = 0) = 0;
+		virtual bool WriteCompressedFile(CompressWriteResult& out, const tchar* from, FileHandle readHandle, u8* readMem, u64 fileSize, const tchar* toFile, const void* header, u64 headerSize, u64 lastWriteTime = 0) = 0;
 
 		struct ExternalFileMapping { FileMappingHandle handle; u64 offset = 0; u64 size = 0; u64 lastWriteTime = 0; u64 writtenSize = InvalidValue; };
 		using ExternalFileMappingsProvider = Function<bool(ExternalFileMapping& out, StringKey fileNameKey, const tchar* fileName)>;
@@ -78,11 +78,12 @@ namespace uba
 
 	struct StorageCreateInfo
 	{
-		StorageCreateInfo(const tchar* rootDir_, LogWriter& w);
+		StorageCreateInfo(const tchar* rootDir_, LogWriter& w, WorkManager& wm);
 
 		void Apply(const Config& config);
 
 		LogWriter& writer;
+		WorkManager& workManager;
 		const tchar* rootDir;
 		u64 casCapacityBytes = 20llu * 1024 * 1024 * 1024;
 		u32 maxParallelCopyOrLink = 1000;
@@ -92,7 +93,6 @@ namespace uba
 		bool writeToDisk = true;
 		bool keepTransientDataMapped = true;	// Will keep the transient data mapped in instead of doing map/unmap everytime we access it
 		bool allowDeleteVerified = false;
-		WorkManager* workManager = nullptr;
 		MutexHandle exclusiveMutex = InvalidMutexHandle; // Set this to hand over pre-exclusive access to storage. Storage will release mutex on shutdown
 		u8 casCompressor = 0;
 		u8 casCompressionLevel = 0;
@@ -168,10 +168,9 @@ namespace uba
 		bool HasBeenSeen(const CasKey& casKey);
 
 		virtual void RegisterExternalFileMappingsProvider(ExternalFileMappingsProvider&& provider) override;
-		virtual bool WriteCompressed(WriteResult& out, const tchar* from, FileHandle readHandle, u8* readMem, u64 fileSize, const tchar* toFile, const void* header, u64 headerSize, u64 lastWriteTime = 0) final;
-		bool WriteMemToCompressedFile(FileAccessor& destination, u32 workCount, const u8* uncompressedData, u64 fileSize, u64 maxUncompressedBlock, u64& totalWritten);
-		bool WriteCasFileNoCheck(WriteResult& out, const StringKey& fileNameKey, const tchar* fileName, const CasKey& casKey, const tchar* casFile, bool storeCompressed);
-		bool WriteCasFile(WriteResult& out, const tchar* fileName, const CasKey& casKey);
+		virtual bool WriteCompressedFile(CompressWriteResult& out, const tchar* from, FileHandle readHandle, u8* readMem, u64 fileSize, const tchar* toFile, const void* header, u64 headerSize, u64 lastWriteTime = 0) final;
+		bool WriteCasFileNoCheck(CompressWriteResult& out, const StringKey& fileNameKey, const tchar* fileName, const CasKey& casKey, const tchar* casFile, bool storeCompressed);
+		bool WriteCasFile(CompressWriteResult& out, const tchar* fileName, const CasKey& casKey);
 		bool VerifyExisting(bool& outReturnValue, ScopedWriteLock& entryLock, const CasKey& casKey, CasEntry& casEntry, StringBufferBase& casFile);
 		bool AddCasFile(StringKey fileNameKey, const tchar* fileName, const CasKey& casKey, bool deferCreation, u64 fileSize = InvalidValue, u64 lastWriteTime = 0);
 		void CasEntryAccessed(const CasKey& casKey);
@@ -196,9 +195,8 @@ namespace uba
 		FileEntry& GetOrCreateFileEntry(const StringKey& fileNameKey);
 		CasEntry& GetOrCreateCasEntry(const CasKey& casKey);
 
-		WorkManager* m_workManager;
 		MutableLogger m_logger;
-
+		WorkManager& m_workManager;
 		BufferSlots m_bufferSlots;
 
 		StringBuffer<> m_rootDir;
diff --git a/Engine/Source/Programs/UnrealBuildAccelerator/Common/Public/UbaStorageClient.h b/Engine/Source/Programs/UnrealBuildAccelerator/Common/Public/UbaStorageClient.h
--- a/Engine/Source/Programs/UnrealBuildAccelerator/Common/Public/UbaStorageClient.h
+++ b/Engine/Source/Programs/UnrealBuildAccelerator/Common/Public/UbaStorageClient.h
@@ -14,7 +14,7 @@ namespace uba
 
 	struct StorageClientCreateInfo : StorageCreateInfo
 	{
-		StorageClientCreateInfo(NetworkClient& c, const tchar* rootDir_) : StorageCreateInfo(rootDir_, c.GetLogWriter()), client(c) {}
+		StorageClientCreateInfo(NetworkClient& c, const tchar* rootDir_) : StorageCreateInfo(rootDir_, c.GetLogWriter(), c), client(c) {}
 
 		void Apply(const Config& config);
 
diff --git a/Engine/Source/Programs/UnrealBuildAccelerator/Common/Public/UbaStorageServer.h b/Engine/Source/Programs/UnrealBuildAccelerator/Common/Public/UbaStorageServer.h
--- a/Engine/Source/Programs/UnrealBuildAccelerator/Common/Public/UbaStorageServer.h
+++ b/Engine/Source/Programs/UnrealBuildAccelerator/Common/Public/UbaStorageServer.h
@@ -9,7 +9,7 @@ namespace uba
 {
 	struct StorageServerCreateInfo : StorageCreateInfo
 	{
-		StorageServerCreateInfo(NetworkServer& s, const tchar* rootDir_, LogWriter& writer) : StorageCreateInfo(rootDir_, writer), server(s) { workManager = &server; }
+		StorageServerCreateInfo(NetworkServer& s, const tchar* rootDir_, LogWriter& writer) : StorageCreateInfo(rootDir_, writer, s), server(s) {}
 
 		void Apply(const Config& config);
 
diff --git a/Engine/Source/Programs/UnrealBuildAccelerator/Common/Public/UbaStorageUtils.h b/Engine/Source/Programs/UnrealBuildAccelerator/Common/Public/UbaStorageUtils.h
--- a/Engine/Source/Programs/UnrealBuildAccelerator/Common/Public/UbaStorageUtils.h
+++ b/Engine/Source/Programs/UnrealBuildAccelerator/Common/Public/UbaStorageUtils.h
@@ -9,6 +9,7 @@
 
 namespace uba
 {
+	class FileAccessor;
 	struct BufferSlots;
 	struct StorageStats;
 
@@ -50,4 +51,54 @@ namespace uba
 		u64 sizeOnDisk = 0;
 		u64 bytesReceived = 0;
 	};
+
+	struct LinearWriter
+	{
+		virtual bool Write(const void* data, u64 dataLen) = 0;
+		virtual const tchar* GetHint() = 0;
+	};
+
+	struct CompressWriteResult
+	{
+		FileMappingHandle mappingHandle;
+		u64 readSize = InvalidValue;
+		u64 size = InvalidValue;
+		u64 offset = InvalidValue;
+	};
+
+	struct CompressWriter
+	{
+		CompressWriter(Logger& l, BufferSlots& bs, WorkManager& wm, StorageStats& s, u8 cc = DefaultCompressor, u8 ccl = DefaultCompressionLevel, bool avof = false);
+
+		bool CompressToFile(CompressWriteResult& out, const tchar* from, FileHandle readHandle, u8* readMem, u64 fileSize, const tchar* toFile, const void* header, u64 headerSize, u64 lastWriteTime, const tchar* tempPath);
+		bool CompressedToMapping(CompressWriteResult& out, u8* readMem, u64 fileSize, const tchar* hint);
+		bool CompressFromMem(LinearWriter& destination, u32 workCount, const u8* uncompressedData, u64 fileSize, u64 maxUncompressedBlock, u64& totalWritten);
+		bool CompressFromMemOrFile(CompressWriteResult& result, LinearWriter& destination, const tchar* from, FileHandle readHandle, u8* readMem, u64 fileSize);
+
+		Logger& m_logger;
+		BufferSlots& m_bufferSlots;
+		WorkManager& m_workManager;
+		StorageStats& m_stats;
+		u8 m_casCompressor;
+		u8 m_casCompressionLevel;
+		bool m_asyncUnmapViewOfFile;
+	};
+
+	struct LinearWriterFile : public LinearWriter
+	{
+		LinearWriterFile(FileAccessor& f);
+		virtual bool Write(const void* data, u64 dataLen) override;
+		virtual const tchar* GetHint() override;
+		FileAccessor& file;
+	};
+
+	struct LinearWriterMem : public LinearWriter
+	{
+		LinearWriterMem(u8* d, u64 c);
+		virtual bool Write(const void* data, u64 dataLen) override;
+		virtual const tchar* GetHint() override;
+		u8* pos;
+		u64 size;
+		u64 capacity;
+	};
 }
\ No newline at end of file
diff --git a/Engine/Source/Programs/UnrealBuildAccelerator/Test/Private/UbaTestSession.cpp b/Engine/Source/Programs/UnrealBuildAccelerator/Test/Private/UbaTestSession.cpp
--- a/Engine/Source/Programs/UnrealBuildAccelerator/Test/Private/UbaTestSession.cpp
+++ b/Engine/Source/Programs/UnrealBuildAccelerator/Test/Private/UbaTestSession.cpp
@@ -30,14 +30,14 @@ namespace uba
 		if (!DeleteAllFiles(logger, rootDir.data))
 			return false;
 
-		StorageCreateInfo storageInfo(rootDir.data, logWriter);
+		bool ctorSuccess = true;
+		NetworkServer server(ctorSuccess, { logWriter });
+
+		StorageCreateInfo storageInfo(rootDir.data, logWriter, server);
 		storageInfo.casCapacityBytes = 1024ull * 1024 * 1024;
 		storageInfo.Apply(config);
 		StorageImpl storage(storageInfo);
 
-		bool ctorSuccess = true;
-		NetworkServer server(ctorSuccess, { logWriter });
-
 		SessionServerCreateInfo sessionServerInfo(storage, server, logWriter);
 		sessionServerInfo.checkMemory = false;
 		sessionServerInfo.rootDir = rootDir.data;
diff --git a/Engine/Source/Programs/UnrealBuildAccelerator/Test/Private/UbaTestStorage.cpp b/Engine/Source/Programs/UnrealBuildAccelerator/Test/Private/UbaTestStorage.cpp
--- a/Engine/Source/Programs/UnrealBuildAccelerator/Test/Private/UbaTestStorage.cpp
+++ b/Engine/Source/Programs/UnrealBuildAccelerator/Test/Private/UbaTestStorage.cpp
@@ -14,9 +14,10 @@ namespace uba
 			return true;
 		#endif
 
+		WorkManagerImpl workManager(1);
 		StringBuffer<> rootDir;
 		rootDir.Append(testRootDir).Append(TCV("Uba"));
-		StorageCreateInfo storageInfo(rootDir.data, logger.m_writer);
+		StorageCreateInfo storageInfo(rootDir.data, logger.m_writer, workManager);
 		storageInfo.casCapacityBytes = 1024ull * 1024 * 1024;
 		StorageImpl storage(storageInfo);
 
