[UBA]
* Moved back CompressWriteResult in to Storage again since it was not really needed in CompressWriter

diff --git a/Engine/Source/Programs/UnrealBuildAccelerator/Common/Private/UbaSession.cpp b/Engine/Source/Programs/UnrealBuildAccelerator/Common/Private/UbaSession.cpp
--- a/Engine/Source/Programs/UnrealBuildAccelerator/Common/Private/UbaSession.cpp
+++ b/Engine/Source/Programs/UnrealBuildAccelerator/Common/Private/UbaSession.cpp
@@ -336,7 +336,7 @@ namespace uba
 
 			if (storeCompressed)
 			{
-				CompressWriteResult res;
+				Storage::WriteResult res;
 				CompressedFileHeader header { CalculateCasKey(mem, fileSize, true, &m_workManager, file.name.c_str()) };
 
 				if (!m_storage.WriteCompressedFile(res, TC("MemoryMap"), InvalidFileHandle, mem, fileSize, file.name.c_str(), &header, sizeof(header), file.lastWriteTime))
diff --git a/Engine/Source/Programs/UnrealBuildAccelerator/Common/Private/UbaStorage.cpp b/Engine/Source/Programs/UnrealBuildAccelerator/Common/Private/UbaStorage.cpp
--- a/Engine/Source/Programs/UnrealBuildAccelerator/Common/Private/UbaStorage.cpp
+++ b/Engine/Source/Programs/UnrealBuildAccelerator/Common/Private/UbaStorage.cpp
@@ -324,13 +324,13 @@ namespace uba
 		m_externalFileMappingsProvider = std::move(provider);
 	}
 
-	bool StorageImpl::WriteCompressedFile(CompressWriteResult& out, const tchar* from, FileHandle readHandle, u8* readMem, u64 fileSize, const tchar* to, const void* header, u64 headerSize, u64 lastWriteTime)
+	bool StorageImpl::WriteCompressedFile(WriteResult& out, const tchar* from, FileHandle readHandle, u8* readMem, u64 fileSize, const tchar* to, const void* header, u64 headerSize, u64 lastWriteTime)
 	{
 		CompressWriter w(m_logger, m_bufferSlots, m_workManager, Stats(), m_casCompressor, m_casCompressionLevel, m_asyncUnmapViewOfFile);
-		return w.CompressToFile(out, from, readHandle, readMem, fileSize, to, header, headerSize, lastWriteTime, m_tempPath.data);
+		return w.CompressToFile(out.size, from, readHandle, readMem, fileSize, to, header, headerSize, lastWriteTime, m_tempPath.data);
 	}
 
-	bool StorageImpl::WriteCasFileNoCheck(CompressWriteResult& out, const StringKey& fileNameKey, const tchar* fileName, const CasKey& casKey, const tchar* casFile, bool storeCompressed)
+	bool StorageImpl::WriteCasFileNoCheck(WriteResult& out, const StringKey& fileNameKey, const tchar* fileName, const CasKey& casKey, const tchar* casFile, bool storeCompressed)
 	{
 		StorageStats& stats = Stats();
 		TimerScope ts(stats.createCas);
@@ -411,7 +411,7 @@ namespace uba
 		return true;
 	}
 
-	bool StorageImpl::WriteCasFile(CompressWriteResult& out, const tchar* fileName, const CasKey& casKey)
+	bool StorageImpl::WriteCasFile(WriteResult& out, const tchar* fileName, const CasKey& casKey)
 	{
 		UBA_ASSERT(IsCompressed(casKey) == m_storeCompressed);
 		StringBuffer<> casFile;
@@ -766,7 +766,7 @@ namespace uba
 		casEntry.verified = true;
 		casEntry.exists = false;
 
-		CompressWriteResult res;
+		WriteResult res;
 		if (!WriteCasFileNoCheck(res, fileNameKey, fileName, casKey, casFile.data, IsCompressed(casKey)))
 			return false;
 
@@ -1982,7 +1982,7 @@ namespace uba
 
 		casEntry.verified = true;
 		deferredLock.Leave();
-		CompressWriteResult res;
+		WriteResult res;
 		if (!WriteCasFile(res, deferredCreation.data, casKey))
 			return m_logger.Error(TC("Failed to write deferred cas %s (%s)"), CasKeyString(casKey).str, deferredCreation.data);
 
@@ -2034,7 +2034,7 @@ namespace uba
 		casEntry.verified = true;
 		if (!fileName)
 			return false;
-		CompressWriteResult res;
+		WriteResult res;
 		if (!WriteCasFile(res, fileName, casKey))
 			return false;
 		casEntry.exists = true;
diff --git a/Engine/Source/Programs/UnrealBuildAccelerator/Common/Private/UbaStorageUtils.cpp b/Engine/Source/Programs/UnrealBuildAccelerator/Common/Private/UbaStorageUtils.cpp
--- a/Engine/Source/Programs/UnrealBuildAccelerator/Common/Private/UbaStorageUtils.cpp
+++ b/Engine/Source/Programs/UnrealBuildAccelerator/Common/Private/UbaStorageUtils.cpp
@@ -819,7 +819,7 @@ namespace uba
 	{
 	}
 
-	bool CompressWriter::CompressToFile(CompressWriteResult& out, const tchar* from, FileHandle readHandle, u8* readMem, u64 fileSize, const tchar* toFile, const void* header, u64 headerSize, u64 lastWriteTime, const tchar* tempPath)
+	bool CompressWriter::CompressToFile(u64& outCompressedSize, const tchar* from, FileHandle readHandle, u8* readMem, u64 fileSize, const tchar* toFile, const void* header, u64 headerSize, u64 lastWriteTime, const tchar* tempPath)
 	{
 		FileAccessor destinationFile(m_logger, toFile);
 		if (!destinationFile.CreateWrite(false, DefaultAttributes(), 0, tempPath))
@@ -829,17 +829,18 @@ namespace uba
 				return false;
 
 		LinearWriterFile destination(destinationFile);
-		if (!CompressFromMemOrFile(out, destination, from, readHandle, readMem, fileSize))
+		if (!CompressFromMemOrFile(destination, from, readHandle, readMem, fileSize))
 			return false;
 		if (lastWriteTime)
 			if (!SetFileLastWriteTime(destinationFile.GetHandle(), lastWriteTime))
 				return m_logger.Error(TC("Failed to set file time on filehandle for %s"), toFile);
 		if (!destinationFile.Close())
 			return false;
+		outCompressedSize = destination.GetWritten();
 		return true;
 	}
 
-	bool CompressWriter::CompressedToMapping(CompressWriteResult& out, u8* readMem, u64 fileSize, const tchar* hint)
+	bool CompressWriter::CompressToMapping(FileMappingHandle& outMappingHandle, u64& outMappingSize, u8* readMem, u64 fileSize, const tchar* hint)
 	{
 		u64 mappingSize = fileSize + 1024;
 		FileMappingHandle destMapping = CreateMemoryMappingW(m_logger, PAGE_READWRITE, mappingSize, nullptr, hint);
@@ -853,11 +854,11 @@ namespace uba
 		auto unmapGuard = MakeGuard([&]() { UnmapViewOfFile(m_logger, compressedData, mappingSize, hint); });
 
 		LinearWriterMem destination(compressedData, mappingSize);
-		if (!CompressFromMemOrFile(out, destination, hint, InvalidFileHandle, readMem, fileSize))
+		if (!CompressFromMemOrFile(destination, hint, InvalidFileHandle, readMem, fileSize))
 			return false;
 		destGuard.Cancel();
-		out.mappingHandle = destMapping;
-		out.offset = 0;
+		outMappingHandle = destMapping;
+		outMappingSize = destination.GetWritten();
 		return true;
 	}
 
@@ -995,7 +996,7 @@ namespace uba
 		return true;
 	}
 
-	bool CompressWriter::CompressFromMemOrFile(CompressWriteResult& result, LinearWriter& destination, const tchar* from, FileHandle readHandle, u8* readMem, u64 fileSize)
+	bool CompressWriter::CompressFromMemOrFile(LinearWriter& destination, const tchar* from, FileHandle readHandle, u8* readMem, u64 fileSize)
 	{
 		if (!destination.Write(&fileSize, sizeof(u64))) // Store file size first in compressed file
 			return false;
@@ -1094,16 +1095,16 @@ namespace uba
 
 		stats.createCasBytesRaw += fileSize;
 		stats.createCasBytesComp += totalWritten;
-
-		result.size = totalWritten;
 		return true;
 	}
 
 	LinearWriterFile::LinearWriterFile(FileAccessor& f) : file(f) {}
-	bool LinearWriterFile::Write(const void* data, u64 dataLen) { return file.Write(data, dataLen); }
+	bool LinearWriterFile::Write(const void* data, u64 dataLen) { written += dataLen; return file.Write(data, dataLen); }
+	u64 LinearWriterFile::GetWritten() { return written; }
 	const tchar* LinearWriterFile::GetHint() { return file.GetFileName(); }
 
 	LinearWriterMem::LinearWriterMem(u8* d, u64 c) : pos(d), size(0), capacity(c) {}
+	u64 LinearWriterMem::GetWritten() { return size; }
 	bool LinearWriterMem::Write(const void* data, u64 dataLen)
 	{
 		u64 newSize = size + dataLen;
diff --git a/Engine/Source/Programs/UnrealBuildAccelerator/Common/Public/UbaStorage.h b/Engine/Source/Programs/UnrealBuildAccelerator/Common/Public/UbaStorage.h
--- a/Engine/Source/Programs/UnrealBuildAccelerator/Common/Public/UbaStorage.h
+++ b/Engine/Source/Programs/UnrealBuildAccelerator/Common/Public/UbaStorage.h
@@ -16,7 +16,6 @@ namespace uba
 	class FileAccessor;
 	class Trace;
 	class WorkManager;
-	struct CompressWriteResult;
 	struct DirectoryEntry;
 	extern CasKey EmptyFileKey;
 	
@@ -69,7 +68,15 @@ namespace uba
 		virtual void SetTrace(Trace* trace, bool detailed) {}
 		virtual void Ping() {}
 
-		virtual bool WriteCompressedFile(CompressWriteResult& out, const tchar* from, FileHandle readHandle, u8* readMem, u64 fileSize, const tchar* toFile, const void* header, u64 headerSize, u64 lastWriteTime = 0) = 0;
+		struct WriteResult
+		{
+			FileMappingHandle mappingHandle;
+			u64 readSize = InvalidValue;
+			u64 size = InvalidValue;
+			u64 offset = InvalidValue;
+			bool isIndependent = false;
+		};
+		virtual bool WriteCompressedFile(WriteResult& out, const tchar* from, FileHandle readHandle, u8* readMem, u64 fileSize, const tchar* toFile, const void* header, u64 headerSize, u64 lastWriteTime = 0) = 0;
 
 		struct ExternalFileMapping { FileMappingHandle handle; u64 offset = 0; u64 size = 0; u64 lastWriteTime = 0; u64 writtenSize = InvalidValue; };
 		using ExternalFileMappingsProvider = Function<bool(ExternalFileMapping& out, StringKey fileNameKey, const tchar* fileName)>;
@@ -168,9 +175,9 @@ namespace uba
 		bool HasBeenSeen(const CasKey& casKey);
 
 		virtual void RegisterExternalFileMappingsProvider(ExternalFileMappingsProvider&& provider) override;
-		virtual bool WriteCompressedFile(CompressWriteResult& out, const tchar* from, FileHandle readHandle, u8* readMem, u64 fileSize, const tchar* toFile, const void* header, u64 headerSize, u64 lastWriteTime = 0) final;
-		bool WriteCasFileNoCheck(CompressWriteResult& out, const StringKey& fileNameKey, const tchar* fileName, const CasKey& casKey, const tchar* casFile, bool storeCompressed);
-		bool WriteCasFile(CompressWriteResult& out, const tchar* fileName, const CasKey& casKey);
+		virtual bool WriteCompressedFile(WriteResult& out, const tchar* from, FileHandle readHandle, u8* readMem, u64 fileSize, const tchar* toFile, const void* header, u64 headerSize, u64 lastWriteTime = 0) final;
+		bool WriteCasFileNoCheck(WriteResult& out, const StringKey& fileNameKey, const tchar* fileName, const CasKey& casKey, const tchar* casFile, bool storeCompressed);
+		bool WriteCasFile(WriteResult& out, const tchar* fileName, const CasKey& casKey);
 		bool VerifyExisting(bool& outReturnValue, ScopedWriteLock& entryLock, const CasKey& casKey, CasEntry& casEntry, StringBufferBase& casFile);
 		bool AddCasFile(StringKey fileNameKey, const tchar* fileName, const CasKey& casKey, bool deferCreation, u64 fileSize = InvalidValue, u64 lastWriteTime = 0);
 		void CasEntryAccessed(const CasKey& casKey);
diff --git a/Engine/Source/Programs/UnrealBuildAccelerator/Common/Public/UbaStorageUtils.h b/Engine/Source/Programs/UnrealBuildAccelerator/Common/Public/UbaStorageUtils.h
--- a/Engine/Source/Programs/UnrealBuildAccelerator/Common/Public/UbaStorageUtils.h
+++ b/Engine/Source/Programs/UnrealBuildAccelerator/Common/Public/UbaStorageUtils.h
@@ -55,25 +55,18 @@ namespace uba
 	struct LinearWriter
 	{
 		virtual bool Write(const void* data, u64 dataLen) = 0;
+		virtual u64 GetWritten() = 0;
 		virtual const tchar* GetHint() = 0;
 	};
 
-	struct CompressWriteResult
-	{
-		FileMappingHandle mappingHandle;
-		u64 readSize = InvalidValue;
-		u64 size = InvalidValue;
-		u64 offset = InvalidValue;
-	};
-
 	struct CompressWriter
 	{
 		CompressWriter(Logger& l, BufferSlots& bs, WorkManager& wm, StorageStats& s, u8 cc = DefaultCompressor, u8 ccl = DefaultCompressionLevel, bool avof = false);
 
-		bool CompressToFile(CompressWriteResult& out, const tchar* from, FileHandle readHandle, u8* readMem, u64 fileSize, const tchar* toFile, const void* header, u64 headerSize, u64 lastWriteTime, const tchar* tempPath);
-		bool CompressedToMapping(CompressWriteResult& out, u8* readMem, u64 fileSize, const tchar* hint);
+		bool CompressToFile(u64& outCompressedSize, const tchar* from, FileHandle readHandle, u8* readMem, u64 fileSize, const tchar* toFile, const void* header, u64 headerSize, u64 lastWriteTime, const tchar* tempPath);
+		bool CompressToMapping(FileMappingHandle& outMappingHandle, u64& outMappingSize, u8* readMem, u64 fileSize, const tchar* hint);
 		bool CompressFromMem(LinearWriter& destination, u32 workCount, const u8* uncompressedData, u64 fileSize, u64 maxUncompressedBlock, u64& totalWritten);
-		bool CompressFromMemOrFile(CompressWriteResult& result, LinearWriter& destination, const tchar* from, FileHandle readHandle, u8* readMem, u64 fileSize);
+		bool CompressFromMemOrFile(LinearWriter& destination, const tchar* from, FileHandle readHandle, u8* readMem, u64 fileSize);
 
 		Logger& m_logger;
 		BufferSlots& m_bufferSlots;
@@ -84,18 +77,21 @@ namespace uba
 		bool m_asyncUnmapViewOfFile;
 	};
 
-	struct LinearWriterFile : public LinearWriter
+	struct LinearWriterFile final : public LinearWriter
 	{
 		LinearWriterFile(FileAccessor& f);
 		virtual bool Write(const void* data, u64 dataLen) override;
+		virtual u64 GetWritten() override;
 		virtual const tchar* GetHint() override;
 		FileAccessor& file;
+		u64 written = 0;
 	};
 
-	struct LinearWriterMem : public LinearWriter
+	struct LinearWriterMem final : public LinearWriter
 	{
 		LinearWriterMem(u8* d, u64 c);
 		virtual bool Write(const void* data, u64 dataLen) override;
+		virtual u64 GetWritten() override;
 		virtual const tchar* GetHint() override;
 		u8* pos;
 		u64 size;
