[UBA]
* Implement support for virtual files created from memory buffer
* Fixed code paths for not writing to disk at all
* Fixed exported api for getting file content of files never hitting disk

diff --git a/Engine/Source/Programs/UnrealBuildAccelerator/Cli/Private/UbaCli.cpp b/Engine/Source/Programs/UnrealBuildAccelerator/Cli/Private/UbaCli.cpp
--- a/Engine/Source/Programs/UnrealBuildAccelerator/Cli/Private/UbaCli.cpp
+++ b/Engine/Source/Programs/UnrealBuildAccelerator/Cli/Private/UbaCli.cpp
@@ -787,9 +787,10 @@ namespace uba
 
 			if (getCas)
 			{
+				StringKey nameKey = CaseInsensitiveFs ? ToStringKeyLower(application) : ToStringKey(application);
 				FileAccessor fa(logger, application.c_str());
-				if (!fa.OpenMemoryRead())
-					return logger.Error(TC("Failed to open file %s"), application.c_str());
+				if (!fa.OpenMemoryRead(0, false))
+					return logger.Error(TC("Failed to open file %s (Key: %s) - %s"), application.c_str(), KeyToString(nameKey).data, LastErrorToText().data);
 				u64 fileSize = fa.GetSize();
 				u8* data = fa.GetData();
 				bool is64Bit = false;
@@ -820,7 +821,7 @@ namespace uba
 							isDotnet = *(u32*)(data + offset + 0x18 + 0x70);
 					}
 				}
-				logger.Info(TC("%s"), application.c_str());
+				logger.Info(TC("%s (Key: %s)"), application.c_str(), KeyToString(nameKey).data);
 				logger.Info(TC("  Is64Bit: %s"), (is64Bit ? TC("true") : TC("false")));
 				logger.Info(TC("  Arch: %s"), (isX64 ? TC("x64") : (isArm64 ? TC("arm64") : (isDotnet ? TC(".net") : TC("unknown")))));
 				logger.Info(TC("  Size: %llu"), fileSize);
diff --git a/Engine/Source/Programs/UnrealBuildAccelerator/Common/Private/UbaCacheClient.cpp b/Engine/Source/Programs/UnrealBuildAccelerator/Common/Private/UbaCacheClient.cpp
--- a/Engine/Source/Programs/UnrealBuildAccelerator/Common/Private/UbaCacheClient.cpp
+++ b/Engine/Source/Programs/UnrealBuildAccelerator/Common/Private/UbaCacheClient.cpp
@@ -841,6 +841,9 @@ namespace uba
 		FileFetcher fetcher { m_storage.m_bufferSlots, storageStats };
 		fetcher.m_errorOnFail = false;
 
+		// TODO: This probably is not needed but is correct. Since there could be deferred entries of the file we are downloading and writing to
+		//m_storage.ReportFileWrite(path.data, 
+
 		if (IsNormalized(casKey))
 		{
 			DowngradedLogger logger(m_connected, m_logger.m_writer, TC("UbaCacheClientNormalizedDownload"));
diff --git a/Engine/Source/Programs/UnrealBuildAccelerator/Common/Private/UbaProcess.cpp b/Engine/Source/Programs/UnrealBuildAccelerator/Common/Private/UbaProcess.cpp
--- a/Engine/Source/Programs/UnrealBuildAccelerator/Common/Private/UbaProcess.cpp
+++ b/Engine/Source/Programs/UnrealBuildAccelerator/Common/Private/UbaProcess.cpp
@@ -463,6 +463,16 @@ namespace uba
 		#endif
 	}
 
+	void ProcessImpl::TraverseOutputFiles(const Function<void(StringView file)>& func)
+	{
+		SCOPED_FUTEX_READ(m_shared.writtenFilesLock, lock);
+		for (auto& kv : m_shared.writtenFiles)
+		{
+			WrittenFile& wf = kv.second;
+			func(wf.name);
+		}
+	}
+
 	bool ProcessImpl::WaitForRead(PipeReader& outReader, PipeReader& errReader)
 	{
 		while (true)
diff --git a/Engine/Source/Programs/UnrealBuildAccelerator/Common/Private/UbaSession.cpp b/Engine/Source/Programs/UnrealBuildAccelerator/Common/Private/UbaSession.cpp
--- a/Engine/Source/Programs/UnrealBuildAccelerator/Common/Private/UbaSession.cpp
+++ b/Engine/Source/Programs/UnrealBuildAccelerator/Common/Private/UbaSession.cpp
@@ -194,6 +194,11 @@ namespace uba
 		UBA_ASSERT(m_process);
 		return m_process->GetExecutionType();
 	}
+	void ProcessHandle::TraverseOutputFiles(const Function<void(StringView file)>& func) const
+	{
+		UBA_ASSERT(m_process);
+		m_process->TraverseOutputFiles(func);
+	}
 	ProcessHandle::ProcessHandle(Process* process)
 	{
 		m_process = process;
@@ -366,30 +371,8 @@ namespace uba
 				}
 				else
 				{
-					// Seems like best combo (for windows at least) is to use writes with overlap and max 16 at the same time.
-					// On one machine we get twice as fast without overlap if no bottleneck. On another machine (ntfs compression on) we get twice as slow without overlap
-					// Both machines behaves well with overlap AND bottleneck. Both machine are 128 logical core thread rippers.
-					bool useFileMapForWrite = fileSize && fileSize <= m_writeFilesFileMapMax; // ::CreateFileMappingW does not work for zero-length files
-					bool useOverlap = !useFileMapForWrite && fileSize >= m_writeFilesNoBufferingMin;//fileSize > 2*1024*1024;
-
-					u32 attributes = DefaultAttributes();
-					if (useOverlap)
-						attributes |= (FILE_FLAG_OVERLAPPED | FILE_FLAG_NO_BUFFERING);
-
-					if (useFileMapForWrite)
-					{
-						if (!destinationFile.CreateMemoryWrite(false, attributes, fileSize, m_tempPath.data))
-							return false;
-						MapMemoryCopy(destinationFile.GetData(), mem, fileSize);
-					}
-					else
-					{
-						if (!destinationFile.CreateWrite(false, attributes, fileSize, m_tempPath.data))
-							return false;
-						if (!destinationFile.Write(mem, fileSize, 0, true))
-							return false;
-					}
-
+					if (!WriteMemoryToDisk(destinationFile, mem, fileSize))
+						return false;
 					writtenSize = fileSize;
 				}
 
@@ -409,6 +392,9 @@ namespace uba
 		{
 			// Delete existing file to make sure it is not picked up (since it is out of date)
 			uba::DeleteFileW(file.name.c_str());
+
+			// If shouldWriteToDisk is false we can't evict from memory because we can't re-read from disk if needed
+			shouldEvictFromMemory = false;
 		}
 
 		FileMappingHandle mh = file.mappingHandle;
@@ -454,19 +440,23 @@ namespace uba
 			entry.usedCount = 0;
 			entry.usedCountBeforeFree = g_globalRules.GetUsedCountBeforeFree(file.name);
 			entry.writtenSize = writtenSize;
+			entry.isInvisible = rules.IsInvisible(file.name);
 
 			#if UBA_DEBUG_TRACK_MAPPING
 			entry.name = file.name;
 			m_debugLogger->Info(TC("Mapping kept 0x%llx (%s) from detoured process (UsedCountBeforeFree: %u)"), u64(mh.mh), entry.name.c_str(), u32(entry.usedCountBeforeFree));
 			#endif
 
-			SCOPED_WRITE_LOCK(m_fileMappingTableMemLock, lock);
-			BinaryWriter writer(m_fileMappingTableMem, m_fileMappingTableSize);
-			writer.WriteStringKey(file.key);
-			writer.WriteString(name);
-			writer.Write7BitEncoded(file.mappingWritten);
-			u32 newSize = (u32)writer.GetPosition();
-			m_fileMappingTableSize = (u32)newSize;
+			if (!entry.isInvisible)
+			{
+				SCOPED_WRITE_LOCK(m_fileMappingTableMemLock, lock);
+				BinaryWriter writer(m_fileMappingTableMem, m_fileMappingTableSize);
+				writer.WriteStringKey(file.key);
+				writer.WriteString(name);
+				writer.Write7BitEncoded(file.mappingWritten);
+				u32 newSize = (u32)writer.GetPosition();
+				m_fileMappingTableSize = (u32)newSize;
+			}
 		}
 
 		if (shouldWriteToDisk)
@@ -475,6 +465,66 @@ namespace uba
 		return true;
 	}
 
+	bool Session::WriteMemoryToDisk(FileAccessor& destinationFile, const void* fileMem, u64 fileSize)
+	{
+		// Seems like best combo (for windows at least) is to use writes with overlap and max 16 at the same time.
+		// On one machine we get twice as fast without overlap if no bottleneck. On another machine (ntfs compression on) we get twice as slow without overlap
+		// Both machines behaves well with overlap AND bottleneck. Both machine are 128 logical core thread rippers.
+		bool useFileMapForWrite = fileSize && fileSize <= m_writeFilesFileMapMax; // ::CreateFileMappingW does not work for zero-length files
+		bool useOverlap = !useFileMapForWrite && fileSize >= m_writeFilesNoBufferingMin;//fileSize > 2*1024*1024;
+
+		u32 attributes = DefaultAttributes();
+		if (useOverlap)
+			attributes |= (FILE_FLAG_OVERLAPPED | FILE_FLAG_NO_BUFFERING);
+
+		if (useFileMapForWrite)
+		{
+			if (!destinationFile.CreateMemoryWrite(false, attributes, fileSize, m_tempPath.data))
+				return false;
+			MapMemoryCopy(destinationFile.GetData(), fileMem, fileSize);
+		}
+		else
+		{
+			if (!destinationFile.CreateWrite(false, attributes, fileSize, m_tempPath.data))
+				return false;
+			if (!destinationFile.Write(fileMem, fileSize, 0, true))
+				return false;
+		}
+		return true;
+	}
+
+	bool Session::GetFileMemory(const Function<bool(const void*, u64)>& func, const StringKey& fileNameKey, StringView filePath, bool deleteInternalMapping)
+	{
+		FileMappingHandle mapping;
+		u64 size;
+		{
+			SCOPED_FUTEX(m_fileMappingTableLookupLock, lookupLock);
+			auto findIt = m_fileMappingTableLookup.find(fileNameKey);
+			if (findIt == m_fileMappingTableLookup.end())
+				return false;
+			FileMappingEntry& entry = findIt->second;
+			if (deleteInternalMapping && !entry.isInvisible)
+				return m_logger.Error(TC("Trying to delete mapping that is not invisible (%s)"), filePath.data);
+			mapping = entry.mapping;
+			size = entry.size;
+			if (deleteInternalMapping)
+			{
+				if (!entry.isInvisible)
+					return m_logger.Error(TC("Trying to delete mapping that is not invisible (%s)"), filePath.data);
+				m_fileMappingTableLookup.erase(findIt);
+			}
+		}
+
+		u8* mem = MapViewOfFile(m_logger, mapping, FILE_MAP_READ, 0, size);
+		if (!mem)
+			return false;
+		bool res = func(mem, size);
+		UnmapViewOfFile(m_logger, mem, size, filePath.data);
+		if (deleteInternalMapping)
+			CloseFileMapping(m_logger, mapping, TC(""));
+		return res;
+	}
+
 	void Session::AddEnvironmentVariableNoLock(const tchar* key, const tchar* value)
 	{
 		m_environmentVariables.insert(m_environmentVariables.end(), key, key + TStrlen(key));
@@ -770,7 +820,7 @@ namespace uba
 		return newSize;
 	}
 
-	bool Session::CreateMemoryMapFromFile(MemoryMap& out, StringKey fileNameKey, const tchar* fileName, bool isCompressedCas, u64 alignment, const tchar* hint, ProcessImpl* requestingProcess, bool canBeFreed)
+	bool Session::GetOrCreateMemoryMapFromFile(MemoryMap& out, StringKey fileNameKey, const tchar* fileName, bool isCompressedCas, u64 alignment, const tchar* hint, ProcessImpl* requestingProcess, bool canBeFreed)
 	{
 		TimerScope ts(Stats().waitMmapFromFile);
 
@@ -925,7 +975,7 @@ namespace uba
 
 			if (isCompressed)
 			{
-				if (!m_storage.DecompressFileToMemory(fileName, fileHandle, mappingMemory, size, TC("CreateMemoryMapFromFile"), fileStartOffset))
+				if (!m_storage.DecompressFileToMemory(fileName, fileHandle, mappingMemory, size, TC("GetOrCreateMemoryMapFromFile"), fileStartOffset))
 					return false;
 			}
 			else
@@ -976,7 +1026,7 @@ namespace uba
 		return true;
 	}
 
-	bool Session::CreateMemoryMapFromView(MemoryMap& out, StringKey fileNameKey, const tchar* fileName, const CasKey& casKey, u64 alignment)
+	bool Session::GetOrCreateMemoryMapFromStorage(MemoryMap& out, StringKey fileNameKey, const tchar* fileName, const CasKey& casKey, u64 alignment)
 	{
 		//StringBuffer<> workName;
 		//u32 len = TStrlen(fileName);
@@ -1027,7 +1077,7 @@ namespace uba
 			else
 			{
 				if (size > m_fileMappingBuffer.GetFileMappingCapacity())
-					return m_logger.Error(TC("File %s has a size (%llu) that is too large to fit in mapping buffer (CreateMemoryMapFromView)"), fileName, size);
+					return m_logger.Error(TC("File %s has a size (%llu) that is too large to fit in mapping buffer (GetOrCreateMemoryMapFromView)"), fileName, size);
 
 				auto mappedViewWrite = m_fileMappingBuffer.AllocAndMapView(MappedView_Transient, size, alignment, fileName);
 				if (!mappedViewWrite.memory)
@@ -1153,7 +1203,7 @@ namespace uba
 			UBA_ASSERT(res); (void)res;
 			dirLock.Enter();
 		}
-		UBA_ASSERTF(dir.parseOffset == 1, TC("Registering create file for write %s with unexpect dir.parseOffset %u "), fileName.data, dir.parseOffset);
+		UBA_ASSERTF(dir.parseOffset == 1 || !m_shouldWriteToDisk, TC("Registering create file for write %s with unexpect dir.parseOffset %u "), fileName.data, dir.parseOffset);
 
 		if (fileNameKey == StringKeyZero)
 		{
@@ -1696,6 +1746,7 @@ namespace uba
 				out.size = entry.size;
 				out.lastWriteTime = entry.lastWriteTime;
 				out.writtenSize = entry.writtenSize;
+				out.createIndependentMapping = entry.createIndependentMapping;
 				return true;
 			});
 
@@ -1943,6 +1994,34 @@ namespace uba
 		return RegisterCreateFileForWrite(fileNameKey, fixedFilePath, false, sourceSize, 0, false);
 	}
 
+	bool Session::CreateVirtualFile(const tchar* filePath, const void* memory, u64 memorySize)
+	{
+		UBA_ASSERT(!m_runningRemote);
+		StringBuffer<> fixedFilePath;
+		auto fileNameKey = GetKeyAndFixedName(fixedFilePath, filePath);
+		if (!CreateVirtualFileInternal(fileNameKey, fixedFilePath, memory, memorySize))
+			return false;
+		return RegisterCreateFileForWrite(fileNameKey, fixedFilePath, false, memorySize, 0, false);
+	}
+
+	bool Session::DeleteVirtualFile(const tchar* filePath)
+	{
+		UBA_ASSERT(!m_runningRemote);
+		StringBuffer<> fixedFilePath;
+		auto key = GetKeyAndFixedName(fixedFilePath, filePath);
+		RegisterDeleteFile(key, fixedFilePath);
+
+		SCOPED_FUTEX(m_virtualSourceFilesLock, virtualSourceFilesLock);
+		auto findIt = m_virtualSourceFiles.find(key); 
+		if (findIt == m_virtualSourceFiles.end())
+			return false;
+		VirtualSourceFile vsf = findIt->second;
+		m_virtualSourceFiles.erase(findIt);
+		virtualSourceFilesLock.Leave();
+		CloseFileMapping(m_logger, vsf.mappingHandle, TC("VirtualFile"));
+		return true;
+	}
+
 	void Session::RegisterDeleteFile(const tchar* filePath)
 	{
 		UBA_ASSERT(!m_runningRemote);
@@ -1972,6 +2051,32 @@ namespace uba
 		m_getNextProcessFunction = function;
 	}
 
+	bool Session::GetOutputFileSize(u64& outSize, const tchar* filePath)
+	{
+		if (m_shouldWriteToDisk)
+			return m_logger.Error(TC("GetFileSize only implemented for path where ShouldWriteToDisk is false"));
+		UBA_ASSERT(!m_runningRemote);
+		StringBuffer<> fixedFilePath;
+		auto key = GetKeyAndFixedName(fixedFilePath, filePath);
+		return GetOutputFileSizeInternal(outSize, key, fixedFilePath);
+	}
+
+	bool Session::GetOutputFileData(void* outData, const tchar* filePath, bool deleteInternalMapping)
+	{
+		UBA_ASSERT(!m_runningRemote);
+		StringBuffer<> fixedFilePath;
+		auto key = GetKeyAndFixedName(fixedFilePath, filePath);
+		return GetOutputFileDataInternal(outData, key, fixedFilePath, deleteInternalMapping);
+	}
+
+	bool Session::WriteOutputFile(const tchar* filePath, bool deleteInternalMapping)
+	{
+		UBA_ASSERT(!m_runningRemote);
+		StringBuffer<> fixedFilePath;
+		auto key = GetKeyAndFixedName(fixedFilePath, filePath);
+		return WriteOutputFileInternal(key, fixedFilePath, deleteInternalMapping);
+	}
+
 	const tchar* Session::GetId() { return m_id.data; }
 	Storage& Session::GetStorage() { return m_storage; }
 	MutableLogger& Session::GetLogger() { return m_logger; }
@@ -2758,7 +2863,7 @@ namespace uba
 			{
 				MemoryMap map;
 				bool canBeFreed = canBeCompressed;
-				if (CreateMemoryMapFromFile(map, fileNameKey, fileName.data, false, alignment, TC("CreateFile"), &process, canBeFreed))
+				if (GetOrCreateMemoryMapFromFile(map, fileNameKey, fileName.data, false, alignment, TC("CreateFile"), &process, canBeFreed))
 				{
 					out.size = map.size;
 					out.fileName.Append(map.name);
@@ -2816,14 +2921,18 @@ namespace uba
 		m_activeFiles.erase(msg.closeId);
 		lock.Leave();
 
-		bool registerRealFile = true;
-		u64 fileSize = 0;
-		u64 lastWriteTime = 0;
-
 		if (!msg.success)
 		{
 			return true;
 		}
+
+		bool registerRealFile = true;
+		u64 fileSize = 0;
+		u64 lastWriteTime = 0;
+
+		ProcessStartInfo& startInfo = msg.process.m_startInfo;
+		auto& rules = *startInfo.rules;
+
 		if (msg.deleteOnClose)
 		{
 			RemoveWrittenFile(msg.process, activeFile.nameKey);
@@ -2894,12 +3003,10 @@ namespace uba
 				writtenFile.owner = &msg.process;
 			}
 
-			ProcessStartInfo& startInfo = msg.process.m_startInfo;
-
 			if (!m_runningRemote && HasVfs(startInfo.rootsHandle)) // For posix we write the dependency file directly to disk so we need to update it if vfs is enabled
 			{
 				bool escapeSpaces;
-				if (!msg.mappingHandle.IsValid() && startInfo.rules->ShouldDevirtualizeFile(activeFile.name, escapeSpaces))
+				if (!msg.mappingHandle.IsValid() && rules.ShouldDevirtualizeFile(activeFile.name, escapeSpaces))
 				{
 					// TODO: On linux we don't use file mappings for outputs yet.. so we have to open the file and change it
 					FileAccessor readFile(m_logger, name.data);
@@ -2948,23 +3055,26 @@ namespace uba
 				lastWriteTime = writtenFile.lastWriteTime;
 			}
 
-			if ((msg.process.m_extractExports || m_extractObjFilesSymbols) && startInfo.rules->ShouldExtractSymbols(activeFile.name))
+			if ((msg.process.m_extractExports || m_extractObjFilesSymbols) && rules.ShouldExtractSymbols(activeFile.name))
 				if (!ExtractSymbolsFromObjectFile(msg, name.data, fileSize))
 					return false;
 		}
 
-		if (!msg.newName.IsEmpty())
-		{
-			RegisterDeleteFile(activeFile.nameKey, activeFile.name);
-			if (RegisterCreateFileForWrite(msg.newNameKey, msg.newName, registerRealFile, fileSize, lastWriteTime) && registerRealFile)
-				TraceWrittenFile(msg.process.m_id, msg.newName, fileSize);
-		}
-		else if (msg.deleteOnClose)
-			RegisterDeleteFile(activeFile.nameKey, activeFile.name);
-		else
+		if (!rules.IsInvisible(activeFile.name))
 		{
-			if (RegisterCreateFileForWrite(activeFile.nameKey, activeFile.name, registerRealFile, fileSize, lastWriteTime) && registerRealFile)
-				TraceWrittenFile(msg.process.m_id, activeFile.name, fileSize);
+			if (!msg.newName.IsEmpty())
+			{
+				RegisterDeleteFile(activeFile.nameKey, activeFile.name);
+				if (RegisterCreateFileForWrite(msg.newNameKey, msg.newName, registerRealFile, fileSize, lastWriteTime) && registerRealFile)
+					TraceWrittenFile(msg.process.m_id, msg.newName, fileSize);
+			}
+			else if (msg.deleteOnClose)
+				RegisterDeleteFile(activeFile.nameKey, activeFile.name);
+			else
+			{
+				if (RegisterCreateFileForWrite(activeFile.nameKey, activeFile.name, registerRealFile, fileSize, lastWriteTime) && registerRealFile)
+					TraceWrittenFile(msg.process.m_id, activeFile.name, fileSize);
+			}
 		}
 
 		out.directoryTableSize = GetDirectoryTableSize();
@@ -3079,9 +3189,11 @@ namespace uba
 			}
 		}
 
+		if (!process.GetStartInfo().rules->IsInvisible(msg.toName))
+			if (RegisterCreateFileForWrite(msg.toKey, msg.toName, true))
+				TraceWrittenFile(process.m_id, msg.toName, 0);
+
 		out.errorCode = ERROR_SUCCESS;
-		if (RegisterCreateFileForWrite(msg.toKey, msg.toName, true))
-			TraceWrittenFile(process.m_id, msg.toName, 0);
 		out.directoryTableSize = RegisterDeleteFile(msg.fromKey, msg.fromName);
 		return true;
 	}
@@ -3305,7 +3417,7 @@ namespace uba
 
 		MemoryMap out;
 		u64 alignment = GetMemoryMapAlignment(fileName);
-		if (!CreateMemoryMapFromFile(out, fileNameKey, fileName.data, false, alignment, TC("Remap"), &process, true))
+		if (!GetOrCreateMemoryMapFromFile(out, fileNameKey, fileName.data, false, alignment, TC("Remap"), &process, true))
 			return m_logger.Error(TC("Failed to remap %s"), fileName.data);
 		writer.WriteU32(GetFileMappingSize());
 		return true;
@@ -3364,30 +3476,59 @@ namespace uba
 		TimerScope ts(Stats().createMmapFromFile);
 
 		VirtualSourceFile virtualFile;
-
+		StringKey sourceFileKey = CaseInsensitiveFs ? ToStringKeyLower(ToView(sourceFile)) : ToStringKey(ToView(sourceFile));
+		SCOPED_FUTEX(m_virtualSourceFilesLock, virtualSourceFilesLock);
+		auto insres = m_virtualSourceFiles.try_emplace(sourceFileKey); 
+		if (insres.second)
 		{
-			StringKey sourceFileKey = CaseInsensitiveFs ? ToStringKeyLower(ToView(sourceFile)) : ToStringKey(ToView(sourceFile));
-			SCOPED_FUTEX(m_virtualSourceFilesLock, virtualSourceFilesLock);
-			auto insres = m_virtualSourceFiles.try_emplace(sourceFileKey); 
-			if (insres.second)
-			{
-				FileHandle fileHandle = uba::CreateFileW(sourceFile, GENERIC_READ, FILE_SHARE_READ, OPEN_EXISTING, DefaultAttributes());
-				if (fileHandle == InvalidFileHandle)
-					return m_logger.Error(TC("[RegisterVirtualFileInternal] CreateFileW for %s failed (%s)"), sourceFile, LastErrorToText().data);
-				u64 fileSize = 0;
-				if (!GetFileSizeEx(fileSize, fileHandle))
-					return m_logger.Error(TC("[RegisterVirtualFileInternal] GetFileSizeEx for %s failed (%s)"), sourceFile, LastErrorToText().data);
-				auto fg = MakeGuard([&]() { uba::CloseFile(sourceFile, fileHandle); });
-				virtualFile.mappingHandle = CreateFileMappingW(m_logger, fileHandle, PAGE_READONLY, fileSize, sourceFile);
-				virtualFile.size = fileSize;
-				insres.first->second = virtualFile;
-			}
-			else
-				virtualFile = insres.first->second;
-			if (!virtualFile.mappingHandle.IsValid())
-				return m_logger.Error(TC("[RegisterVirtualFileInternal] CreateFileMapping for %s failed (%s)"), sourceFile, LastErrorToText().data);
+			FileHandle fileHandle = uba::CreateFileW(sourceFile, GENERIC_READ, FILE_SHARE_READ, OPEN_EXISTING, DefaultAttributes());
+			if (fileHandle == InvalidFileHandle)
+				return m_logger.Error(TC("[RegisterVirtualFileInternal] CreateFileW for %s failed (%s)"), sourceFile, LastErrorToText().data);
+			u64 fileSize = 0;
+			if (!GetFileSizeEx(fileSize, fileHandle))
+				return m_logger.Error(TC("[RegisterVirtualFileInternal] GetFileSizeEx for %s failed (%s)"), sourceFile, LastErrorToText().data);
+			auto fg = MakeGuard([&]() { uba::CloseFile(sourceFile, fileHandle); });
+			virtualFile.mappingHandle = CreateFileMappingW(m_logger, fileHandle, PAGE_READONLY, fileSize, sourceFile);
+			virtualFile.size = fileSize;
+			insres.first->second = virtualFile;
 		}
+		else
+			virtualFile = insres.first->second;
+		if (!virtualFile.mappingHandle.IsValid())
+			return m_logger.Error(TC("[RegisterVirtualFileInternal] CreateFileMapping for %s failed (%s)"), sourceFile, LastErrorToText().data);
+		virtualSourceFilesLock.Leave();
 
+		if (sourceSize + sourceOffset > virtualFile.size)
+			return m_logger.Error(TC("Virtual file offset(%llu)+size(%llu) outside source file size(%llu)"), sourceOffset, sourceSize, virtualFile.size, filePath.data);
+
+		return RegisterVirtualFileInternal(fileNameKey, filePath, virtualFile.mappingHandle, sourceSize, sourceOffset);
+	}
+
+	bool Session::CreateVirtualFileInternal(const StringKey& fileNameKey, const StringView& filePath, const void* memory, u64 memorySize)
+	{
+		TimerScope ts(Stats().createMmapFromFile);
+
+		SCOPED_FUTEX(m_virtualSourceFilesLock, virtualSourceFilesLock);
+		auto insres = m_virtualSourceFiles.try_emplace(fileNameKey); 
+		if (!insres.second)
+			return false;
+		VirtualSourceFile& virtualFile = insres.first->second;
+		FileMappingHandle mapping = CreateMemoryMappingW(m_logger, PAGE_READWRITE, memorySize, nullptr, filePath.data);
+		if (!mapping.IsValid())
+			return false;
+		u8* mem2 = MapViewOfFile(m_logger, mapping, FILE_MAP_WRITE, 0, memorySize);
+		if (!mem2)
+			return false;
+		MapMemoryCopy(mem2, memory, memorySize);
+		UnmapViewOfFile(m_logger, mem2, memorySize, filePath.data);
+
+		virtualFile.mappingHandle = mapping;
+		virtualFile.size = memorySize;
+		return RegisterVirtualFileInternal(fileNameKey, filePath, virtualFile.mappingHandle, memorySize, 0);
+	}
+
+	bool Session::RegisterVirtualFileInternal(const StringKey& fileNameKey, const StringView& filePath, FileMappingHandle mappingHandle, u64 mappingSize, u64 mappingOffset)
+	{
 		SCOPED_FUTEX(m_fileMappingTableLookupLock, lookupLock);
 		auto insres = m_fileMappingTableLookup.try_emplace(fileNameKey);
 		FileMappingEntry& entry = insres.first->second;
@@ -3398,17 +3539,15 @@ namespace uba
 		if (entry.handled)
 			return m_logger.Error(TC("Virtual file %s has already been registered"), filePath.data);
 
-		if (sourceSize + sourceOffset > virtualFile.size)
-			return m_logger.Error(TC("Virtual file offset(%llu)+size(%llu) outside source file size(%llu)"), sourceOffset, sourceSize, virtualFile.size, filePath.data);
-
-		entry.mapping = virtualFile.mappingHandle;
-		entry.mappingOffset = sourceOffset;
-		entry.size = sourceSize;
+		entry.mapping = mappingHandle;
+		entry.mappingOffset = mappingOffset;
+		entry.size = mappingSize;
 		entry.handled = true;
 		entry.lastWriteTime = 0; // TODO: Take lastwritetime of source file?
+		entry.createIndependentMapping = true; // We want independent mapping in casdb so they can be deleted
 
 		StringBuffer<> mappingName;
-		Storage::GetMappingString(mappingName, virtualFile.mappingHandle, sourceOffset);
+		Storage::GetMappingString(mappingName, mappingHandle, mappingOffset);
 		entry.success = true;
 
 		#if UBA_DEBUG_TRACK_MAPPING
@@ -3420,12 +3559,41 @@ namespace uba
 		BinaryWriter writer(m_fileMappingTableMem, m_fileMappingTableSize);
 		writer.WriteStringKey(fileNameKey);
 		writer.WriteString(mappingName);
-		writer.Write7BitEncoded(sourceSize);
+		writer.Write7BitEncoded(mappingSize);
 		m_fileMappingTableSize = (u32)writer.GetPosition();
 
 		return true;
 	}
 
+	bool Session::GetOutputFileSizeInternal(u64& outSize, const StringKey& fileNameKey, StringView filePath)
+	{
+		SCOPED_FUTEX(m_fileMappingTableLookupLock, lookupLock);
+		auto findIt = m_fileMappingTableLookup.find(fileNameKey);
+		if (findIt == m_fileMappingTableLookup.end())
+			return false;
+		FileMappingEntry& entry = findIt->second;
+		lookupLock.Leave();
+		SCOPED_FUTEX(entry.lock, entryCs);
+		outSize = entry.size;
+		return true;
+	}
+
+	bool Session::GetOutputFileDataInternal(void* outData, const StringKey& fileNameKey, StringView filePath, bool deleteInternalMapping)
+	{
+		return GetFileMemory([&](const void* fileMem, u64 fileSize) { MapMemoryCopy(outData, fileMem, fileSize); return true; }, fileNameKey, filePath, deleteInternalMapping);
+	}
+
+	bool Session::WriteOutputFileInternal(const StringKey& fileNameKey, StringView filePath, bool deleteInternalMapping)
+	{
+		//m_logger.Info(TC("TRYING TO WRITE OUTPUT %s (%s)"), filePath, KeyToString(fileNameKey).data);
+		FileAccessor destinationFile(m_logger, filePath.data);
+		if (!GetFileMemory([&](const void* fileMem, u64 fileSize) { return WriteMemoryToDisk(destinationFile, fileMem, fileSize); }, fileNameKey, filePath, deleteInternalMapping))
+			return false;
+		if (!destinationFile.Close())
+			return false;
+		return true;
+	}
+
 	bool Session::CreateProcessJobObject()
 	{
 		#if PLATFORM_WINDOWS
diff --git a/Engine/Source/Programs/UnrealBuildAccelerator/Common/Private/UbaSessionClient.cpp b/Engine/Source/Programs/UnrealBuildAccelerator/Common/Private/UbaSessionClient.cpp
--- a/Engine/Source/Programs/UnrealBuildAccelerator/Common/Private/UbaSessionClient.cpp
+++ b/Engine/Source/Programs/UnrealBuildAccelerator/Common/Private/UbaSessionClient.cpp
@@ -651,7 +651,7 @@ namespace uba
 						retryEntryGuard.Execute();
 
 					MemoryMap map;
-					if (!CreateMemoryMapFromFile(map, fileNameKey, newName.data, IsCompressed(newCasKey), memoryMapAlignment, fileName.data, nullptr, false))
+					if (!GetOrCreateMemoryMapFromFile(map, fileNameKey, newName.data, IsCompressed(newCasKey), memoryMapAlignment, fileName.data, nullptr, false))
 					{
 						if (!shouldRetry)
 							return false;
diff --git a/Engine/Source/Programs/UnrealBuildAccelerator/Common/Private/UbaSessionServer.cpp b/Engine/Source/Programs/UnrealBuildAccelerator/Common/Private/UbaSessionServer.cpp
--- a/Engine/Source/Programs/UnrealBuildAccelerator/Common/Private/UbaSessionServer.cpp
+++ b/Engine/Source/Programs/UnrealBuildAccelerator/Common/Private/UbaSessionServer.cpp
@@ -62,6 +62,13 @@ namespace uba
 		virtual ProcessExecutionType GetExecutionType() const override { return ProcessExecutionType_Remote; }
 		virtual bool IsChild() override { return false; }
 
+		virtual void TraverseOutputFiles(const Function<void(StringView file)>& func) override
+		{
+			SCOPED_READ_LOCK(m_writtenFilesLock, lock);
+			for (auto& kv : m_writtenFiles)
+				func(kv.first);
+		}
+
 		void CallProcessExit(ProcessHandle& h)
 		{
 			SCOPED_FUTEX(m_exitedLock, lock);
@@ -96,6 +103,9 @@ namespace uba
 		KnownInput* m_knownInputs = nullptr;
 		u32 m_knownInputsCount = 0;
 		Event m_knownInputsDone;
+
+		ReaderWriterLock m_writtenFilesLock;
+		UnorderedMap<TString, CasKey> m_writtenFiles;
 	};
 
 	void SessionServerCreateInfo::Apply(const Config& config)
@@ -1092,117 +1102,119 @@ namespace uba
 		Storage::RetrieveResult res;
 		bool success = m_storage.RetrieveCasFile(res, casKey, destination.data, nullptr, 1, true, clientId);
 		casKey = res.casKey;
+
+		auto writeResponse = MakeGuard([&]() { writer.WriteBool(success); });
+
 		if (!success)
 		{
 			auto logType = connectionInfo.ShouldDisconnect() ? LogEntryType_Info : LogEntryType_Warning;
 			m_logger.Logf(logType, TC("Failed to retrieve cas for %s from client %u (Needed to write %s)"), CasKeyString(casKey).str, clientId, destination.data);
+			return true;
 		}
 
-		auto writeResponse = MakeGuard([&]() { writer.WriteBool(success); });
+		if (destination.StartsWith(TC("<log>")))
+		{
+			StringBuffer<> logPath;
+			logPath.Append(m_sessionLogDir).Append(destination.data + 5);
+			if (!m_storage.CopyOrLink(casKey, logPath.data, attributes))
+				m_logger.Error(TC("Failed to copy cas from %s to %s"), CasKeyString(casKey).str, logPath.data);
+			else if (!m_storage.DropCasFile(casKey, false, logPath.data))
+				m_logger.Error(TC("Failed to drop cas %s"), CasKeyString(casKey).str);
+			return true;
+		}
 
-		bool shouldWriteToDisk = ShouldWriteToDisk(destination);
-		if (success)
+		if (destination.StartsWith(TC("<uba>")))
 		{
-			if (destination.StartsWith(TC("<log>")))
-			{
-				StringBuffer<> logPath;
-				logPath.Append(m_sessionLogDir).Append(destination.data + 5);
-				if (!m_storage.CopyOrLink(casKey, logPath.data, attributes))
-					m_logger.Error(TC("Failed to copy cas from %s to %s"), CasKeyString(casKey).str, logPath.data);
-				else if (!m_storage.DropCasFile(casKey, false, logPath.data))
-					m_logger.Error(TC("Failed to drop cas %s"), CasKeyString(casKey).str);
-				return true;
-			}
+			StringBuffer<> ubaPath;
+			ubaPath.Append(m_sessionLogDir);
 
-			if (destination.StartsWith(TC("<uba>")))
-			{
-				StringBuffer<> ubaPath;
-				ubaPath.Append(m_sessionLogDir);
-
-				ClientSession* session = nullptr;
-				for (auto& s : m_clientSessions)
-					if (s->clientId == clientId)
-						session = s;
-				if (session)
-					ubaPath.Append(session->name);
-				else
-					ubaPath.Append(TCV("Connection")).AppendValue(clientId);
-
-				ubaPath.Append(TCV(".uba"));
-				m_storage.CopyOrLink(casKey, ubaPath.data, attributes);
-				m_storage.DropCasFile(casKey, false, ubaPath.data);
-				return true;
-			}
+			ClientSession* session = nullptr;
+			for (auto& s : m_clientSessions)
+				if (s->clientId == clientId)
+					session = s;
+			if (session)
+				ubaPath.Append(session->name);
+			else
+				ubaPath.Append(TCV("Connection")).AppendValue(clientId);
 
-			// We need to report before writing to flush out potential 
-			m_storage.ReportFileWrite(destinationKey, destination.data);
+			ubaPath.Append(TCV(".uba"));
+			m_storage.CopyOrLink(casKey, ubaPath.data, attributes);
+			m_storage.DropCasFile(casKey, false, ubaPath.data);
+			return true;
+		}
 
-			if (shouldWriteToDisk)
-			{
-				bool writeCompressed = false;
+		ProcessHandle h = GetProcess(processId);
+		if (!h.IsValid())
+		{
+			m_logger.Info(TC("Failed to find process for id %u when receiving SendFileToServer message"), processId);
+			success = false;
+			return false;
+		}
 
-				ProcessHandle h = GetProcess(processId);
-				if (!h.IsValid())
-				{
-					success = false;
-					m_logger.Info(TC("Failed to find process for id %u when receiving SendFileToServer message"), processId);
-					return false;
-				}
+		auto& process = *(RemoteProcess*)h.m_process;
+		auto& startInfo = process.GetStartInfo();
+		auto& rules = *startInfo.rules;
+
+		bool shouldWriteToDisk = ShouldWriteToDisk(destination);
 
-				RootsHandle rootsHandle = h.GetStartInfo().rootsHandle;
+		{
+			SCOPED_WRITE_LOCK(process.m_writtenFilesLock, lock);
+			process.m_writtenFiles.try_emplace(destination.ToString(), casKey);
+		}
 
-				auto& rules = *h.GetStartInfo().rules;
+		// We need to report before writing to flush out potential 
+		m_storage.ReportFileWrite(destinationKey, destination.data);
 
-				Storage::FormattingFunc formattingFunc;
-				bool escapeSpaces;
-				if (HasVfs(rootsHandle) && rules.ShouldDevirtualizeFile(destination, escapeSpaces))
-				{
-					formattingFunc = [&](MemoryBlock& destData, const void* sourceData, u64 sourceSize, const tchar* hint)
-						{
-							return DevirtualizeDepsFile(rootsHandle, destData, sourceData, sourceSize, escapeSpaces, hint);
-						};
-				}
-				else if (m_storeIntermediateFilesCompressed)
-				{
-					writeCompressed = g_globalRules.FileCanBeCompressed(destination);
-				}
-				success = m_storage.CopyOrLink(casKey, destination.data, attributes, writeCompressed, formattingFunc);
-				if (!success)
-					m_logger.Error(TC("Failed to copy cas from %s to %s (%s)"), CasKeyString(casKey).str, destination.data, GetProcessDescription(processId).c_str());
-				else
-					TraceWrittenFile(processId, destination);
+		if (shouldWriteToDisk)
+		{
+			bool writeCompressed = false;
+
+			RootsHandle rootsHandle = startInfo.rootsHandle;
+
+			Storage::FormattingFunc formattingFunc;
+			bool escapeSpaces;
+			if (HasVfs(rootsHandle) && rules.ShouldDevirtualizeFile(destination, escapeSpaces))
+			{
+				formattingFunc = [&](MemoryBlock& destData, const void* sourceData, u64 sourceSize, const tchar* hint)
+					{
+						return DevirtualizeDepsFile(rootsHandle, destData, sourceData, sourceSize, escapeSpaces, hint);
+					};
+			}
+			else if (m_storeIntermediateFilesCompressed)
+			{
+				writeCompressed = g_globalRules.FileCanBeCompressed(destination);
 			}
+
+			success = m_storage.CopyOrLink(casKey, destination.data, attributes, writeCompressed, formattingFunc);
+			if (!success)
+				m_logger.Error(TC("Failed to copy cas from %s to %s (%s)"), CasKeyString(casKey).str, destination.data, GetProcessDescription(processId).c_str());
 			else
+				TraceWrittenFile(processId, destination);
+		}
+		else
+		{
+			if (!rules.IsInvisible(destination))
 			{
 				success = m_storage.FakeCopy(casKey, destination.data);
 				if (!success)
 					m_logger.Error(TC("Failed to fake copy cas from %s to %s (%s)"), CasKeyString(casKey).str, destination.data, GetProcessDescription(processId).c_str());
-				SCOPED_WRITE_LOCK(m_receivedFilesLock, lock);
-				m_receivedFiles.try_emplace(destinationKey, casKey);
 			}
+			SCOPED_WRITE_LOCK(m_receivedFilesLock, lock);
+			m_receivedFiles.try_emplace(destinationKey, casKey);
 		}
 
-		if (success)
+		if (success && !rules.IsInvisible(destination))
 		{
 			bool invalidateStorage = false; // No need, already handled in m_storage.CopyOrLink
 			RegisterCreateFileForWrite(StringKeyZero, destination, shouldWriteToDisk, 0, 0, invalidateStorage);
 
-
-			SCOPED_FUTEX_READ(m_processesLock, lock);
-			auto findIt = m_processes.find(processId);
-			if (findIt != m_processes.end())
+			if (startInfo.trackInputs)
 			{
-				ProcessHandle h(findIt->second);
-				lock.Leave();
-				auto& process = *(RemoteProcess*)h.m_process;
-				if (process.m_startInfo.trackInputs)
-				{
-					u64 bytes = GetStringWriteSize(destination.data, destination.count);
-					u64 prevSize = process.m_trackedOutputs.size();
-					process.m_trackedOutputs.resize(prevSize + bytes);
-					BinaryWriter w2(process.m_trackedOutputs.data(), prevSize, prevSize + bytes);
-					w2.WriteString(destination);
-				}
+				u64 bytes = GetStringWriteSize(destination.data, destination.count);
+				u64 prevSize = process.m_trackedOutputs.size();
+				process.m_trackedOutputs.resize(prevSize + bytes);
+				BinaryWriter w2(process.m_trackedOutputs.data(), prevSize, prevSize + bytes);
+				w2.WriteString(destination);
 			}
 		}
 		return true;
@@ -2181,19 +2193,19 @@ namespace uba
 		return true;
 	}
 
-	bool SessionServer::CreateFile(CreateFileResponse& out, const CreateFileMessage& msg)
+	bool SessionServer::CreateFileForRead(CreateFileResponse& out, TrackWorkScope& tws, const StringView& fileName, const StringKey& fileNameKey, ProcessImpl& process, const ApplicationRules& rules)
 	{
-		if (!m_shouldWriteToDisk && ((msg.access & FileAccess_Write) == 0))
+		if (!m_shouldWriteToDisk)
 		{
 			SCOPED_READ_LOCK(m_receivedFilesLock, lock);
-			auto findIt = m_receivedFiles.find(msg.fileNameKey);
+			auto findIt = m_receivedFiles.find(fileNameKey);
 			if (findIt != m_receivedFiles.end())
 			{
-				u64 memoryMapAlignment = GetMemoryMapAlignment(msg.fileName);
+				u64 memoryMapAlignment = GetMemoryMapAlignment(fileName);
 				if (!memoryMapAlignment)
 					memoryMapAlignment = 4096;
 				MemoryMap map;
-				if (!CreateMemoryMapFromView(map, msg.fileNameKey, msg.fileName.data, findIt->second, memoryMapAlignment))
+				if (!GetOrCreateMemoryMapFromStorage(map, fileNameKey, fileName.data, findIt->second, memoryMapAlignment))
 					return false;
 				out.directoryTableSize = GetDirectoryTableSize();
 				out.mappedFileTableSize = GetFileMappingSize();
@@ -2202,7 +2214,68 @@ namespace uba
 				return true;
 			}
 		}
-		return Session::CreateFile(out, msg);
+		return Session::CreateFileForRead(out, tws, fileName, fileNameKey, process, rules);
+	}
+
+	bool SessionServer::GetOutputFileSizeInternal(u64& outSize, const StringKey& fileNameKey, StringView filePath)
+	{
+		if (Session::GetOutputFileSizeInternal(outSize, fileNameKey, filePath))
+			return true;
+		SCOPED_READ_LOCK(m_receivedFilesLock, lock);
+		auto findIt = m_receivedFiles.find(fileNameKey);
+		if (findIt == m_receivedFiles.end())
+			return false;
+		CasKey casKey = findIt->second;
+		lock.Leave();
+		return m_storage.GetFileSize(outSize, casKey, filePath.data);
+	}
+
+	bool SessionServer::GetOutputFileDataInternal(void* outData, const StringKey& fileNameKey, StringView filePath, bool deleteInternalMapping)
+	{
+		if (Session::GetOutputFileDataInternal(outData, fileNameKey, filePath, deleteInternalMapping))
+			return true;
+		SCOPED_READ_LOCK(m_receivedFilesLock, lock);
+		auto findIt = m_receivedFiles.find(fileNameKey);
+		if (findIt == m_receivedFiles.end())
+			return false;
+		CasKey casKey = findIt->second;
+		lock.Leave();
+		if (!m_storage.GetFileData(outData, casKey, filePath.data))
+			return false;
+		if (!deleteInternalMapping)
+			return true;
+
+		SCOPED_WRITE_LOCK(m_receivedFilesLock, lock2);
+		m_receivedFiles.erase(findIt);
+		lock2.Leave();
+
+		m_storage.DropCasFile(casKey, true, filePath.data);
+		return true;
+	}
+
+	bool SessionServer::WriteOutputFileInternal(const StringKey& fileNameKey, StringView filePath, bool deleteInternalMapping)
+	{
+		if (Session::WriteOutputFileInternal(fileNameKey, filePath, deleteInternalMapping))
+			return true;
+
+		SCOPED_READ_LOCK(m_receivedFilesLock, lock);
+		auto findIt = m_receivedFiles.find(fileNameKey);
+		if (findIt == m_receivedFiles.end())
+			return false;
+		CasKey casKey = findIt->second;
+		lock.Leave();
+
+		bool res = m_storage.CopyOrLink(casKey, filePath.data, DefaultAttributes());
+
+		if (!deleteInternalMapping)
+			return res;
+
+		SCOPED_WRITE_LOCK(m_receivedFilesLock, lock2);
+		m_receivedFiles.erase(findIt);
+		lock2.Leave();
+
+		m_storage.DropCasFile(casKey, true, filePath.data);
+		return res;
 	}
 
 	void SessionServer::FileEntryAdded(StringKey fileNameKey, u64 lastWritten, u64 size)
diff --git a/Engine/Source/Programs/UnrealBuildAccelerator/Common/Private/UbaStorage.cpp b/Engine/Source/Programs/UnrealBuildAccelerator/Common/Private/UbaStorage.cpp
--- a/Engine/Source/Programs/UnrealBuildAccelerator/Common/Private/UbaStorage.cpp
+++ b/Engine/Source/Programs/UnrealBuildAccelerator/Common/Private/UbaStorage.cpp
@@ -345,7 +345,15 @@ namespace uba
 				if (!fileMem)
 					return false;
 				out.readSize = externalMapping.writtenSize;
-				return WriteCompressedFile(out, fileName, InvalidFileHandle, fileMem, externalMapping.size, casFile, nullptr, 0);
+				out.isIndependent = externalMapping.createIndependentMapping;
+				if (externalMapping.createIndependentMapping)
+				{
+					CompressWriter w(m_logger, m_bufferSlots, m_workManager, Stats(), m_casCompressor, m_casCompressionLevel, m_asyncUnmapViewOfFile);
+					out.mappingOffset = 0;
+					return w.CompressToMapping(out.mappingHandle, out.size, fileMem, externalMapping.size, casFile);
+				}
+				else
+					return WriteCompressedFile(out, fileName, InvalidFileHandle, fileMem, externalMapping.size, casFile, nullptr, 0);
 			}
 		}
 
@@ -1026,8 +1034,8 @@ namespace uba
 	#endif
 
 	StorageImpl::StorageImpl(const StorageCreateInfo& info, const tchar* logPrefix)
-	:	m_workManager(info.workManager)
-	,	m_logger(info.writer, logPrefix)
+	:	m_logger(info.writer, logPrefix)
+	,	m_workManager(info.workManager)
 	,	m_maxParallelCopyOrLinkBottleneck(info.maxParallelCopyOrLink)
 	,	m_casDataBuffer(m_logger, &info.workManager)
 	{
@@ -1115,6 +1123,23 @@ namespace uba
 		return m_casLookup.try_emplace(casKey, casKey).first->second;
 	}
 
+	bool StorageImpl::GetFileEntry(FileEntry*& out, const tchar* file)
+	{
+		StringBuffer<> forKey;
+		if (!FixPath(file, nullptr, 0, forKey))
+			return m_logger.Error(TC("GetCasEntry - FixPath failed for %s"), file);
+		if (CaseInsensitiveFs)
+			forKey.MakeLower();
+		StringKey fileNameKey = ToStringKey(forKey);
+
+		SCOPED_READ_LOCK(m_fileTableLookupLock, lookupLock);
+		auto findIt = m_fileTableLookup.find(fileNameKey);
+		if (findIt == m_fileTableLookup.end())
+			return false;
+		out = &findIt->second;
+		return true;
+	}
+
 	bool StorageImpl::LoadCasTable(bool logStats, bool alwaysCheckAllFiles, bool* outWasTerminated)
 	{
 		if (!m_writeToDisk)
@@ -1365,6 +1390,12 @@ namespace uba
 						continue;
 					if (DeleteFileW(casFileName.data))
 						++deleteCount;
+					if (casEntry.mappingIsIndependent) // Log files etc could end up here
+					{
+						CloseFileMapping(m_logger, casEntry.mappingHandle, TC("SaveCasTable"));
+						casEntry.mappingHandle = {};
+						casEntry.mappingIsIndependent = false;
+					}
 					continue;
 				}
 
@@ -1989,10 +2020,21 @@ namespace uba
 		if (res.readSize != deferredFileSize) // TODO: Should also check last written time
 			m_logger.Warning(TC("Deferred cas %s was deferred with file size %llu but was read with size %llu (%s)"), CasKeyString(casKey).str, deferredFileSize, res.readSize, deferredCreation.data);
 
+		bool hasMapping = false;
+		if (res.mappingHandle.IsValid())
+		{
+			casEntry.mappingHandle = res.mappingHandle;
+			casEntry.mappingOffset = res.mappingOffset;
+			casEntry.mappingSize = res.size;
+			casEntry.mappingIsIndependent = res.isIndependent;
+			hasMapping = true;
+		}
+
 		casEntry.exists = true;
 		entryLock.Leave();
 
-		CasEntryWritten(casEntry, res.size);
+		if (!hasMapping)
+			CasEntryWritten(casEntry, res.size);
 		return true;
 	}
 
@@ -2115,23 +2157,36 @@ namespace uba
 
 		if (forceDelete)
 		{
-			StringBuffer<> casFile;
-			if (!StorageImpl::GetCasFileName(casFile, casKey))
-				return false;
-
 			u64 sizeDeleted = 0;
-			if (DeleteFileW(casFile.data) == 0)
+
+			if (casEntry.mappingHandle.IsValid())
 			{
-				u32 lastError = GetLastError();
-				if (lastError != ERROR_FILE_NOT_FOUND && lastError != ERROR_PATH_NOT_FOUND)
-					return m_logger.Error(TC("Failed to drop cas %s (%s) (%s)"), casFile.data, hint, LastErrorToText(lastError).data);
+				if (casEntry.mappingIsIndependent)
+				{
+					CloseFileMapping(m_logger, casEntry.mappingHandle, TC("DropCasFile"));
+					casEntry.mappingHandle = {};
+					casEntry.mappingIsIndependent = false;
+				}
 			}
 			else
 			{
-				m_casDroppedBytes += casEntry.size;
-				m_casDroppedCount++;
-				//m_logger.Debug(TC("Evicted %s from cache (%llukb)"), casFile, casEntry.size/1024);
-				sizeDeleted = casEntry.size;
+				StringBuffer<> casFile;
+				if (!StorageImpl::GetCasFileName(casFile, casKey))
+					return false;
+
+				if (DeleteFileW(casFile.data) == 0)
+				{
+					u32 lastError = GetLastError();
+					if (lastError != ERROR_FILE_NOT_FOUND && lastError != ERROR_PATH_NOT_FOUND)
+						return m_logger.Error(TC("Failed to drop cas %s (%s) (%s)"), casFile.data, hint, LastErrorToText(lastError).data);
+				}
+				else
+				{
+					m_casDroppedBytes += casEntry.size;
+					m_casDroppedCount++;
+					//m_logger.Debug(TC("Evicted %s from cache (%llukb)"), casFile, casEntry.size/1024);
+					sizeDeleted = casEntry.size;
+				}
 			}
 			casEntry.verified = true;
 			casEntry.exists = false;
@@ -2225,7 +2280,7 @@ namespace uba
 				if (casEntry->mappingHandle.IsValid())
 				{
 					casFile.Append(CasKeyString(actualKey).str);
-					mappedView = m_casDataBuffer.MapView(casEntry->mappingHandle, casEntry->mappingOffset, casEntry->mappingSize, casFile.data);
+					mappedView = m_casDataBuffer.MapView(casEntry->mappingHandle, casEntry->mappingOffset, casEntry->mappingSize, casFile.data, casEntry->mappingIsIndependent);
 					compressedData = mappedView.memory;
 					if (!compressedData)
 						return m_logger.Error(TC("[CopyOrLink] Failed to map view of mapping %s (%s)"), casFile.data, LastErrorToText().data);
@@ -2449,6 +2504,7 @@ namespace uba
 		fileEntry.lastWritten = lastWritten;
 		fileEntry.size = size;
 		fileEntry.verified = true;
+		fileEntry.isTemp = true;
 		return true;
 	}
 
@@ -2668,23 +2724,57 @@ namespace uba
 
 	bool StorageImpl::DeleteCasForFile(const tchar* file)
 	{
-		StringBuffer<> forKey;
-		FixPath(file, nullptr, 0, forKey);
-		if (CaseInsensitiveFs)
-			forKey.MakeLower();
-		StringKey fileNameKey = ToStringKey(forKey);
-
-		SCOPED_READ_LOCK(m_fileTableLookupLock, lookupLock);
-		auto findIt = m_fileTableLookup.find(fileNameKey);
-		if (findIt == m_fileTableLookup.end())
+		FileEntry* fileEntry;
+		if (!GetFileEntry(fileEntry, file))
 			return false;
-		FileEntry& fileEntry = findIt->second;
-		lookupLock.Leave();
 
-		SCOPED_FUTEX(fileEntry.lock, entryLock);
-		fileEntry.verified = false;
+		SCOPED_FUTEX(fileEntry->lock, entryLock);
+		fileEntry->verified = false;
+		return DropCasFile(fileEntry->casKey, true, file);
+	}
 
-		return DropCasFile(fileEntry.casKey, true, file);
+	bool StorageImpl::GetFileSize(u64& outSize, const CasKey& casKey, const tchar* hint)
+	{
+		SCOPED_READ_LOCK(m_casLookupLock, lookupLock);
+		auto findIt = m_casLookup.find(casKey);
+		if  (findIt == m_casLookup.end())
+			return m_logger.Error(TC("GetCasEntry - Couldn't find cas entry for %s"), hint);
+		CasEntry& casEntry = findIt->second;
+		SCOPED_READ_LOCK(casEntry.lock, casEntryLock);
+		if (!casEntry.mappingHandle.IsValid())
+			return m_logger.Error(TC("Code path only supports mapped cas entries"));
+		if (!IsCompressed(casEntry.key))
+		{
+			outSize = casEntry.mappingSize;
+			return true;
+		}
+		FileMappingScope scope(m_logger);
+		u8* fileMem = scope.MapView(TC("ReadFileSize"), casEntry.mappingHandle, casEntry.mappingOffset, casEntry.mappingSize, hint);
+		if (!fileMem)
+			return false;
+		outSize = *(u64*)fileMem;
+		return true;
 	}
 
+	bool StorageImpl::GetFileData(void* outData, const CasKey& casKey, const tchar* hint)
+	{
+		SCOPED_READ_LOCK(m_casLookupLock, lookupLock);
+		auto findIt = m_casLookup.find(casKey);
+		if  (findIt == m_casLookup.end())
+			return m_logger.Error(TC("GetCasEntry - Couldn't find cas entry for %s"), hint);
+		CasEntry& casEntry = findIt->second;
+		SCOPED_READ_LOCK(casEntry.lock, casEntryLock);
+		FileMappingScope scope(m_logger);
+		u8* fileMem = scope.MapView(TC("ReadFileSize"), casEntry.mappingHandle, casEntry.mappingOffset, casEntry.mappingSize, hint);
+		if (!fileMem)
+			return false;
+		if (!IsCompressed(casEntry.key))
+		{
+			memcpy(outData, fileMem, casEntry.mappingSize);
+			return true;
+		}
+		u64 decompressedSize = *(u64*)fileMem;
+		fileMem += sizeof(u64);
+		return DecompressMemoryToMemory(fileMem, casEntry.mappingSize, (u8*)outData, decompressedSize, CasKeyString(casEntry.key).str, hint);
+	}
 }
diff --git a/Engine/Source/Programs/UnrealBuildAccelerator/Common/Private/UbaStorageServer.cpp b/Engine/Source/Programs/UnrealBuildAccelerator/Common/Private/UbaStorageServer.cpp
--- a/Engine/Source/Programs/UnrealBuildAccelerator/Common/Private/UbaStorageServer.cpp
+++ b/Engine/Source/Programs/UnrealBuildAccelerator/Common/Private/UbaStorageServer.cpp
@@ -17,6 +17,7 @@ namespace uba
 		const ConfigTable& table = *tablePtr;
 		table.GetValueAsBool(writeReceivedCasFilesToDisk, TC("WriteReceivedCasFilesToDisk"));
 		table.GetValueAsBool(allowHintAsFallback, TC("AllowHintAsFallback"));
+		table.GetValueAsBool(createIndependentMappings, TC("CreateIndependentMappings"));
 	}
 
 	StorageServer::StorageServer(const StorageServerCreateInfo& info)
@@ -26,6 +27,7 @@ namespace uba
 		m_zone = info.zone;
 		m_allowHintAsFallback = info.allowHintAsFallback;
 		m_writeReceivedCasFilesToDisk = info.writeReceivedCasFilesToDisk;
+		m_createIndependentMappings = info.createIndependentMappings;
 
 		if (!CreateGuid(m_uid))
 			UBA_ASSERT(false);
@@ -150,7 +152,7 @@ namespace uba
 					}, 1, TC("ActiveFetchRelease"));
 			}
 			else
-				server.m_casDataBuffer.UnmapView(mappedView, TC("OnDisconnected"));
+				server.m_casDataBuffer.UnmapView(mappedView, reason);
 		}
 		else
 		{
@@ -205,6 +207,8 @@ namespace uba
 				}
 
 				m_casDataBuffer.UnmapView(store.mappedView, TC("OnDisconnected"));
+				if (store.mappedView.isIndependent)
+					CloseFileMapping(m_logger, store.mappedView.handle, TC("StoreDone"));
 				it = m_activeStores.erase(it);
 			}
 		}
@@ -573,7 +577,7 @@ namespace uba
 				bool useFileMapping = casEntry->mappingHandle.IsValid();
 				if (useFileMapping)
 				{
-					mappedView = m_casDataBuffer.MapView(casEntry->mappingHandle, casEntry->mappingOffset, casEntry->mappingSize, CasKeyString(casKey).str);
+					mappedView = m_casDataBuffer.MapView(casEntry->mappingHandle, casEntry->mappingOffset, casEntry->mappingSize, CasKeyString(casKey).str, casEntry->mappingIsIndependent);
 					memoryBegin = mappedView.memory;
 					fileSize = casEntry->mappingSize;
 					if (!memoryBegin)
@@ -899,7 +903,7 @@ namespace uba
 				}
 				else
 				{
-					mappedView = m_casDataBuffer.AllocAndMapView(MappedView_Transient, fileSize, 1, CasKeyString(casKey).str);
+					mappedView = m_casDataBuffer.AllocAndMapView(MappedView_Transient, fileSize, 1, CasKeyString(casKey).str, false, m_createIndependentMappings);
 					if (!mappedView.memory)
 					{
 						casEntry.verified = false;
@@ -982,6 +986,7 @@ namespace uba
 						casEntry.mappingHandle = activeStore.mappedView.handle;
 						casEntry.mappingOffset = activeStore.mappedView.offset;
 						casEntry.mappingSize = totalWritten;
+						casEntry.mappingIsIndependent = activeStore.mappedView.isIndependent;
 						casEntry.exists = true;
 						casEntry.beingWritten = false;
 					}
diff --git a/Engine/Source/Programs/UnrealBuildAccelerator/Common/Public/UbaProcess.h b/Engine/Source/Programs/UnrealBuildAccelerator/Common/Public/UbaProcess.h
--- a/Engine/Source/Programs/UnrealBuildAccelerator/Common/Public/UbaProcess.h
+++ b/Engine/Source/Programs/UnrealBuildAccelerator/Common/Public/UbaProcess.h
@@ -28,6 +28,7 @@ namespace uba
 		virtual ProcessExecutionType GetExecutionType() const { UBA_ASSERT(false); return ProcessExecutionType_Native; }
 		virtual bool IsChild() { UBA_ASSERT(false); return false; }
 		virtual bool IsArm() { UBA_ASSERT(false); return false; }
+		virtual void TraverseOutputFiles(const Function<void(StringView file)>& func) {}
 
 	protected:
 		void AddRef();
@@ -57,6 +58,7 @@ namespace uba
 		virtual ProcessExecutionType GetExecutionType() const override { return m_detourEnabled ? ProcessExecutionType_Local : ProcessExecutionType_Native; }
 		virtual bool IsChild() override { return m_parentProcess != nullptr; }
 		virtual bool IsArm() override { return m_isArmBinary; }
+		virtual void TraverseOutputFiles(const Function<void(StringView file)>& func) override;
 
 		ProcessImpl(Session& session, u32 id, ProcessImpl* parent, bool detourEnabled);
 		~ProcessImpl();
diff --git a/Engine/Source/Programs/UnrealBuildAccelerator/Common/Public/UbaProcessHandle.h b/Engine/Source/Programs/UnrealBuildAccelerator/Common/Public/UbaProcessHandle.h
--- a/Engine/Source/Programs/UnrealBuildAccelerator/Common/Public/UbaProcessHandle.h
+++ b/Engine/Source/Programs/UnrealBuildAccelerator/Common/Public/UbaProcessHandle.h
@@ -49,6 +49,8 @@ namespace uba
 		bool IsRemote() const;								// Returns true if process is a remote process
 		ProcessExecutionType GetExecutionType() const;		// Returns execution type
 
+		void TraverseOutputFiles(const Function<void(StringView file)>& func) const; // Traverse output files.
+
 		ProcessHandle();
 		ProcessHandle(const ProcessHandle& o);
 		ProcessHandle(ProcessHandle&& o) noexcept;
diff --git a/Engine/Source/Programs/UnrealBuildAccelerator/Common/Public/UbaSession.h b/Engine/Source/Programs/UnrealBuildAccelerator/Common/Public/UbaSession.h
--- a/Engine/Source/Programs/UnrealBuildAccelerator/Common/Public/UbaSession.h
+++ b/Engine/Source/Programs/UnrealBuildAccelerator/Common/Public/UbaSession.h
@@ -65,6 +65,8 @@ namespace uba
 		bool RefreshDirectory(const tchar* dirName, bool forceRegister = false); // Tell uba a directory on disk has been changed by some other system while session is running
 		bool RegisterNewFile(const tchar* filePath); // Tell uba a new file on disk has been added by some other system while session is running
 		bool RegisterVirtualFile(const tchar* filePath, const tchar* sourceFile, u64 sourceOffset, u64 sourceSize);
+		bool CreateVirtualFile(const tchar* filePath, const void* memory, u64 memorySize);
+		bool DeleteVirtualFile(const tchar* filePath);
 		void RegisterDeleteFile(const tchar* filePath); // Tell uba a file on disk has been deleted by some other system while session is running
 		bool RegisterNewDirectory(const tchar* directoryPath); // Tell uba a directory on disk has been added by some other system while session is running
 
@@ -76,6 +78,11 @@ namespace uba
 		using GetNextProcessFunction = Function<bool(Process& handle, NextProcessInfo& outNextProcess, u32 prevExitCode)>;
 		void RegisterGetNextProcess(GetNextProcessFunction&& function); // Register a custom service (that can be communicated with from the remote agents)
 
+		bool GetOutputFileSize(u64& outSize, const tchar* filePath); // Available if ShouldWriteToDisk is false
+		bool GetOutputFileData(void* outData, const tchar* filePath, bool deleteInternalMapping); // Available if ShouldWriteToDisk is false
+		bool WriteOutputFile(const tchar* filePath, bool deleteInternalMapping); // Available if ShouldWriteToDisk is false
+
+
 		const tchar* GetId(); // Id for session. Will be "yymmdd_hhmmss" unless SessionCreateInfo.useUniqueId is set to false
 		u32 GetActiveProcessCount(); // Current active processes running inside session
 		Storage& GetStorage(); // Storage (only used when remote machines are connected)
@@ -136,8 +143,8 @@ namespace uba
 		
 		struct MemoryMap { StringBuffer<128> name; u64 size = 0; };
 		struct FileMappingEntry;
-		bool CreateMemoryMapFromFile(MemoryMap& out, StringKey fileNameKey, const tchar* fileName, bool isCompressedCas, u64 alignment, const tchar* hint, ProcessImpl* requestingProcess, bool canBeFreed);
-		bool CreateMemoryMapFromView(MemoryMap& out, StringKey fileNameKey, const tchar* fileName, const CasKey& casKey, u64 alignment);
+		bool GetOrCreateMemoryMapFromFile(MemoryMap& out, StringKey fileNameKey, const tchar* fileName, bool isCompressedCas, u64 alignment, const tchar* hint, ProcessImpl* requestingProcess, bool canBeFreed);
+		bool GetOrCreateMemoryMapFromStorage(MemoryMap& out, StringKey fileNameKey, const tchar* fileName, const CasKey& casKey, u64 alignment);
 
 		bool RegisterCreateFileForWrite(StringKey fileNameKey, const StringView& fileName, bool registerRealFile, u64 fileSize = 0, u64 lastWriteTime = 0, bool invalidateStorage = true);
 		u32 RegisterDeleteFile(StringKey fileNameKey, const StringView& fileName);
@@ -147,6 +154,11 @@ namespace uba
 		virtual void PrintSessionStats(Logger& logger);
 
 		virtual bool RegisterVirtualFileInternal(const StringKey& fileNameKey, const StringView& filePath, const tchar* sourceFile, u64 sourceOffset, u64 sourceSize);
+		virtual bool CreateVirtualFileInternal(const StringKey& fileNameKey, const StringView& filePath, const void* memory, u64 memorySize);
+		virtual bool RegisterVirtualFileInternal(const StringKey& fileNameKey, const StringView& filePath, FileMappingHandle mappingHandle, u64 mappingSize, u64 mappingOffset);
+		virtual bool GetOutputFileSizeInternal(u64& outSize, const StringKey& fileNameKey, StringView filePath);
+		virtual bool GetOutputFileDataInternal(void* outData, const StringKey& fileNameKey, StringView filePath, bool deleteInternalMapping);
+		virtual bool WriteOutputFileInternal(const StringKey& fileNameKey, StringView filePath, bool deleteInternalMapping);
 
 		virtual bool ProcessThreadStart(ProcessImpl& process);
 		virtual bool GetInitResponse(InitResponse& out, const InitMessage& msg);
@@ -180,6 +192,8 @@ namespace uba
 
 		void RemoveWrittenFile(ProcessImpl& process, const StringKey& fileKey);
 		bool WriteFileToDisk(ProcessImpl& process, WrittenFile& file);
+		bool WriteMemoryToDisk(FileAccessor& destinationFile, const void* fileMem, u64 fileSize);
+		bool GetFileMemory(const Function<bool(const void*, u64)>& func, const StringKey& fileNameKey, StringView filePath, bool deleteInternalMapping);
 		void AddEnvironmentVariableNoLock(const tchar* key, const tchar* value);
 		bool WriteDirectoryEntriesInternal(DirectoryTable::Directory& dir, const StringKey& dirKey, StringView dirPath, bool isRefresh, u32& outTableOffset);
 		void WriteDirectoryEntriesRecursive(const StringKey& dirKey, StringView dirPath, u32& outTableOffset);
@@ -268,10 +282,13 @@ namespace uba
 			u16 refCount = 0;
 			u8 usedCount = 0;
 			u8 usedCountBeforeFree = 0;
-			bool isDir = false;
-			bool handled = false;
-			bool success = false;
-			bool canBeFreed = false;
+
+			u8 isDir:1 = false;
+			u8 handled:1 = false;
+			u8 success:1 = false;
+			u8 canBeFreed:1 = false;
+			u8 createIndependentMapping:1 = false;
+			u8 isInvisible:1 = false;
 
 			#if UBA_DEBUG_TRACK_MAPPING
 			TString name;
diff --git a/Engine/Source/Programs/UnrealBuildAccelerator/Common/Public/UbaSessionServer.h b/Engine/Source/Programs/UnrealBuildAccelerator/Common/Public/UbaSessionServer.h
--- a/Engine/Source/Programs/UnrealBuildAccelerator/Common/Public/UbaSessionServer.h
+++ b/Engine/Source/Programs/UnrealBuildAccelerator/Common/Public/UbaSessionServer.h
@@ -106,7 +106,10 @@ namespace uba
 		TString GetProcessDescription(u32 processId);
 
 		virtual bool PrepareProcess(ProcessImpl& process, bool isChild, StringBufferBase& outRealApplication, const tchar*& outRealWorkingDir) override final;
-		virtual bool CreateFile(CreateFileResponse& out, const CreateFileMessage& msg) override final;
+		virtual bool CreateFileForRead(CreateFileResponse& out, TrackWorkScope& tws, const StringView& fileName, const StringKey& fileNameKey, ProcessImpl& process, const ApplicationRules& rules) override final;
+		virtual bool GetOutputFileSizeInternal(u64& outSize, const StringKey& fileNameKey, StringView filePath) override;
+		virtual bool GetOutputFileDataInternal(void* outData, const StringKey& fileNameKey, StringView filePath, bool deleteInternalMapping) override;
+		virtual bool WriteOutputFileInternal(const StringKey& fileNameKey, StringView filePath, bool deleteInternalMapping) override;
 		virtual void FileEntryAdded(StringKey fileNameKey, u64 lastWritten, u64 size) override final;
 		virtual bool RunSpecialProgram(ProcessImpl& process, BinaryReader& reader, BinaryWriter& writer) override final;
 		virtual void PrintSessionStats(Logger& logger) override final;
diff --git a/Engine/Source/Programs/UnrealBuildAccelerator/Common/Public/UbaStorage.h b/Engine/Source/Programs/UnrealBuildAccelerator/Common/Public/UbaStorage.h
--- a/Engine/Source/Programs/UnrealBuildAccelerator/Common/Public/UbaStorage.h
+++ b/Engine/Source/Programs/UnrealBuildAccelerator/Common/Public/UbaStorage.h
@@ -36,6 +36,8 @@ namespace uba
 		virtual bool DecompressMemoryToMemory(const u8* compressedData, u64 compressedSize, u8* writeData, u64 decompressedSize, const tchar* readHint, const tchar* writeHint) = 0;
 		virtual bool CreateDirectory(const tchar* dir) = 0;
 		virtual bool DeleteCasForFile(const tchar* file) = 0;
+		virtual bool GetFileSize(u64& outSize, const CasKey& casKey, const tchar* hint) = 0;
+		virtual bool GetFileData(void* outData, const CasKey& casKey, const tchar* hint) = 0;
 
 		struct RetrieveResult { CasKey casKey; u64 size = 0; MappedView view; };
 		virtual bool RetrieveCasFile(RetrieveResult& out, const CasKey& casKey, const tchar* hint, FileMappingBuffer* mappingBuffer = nullptr, u64 memoryMapAlignment = 1, bool allowProxy = true, u32 clientId = 0) = 0;
@@ -71,14 +73,14 @@ namespace uba
 		struct WriteResult
 		{
 			FileMappingHandle mappingHandle;
-			u64 readSize = InvalidValue;
+			u64 mappingOffset = InvalidValue;
 			u64 size = InvalidValue;
-			u64 offset = InvalidValue;
+			u64 readSize = InvalidValue;
 			bool isIndependent = false;
 		};
 		virtual bool WriteCompressedFile(WriteResult& out, const tchar* from, FileHandle readHandle, u8* readMem, u64 fileSize, const tchar* toFile, const void* header, u64 headerSize, u64 lastWriteTime = 0) = 0;
 
-		struct ExternalFileMapping { FileMappingHandle handle; u64 offset = 0; u64 size = 0; u64 lastWriteTime = 0; u64 writtenSize = InvalidValue; };
+		struct ExternalFileMapping { FileMappingHandle handle; u64 offset = 0; u64 size = 0; u64 lastWriteTime = 0; u64 writtenSize = InvalidValue; bool createIndependentMapping = false; };
 		using ExternalFileMappingsProvider = Function<bool(ExternalFileMapping& out, StringKey fileNameKey, const tchar* fileName)>;
 		virtual void RegisterExternalFileMappingsProvider(ExternalFileMappingsProvider&& provider) {}
 	};
@@ -143,6 +145,8 @@ namespace uba
 		virtual bool DecompressFileToMemory(const tchar* fileName, FileHandle fileHandle, u8* dest, u64 decompressedSize, const tchar* writeHint, u64 fileStartOffset) override;
 		virtual bool CreateDirectory(const tchar* dir) override;
 		virtual bool DeleteCasForFile(const tchar* file) override;
+		virtual bool GetFileSize(u64& outSize, const CasKey& casKey, const tchar* hint) override;
+		virtual bool GetFileData(void* outData, const CasKey& casKey, const tchar* hint) override;
 		virtual bool RetrieveCasFile(RetrieveResult& out, const CasKey& casKey, const tchar* hint, FileMappingBuffer* mappingBuffer = nullptr, u64 memoryMapAlignment = 1, bool allowProxy = true, u32 clientId = 0) override;
 		virtual bool VerifyAndGetCachedFileInfo(CachedFileInfo& out, StringKey fileNameKey, u64 verifiedLastWriteTime, u64 verifiedSize) override;
 		virtual bool InvalidateCachedFileInfo(StringKey fileNameKey) override;
@@ -201,6 +205,7 @@ namespace uba
 		struct FileEntry;
 		FileEntry& GetOrCreateFileEntry(const StringKey& fileNameKey);
 		CasEntry& GetOrCreateCasEntry(const CasKey& casKey);
+		bool GetFileEntry(FileEntry*& out, const tchar* file);
 
 		MutableLogger m_logger;
 		WorkManager& m_workManager;
@@ -237,6 +242,7 @@ namespace uba
 			bool beingWritten = false; // This is set while file is being written (when coming from network)..
 			bool disallowed = false; // This is set if cas is created from disallowed file
 
+			bool mappingIsIndependent = false; // if true, mappingHandle is stand-alone and must be manually deleted from the outside using DropCasFile
 			FileMappingHandle mappingHandle;
 			u64 mappingOffset = 0;
 			u64 mappingSize = 0;
diff --git a/Engine/Source/Programs/UnrealBuildAccelerator/Common/Public/UbaStorageServer.h b/Engine/Source/Programs/UnrealBuildAccelerator/Common/Public/UbaStorageServer.h
--- a/Engine/Source/Programs/UnrealBuildAccelerator/Common/Public/UbaStorageServer.h
+++ b/Engine/Source/Programs/UnrealBuildAccelerator/Common/Public/UbaStorageServer.h
@@ -16,6 +16,7 @@ namespace uba
 		NetworkServer& server;
 		bool allowHintAsFallback = true; // Will fallback to file system to recalculate cas if cas content does not exist anymore
 		bool writeReceivedCasFilesToDisk = false; // Write received files to disk. This is only useful if we expect same file to be written multiple times (requires deterministic outputs ofc)
+		bool createIndependentMappings = false; // Set to true if we want every store to end up in its own file mapping. This will use more memory but makes it possible to delete files
 		const tchar* zone = TC("");
 	};
 
@@ -141,5 +142,6 @@ namespace uba
 
 		bool m_allowHintAsFallback;
 		bool m_writeReceivedCasFilesToDisk;
+		bool m_createIndependentMappings;
 	};
 }
diff --git a/Engine/Source/Programs/UnrealBuildAccelerator/Core/Private/UbaApplicationRules.cpp b/Engine/Source/Programs/UnrealBuildAccelerator/Core/Private/UbaApplicationRules.cpp
--- a/Engine/Source/Programs/UnrealBuildAccelerator/Core/Private/UbaApplicationRules.cpp
+++ b/Engine/Source/Programs/UnrealBuildAccelerator/Core/Private/UbaApplicationRules.cpp
@@ -668,6 +668,11 @@ namespace uba
 		{
 			return true;
 		}
+
+		virtual bool IsInvisible(const StringView& fileName) const override
+		{
+			return fileName.EndsWith(TCV(".out"));
+		}
 	};
 
 	class ApplicationRulesUbaObjTool : public ApplicationRules
@@ -684,6 +689,19 @@ namespace uba
 		}
 	};
 
+	class ApplicationRulesUbaTestApp : public ApplicationRules
+	{
+		virtual bool IsOutputFile(const StringView& fileName) const override
+		{
+			return fileName.EndsWith(TCV(".out"));
+		}
+
+		virtual bool IsInvisible(const StringView& fileName) const override
+		{
+			return fileName.EndsWith(TCV(".out"));
+		}
+	};
+
 	template<typename T, size_t size>
 	constexpr u32 ApplicationHash(T (&buffer)[size])
 	{
@@ -738,6 +756,7 @@ namespace uba
 			{ ApplicationHash(TC("ubaobjtool.exe")),			new ApplicationRulesUbaObjTool() },
 			{ ApplicationHash(TC("shadercompileworker.exe")),	new ApplicationRulesShaderCompileWorker() },
 			{ ApplicationHash(TC("instr-clang-cl.exe")),		new ApplicationRulesClangClExe() },
+			{ ApplicationHash(TC("ubatestapp.exe")),			new ApplicationRulesUbaTestApp() },
 			//{ L"MSBuild.dll"),				new ApplicationRules() },
 			//{ L"BreakpadSymbolEncoder.exe"),	new ApplicationRulesClang() },
 			//{ L"cmd.exe"),		new ApplicationRules() },
@@ -746,6 +765,7 @@ namespace uba
 			{ ApplicationHash(TC("ld.lld")),					new ApplicationRulesLdLLdExe() },
 			{ ApplicationHash(TC("ispc")),						new ApplicationRulesISPCExe() },
 			{ ApplicationHash(TC("shadercompileworker")),		new ApplicationRulesShaderCompileWorker() },
+			{ ApplicationHash(TC("ubatestapp")),				new ApplicationRulesUbaTestApp() },
 #endif
 			{ 0, nullptr }
 		};
diff --git a/Engine/Source/Programs/UnrealBuildAccelerator/Core/Public/UbaApplicationRules.h b/Engine/Source/Programs/UnrealBuildAccelerator/Core/Public/UbaApplicationRules.h
--- a/Engine/Source/Programs/UnrealBuildAccelerator/Core/Public/UbaApplicationRules.h
+++ b/Engine/Source/Programs/UnrealBuildAccelerator/Core/Public/UbaApplicationRules.h
@@ -75,6 +75,12 @@ namespace uba
 			return false;
 		}
 
+		// Useful when sending files back that is never read by any processes spawned by uba (keep internal tables smaller)
+		virtual bool IsInvisible(const StringView& fileName) const
+		{
+			return false;
+		}
+
 		// If returns false this means that all GetFileAttribute etc will return file-not-found on this file
 		virtual bool CanExist(const tchar* file) const
 		{
diff --git a/Engine/Source/Programs/UnrealBuildAccelerator/Host/Private/UbaExports.cpp b/Engine/Source/Programs/UnrealBuildAccelerator/Host/Private/UbaExports.cpp
--- a/Engine/Source/Programs/UnrealBuildAccelerator/Host/Private/UbaExports.cpp
+++ b/Engine/Source/Programs/UnrealBuildAccelerator/Host/Private/UbaExports.cpp
@@ -365,7 +365,7 @@ extern "C"
 		return handle->GetExecutionType();
 	}
 
-	const uba::tchar* ProcessHandle_GetExecutingHost(uba::ProcessHandle* handle)
+	const uba::tchar* ProcessHandle_GetExecutingHost(const uba::ProcessHandle* handle)
 	{
 		return handle->GetExecutingHost();
 	}
@@ -377,44 +377,49 @@ extern "C"
 		return lines[index].text.c_str();
 	}
 
-	uba::u64 ProcessHandle_GetHash(uba::ProcessHandle* handle)
+	uba::u64 ProcessHandle_GetHash(const uba::ProcessHandle* handle)
 	{
 		return handle->GetHash();
 	}
 
 	// 100ns ticks
-	uba::u64 ProcessHandle_GetTotalProcessorTime(uba::ProcessHandle* handle)
+	uba::u64 ProcessHandle_GetTotalProcessorTime(const uba::ProcessHandle* handle)
 	{
 		return uba::TimeToTick(handle->GetTotalProcessorTime());
 	}
 
 	// 100ns ticks
-	uba::u64 ProcessHandle_GetTotalWallTime(uba::ProcessHandle* handle)
+	uba::u64 ProcessHandle_GetTotalWallTime(const uba::ProcessHandle* handle)
 	{
 		return uba::TimeToTick(handle->GetTotalWallTime());
 	}
 
-	bool ProcessHandle_WaitForExit(uba::ProcessHandle* handle, uba::u32 millisecondsTimeout)
+	void ProcessHandle_TraverseOutputFiles(const uba::ProcessHandle* handle, ProcessHandle_OutputFileCallback* callback, void* userData)
+	{
+		handle->TraverseOutputFiles([&](uba::StringView file) { callback(userData, file.data); });
+	}
+
+	bool ProcessHandle_WaitForExit(const uba::ProcessHandle* handle, uba::u32 millisecondsTimeout)
 	{
 		return handle->WaitForExit(millisecondsTimeout);
 	}
 
-	void ProcessHandle_Cancel(uba::ProcessHandle* handle, bool terminate)
+	void ProcessHandle_Cancel(const uba::ProcessHandle* handle, bool terminate)
 	{
 		handle->Cancel(terminate);
 	}
 
-	void ProcessHandle_Destroy(uba::ProcessHandle* handle)
+	void ProcessHandle_Destroy(const uba::ProcessHandle* handle)
 	{
 		delete handle;
 	}
 
-	void DestroyProcessHandle(uba::ProcessHandle* handle)
+	void DestroyProcessHandle(const uba::ProcessHandle* handle)
 	{
 		delete handle;
 	}
 
-	const uba::ProcessStartInfo* Process_GetStartInfo(uba::Process& process)
+	const uba::ProcessStartInfo* Process_GetStartInfo(const uba::Process& process)
 	{
 		return &process.GetStartInfo();
 	}
@@ -483,6 +488,26 @@ extern "C"
 	{
 		return server->RegisterVirtualFile(filename, sourceFile, sourceOffset, sourceSize);
 	}
+	bool SessionServer_CreateVirtualFile(uba::SessionServer* server, const uba::tchar* filename, const void* memory, uba::u64 memorySize)
+	{
+		return server->CreateVirtualFile(filename, memory, memorySize);
+	}
+	bool SessionServer_DeleteVirtualFile(uba::SessionServer* server, const uba::tchar* filename)
+	{
+		return server->DeleteVirtualFile(filename);
+	}
+	bool SessionServer_GetOutputFileSize(uba::SessionServer* server, uba::u64& outSize, const uba::tchar* filename)
+	{
+		return server->GetOutputFileSize(outSize, filename);
+	}
+	bool SessionServer_GetOutputFileData(uba::SessionServer* server, void* outData, const uba::tchar* filename, bool deleteInternalMapping)
+	{
+		return server->GetOutputFileData(outData, filename, deleteInternalMapping);
+	}
+	bool SessionServer_WriteOutputFile(uba::SessionServer* server, const uba::tchar* filename, bool deleteInternalMapping)
+	{
+		return server->WriteOutputFile(filename, deleteInternalMapping);
+	}
 	uba::ProcessHandle* SessionServer_RunProcess(uba::SessionServer* server, uba::ProcessStartInfo& info, bool async, bool enableDetour)
 	{
 		return new uba::ProcessHandle(server->RunProcess(info, async, enableDetour));
diff --git a/Engine/Source/Programs/UnrealBuildAccelerator/Host/Public/UbaExports.h b/Engine/Source/Programs/UnrealBuildAccelerator/Host/Public/UbaExports.h
--- a/Engine/Source/Programs/UnrealBuildAccelerator/Host/Public/UbaExports.h
+++ b/Engine/Source/Programs/UnrealBuildAccelerator/Host/Public/UbaExports.h
@@ -86,17 +86,19 @@ extern "C"
 	UBA_API void ProcessStartInfo_Destroy(uba::ProcessStartInfo* info);
 
 	// ProcessHandle
+	using ProcessHandle_OutputFileCallback = void(void* userData, const uba::tchar* file);
 	UBA_API uba::u32 ProcessHandle_GetExitCode(const uba::ProcessHandle* handle);
 	UBA_API uba::u8 ProcessHandle_GetExecutionType(const uba::ProcessHandle* handle);
-	UBA_API const uba::tchar* ProcessHandle_GetExecutingHost(uba::ProcessHandle* handle);
+	UBA_API const uba::tchar* ProcessHandle_GetExecutingHost(const uba::ProcessHandle* handle);
 	UBA_API const uba::tchar* ProcessHandle_GetLogLine(const uba::ProcessHandle* handle, uba::u32 index);
-	UBA_API uba::u64 ProcessHandle_GetHash(uba::ProcessHandle* handle);
-	UBA_API uba::u64 ProcessHandle_GetTotalProcessorTime(uba::ProcessHandle* handle);
-	UBA_API uba::u64 ProcessHandle_GetTotalWallTime(uba::ProcessHandle* handle);
-	UBA_API bool ProcessHandle_WaitForExit(uba::ProcessHandle* handle, uba::u32 millisecondsTimeout);
-	UBA_API void ProcessHandle_Cancel(uba::ProcessHandle* handle, bool terminate);
-	UBA_API void ProcessHandle_Destroy(uba::ProcessHandle* handle);
-	UBA_API const uba::ProcessStartInfo* Process_GetStartInfo(uba::Process& process);
+	UBA_API uba::u64 ProcessHandle_GetHash(const uba::ProcessHandle* handle);
+	UBA_API uba::u64 ProcessHandle_GetTotalProcessorTime(const uba::ProcessHandle* handle);
+	UBA_API uba::u64 ProcessHandle_GetTotalWallTime(const uba::ProcessHandle* handle);
+	UBA_API void ProcessHandle_TraverseOutputFiles(const uba::ProcessHandle* handle, ProcessHandle_OutputFileCallback* callback, void* userData);
+	UBA_API bool ProcessHandle_WaitForExit(const uba::ProcessHandle* handle, uba::u32 millisecondsTimeout);
+	UBA_API void ProcessHandle_Cancel(const uba::ProcessHandle* handle, bool terminate);
+	UBA_API void ProcessHandle_Destroy(const uba::ProcessHandle* handle);
+	UBA_API const uba::ProcessStartInfo* Process_GetStartInfo(const uba::Process& process);
 
 	// SessionServer
 	using SessionServer_RemoteProcessAvailableCallback = void(void* userData, bool isCrossArchitecture);
@@ -111,11 +113,19 @@ extern "C"
 	UBA_API uba::SessionServer* SessionServer_Create2(uba::StorageServer& s, uba::NetworkServer& ns, const uba::Config& c, uba::LogWriter& lw = uba::g_consoleLogWriter, const uba::u8* environment = nullptr, uba::u32 environmentSize = 0);
 	UBA_API void SessionServer_SetRemoteProcessAvailable(uba::SessionServer* server, SessionServer_RemoteProcessAvailableCallback* available, void* userData);
 	UBA_API void SessionServer_SetRemoteProcessReturned(uba::SessionServer* server, SessionServer_RemoteProcessReturnedCallback* returned, void* userData);
+
 	UBA_API bool SessionServer_RefreshDirectory(uba::SessionServer* server, const uba::tchar* directory);
 	UBA_API bool SessionServer_RegisterNewFile(uba::SessionServer* server, const uba::tchar* filePath);
 	UBA_API void SessionServer_RegisterDeleteFile(uba::SessionServer* server, const uba::tchar* filePath);
 	UBA_API bool SessionServer_RegisterNewDirectory(uba::SessionServer* server, const uba::tchar* directoryPath);
 	UBA_API bool SessionServer_RegisterVirtualFile(uba::SessionServer* server, const uba::tchar* filename, const uba::tchar* sourceFile, uba::u64 sourceOffset, uba::u64 sourceSize);
+	UBA_API bool SessionServer_CreateVirtualFile(uba::SessionServer* server, const uba::tchar* filename, const void* memory, uba::u64 memorySize);
+	UBA_API bool SessionServer_DeleteVirtualFile(uba::SessionServer* server, const uba::tchar* filename);
+
+	UBA_API bool SessionServer_GetOutputFileSize(uba::SessionServer* server, uba::u64& outSize, const uba::tchar* filename);
+	UBA_API bool SessionServer_GetOutputFileData(uba::SessionServer* server, void* outData, const uba::tchar* filename, bool deleteInternalMapping);
+	UBA_API bool SessionServer_WriteOutputFile(uba::SessionServer* server, const uba::tchar* filename, bool deleteInternalMapping);
+
 	UBA_API uba::ProcessHandle* SessionServer_RunProcess(uba::SessionServer* server, uba::ProcessStartInfo& info, bool async, bool enableDetour);
 	UBA_API uba::ProcessHandle* SessionServer_RunProcessRemote(uba::SessionServer* server, uba::ProcessStartInfo& info, float weight, const void* knownInputs = nullptr, uba::u32 knownInputsCount = 0, bool allowCrossArchitecture = false);
 	UBA_API uba::ProcessHandle* SessionServer_RunProcessRacing(uba::SessionServer* server, uba::u32 raceAgainstRemoteProcessId);
@@ -182,7 +192,7 @@ extern "C"
 
 
 	// DEPRECATED, don't use
-	UBA_API void DestroyProcessHandle(uba::ProcessHandle* handle);
+	UBA_API void DestroyProcessHandle(const uba::ProcessHandle* handle);
 	UBA_API uba::RootPaths* RootPaths_Create(uba::LogWriter& writer);
 	UBA_API bool RootPaths_RegisterRoot(uba::RootPaths* rootPaths, const uba::tchar* path, bool includeInKey, uba::u8 id = 0);
 	UBA_API bool RootPaths_RegisterSystemRoots(uba::RootPaths* rootPaths, uba::u8 startId = 0);
diff --git a/Engine/Source/Programs/UnrealBuildAccelerator/Test/Private/UbaTestAll.h b/Engine/Source/Programs/UnrealBuildAccelerator/Test/Private/UbaTestAll.h
--- a/Engine/Source/Programs/UnrealBuildAccelerator/Test/Private/UbaTestAll.h
+++ b/Engine/Source/Programs/UnrealBuildAccelerator/Test/Private/UbaTestAll.h
@@ -74,6 +74,8 @@ namespace uba
 		UBA_TEST(TestStdOutViaCmd) \
 		UBA_TEST(TestVolumeCache) \
 		UBA_TEST(TestDependencyCrawler) \
+		UBA_TEST(TestVirtualFile) \
+		UBA_TEST(TestRemoteVirtualFile) \
 		UBA_TEST(TestRemoteProcessSpecialCase1) \
 
 
diff --git a/Engine/Source/Programs/UnrealBuildAccelerator/Test/Private/UbaTestSession.cpp b/Engine/Source/Programs/UnrealBuildAccelerator/Test/Private/UbaTestSession.cpp
--- a/Engine/Source/Programs/UnrealBuildAccelerator/Test/Private/UbaTestSession.cpp
+++ b/Engine/Source/Programs/UnrealBuildAccelerator/Test/Private/UbaTestSession.cpp
@@ -749,6 +749,72 @@ namespace uba
 #endif
 	}
 
+	bool RunVirtualFileTest(LoggerWithWriter& logger, SessionServer& session, const tchar* workingDir, const RunProcessFunction& runProcess)
+	{
+		StringBuffer<> inFile(workingDir);
+		inFile.Append(TCV("VirtualFile.in"));
+		if (!session.CreateVirtualFile(inFile.data, "FOO", 3))
+			return false;
+		ProcessHandle ph;
+		if (!RunTestApp(logger, workingDir, runProcess, TC("-virtualFile"), &ph))
+			return false;
+
+		StringBuffer<> outFile(workingDir);
+		outFile.Append(TCV("VirtualFile.out"));
+
+		bool success = false;
+		ph.TraverseOutputFiles([&](StringView file) { success = file.Equals(outFile); });
+		if (!success)
+			return false;
+
+		if (!session.DeleteVirtualFile(inFile.data))
+			return false;
+
+		if (FileExists(logger, outFile.data))
+			return false;
+		u64 outSize;
+		if (!session.GetOutputFileSize(outSize, outFile.data))
+			return false;
+		if (outSize != 3)
+			return false;
+		u8 data[3];
+		if (!session.GetOutputFileData(data, outFile.data, false))
+			return false;
+		if (memcmp(data, "BAR", 3) != 0)
+			return false;
+		if (FileExists(logger, outFile.data))
+			return false;
+		if (!session.WriteOutputFile(outFile.data, true))
+			return false;
+		if (!FileExists(logger, outFile.data))
+			return false;
+
+		if (ph.IsRemote())
+			if (!session.GetStorage().DeleteCasForFile(inFile.data))
+				return false;
+
+		if (session.GetOutputFileSize(outSize, outFile.data))
+			return false;
+		if (session.GetOutputFileData(data, outFile.data, true))
+			return false;
+		return true;
+	}
+
+	bool TestVirtualFile(LoggerWithWriter& logger, const StringBufferBase& testRootDir)
+	{
+		Config config;
+		config.AddTable(TC("Session")).AddValue(TC("ShouldWriteToDisk"), false);
+		return RunLocal(logger, testRootDir, config, RunVirtualFileTest);
+	}
+
+	bool TestRemoteVirtualFile(LoggerWithWriter& logger, const StringBufferBase& testRootDir)
+	{
+		Config config;
+		config.AddTable(TC("Session")).AddValue(TC("ShouldWriteToDisk"), false);
+		config.AddTable(TC("Storage")).AddValue(TC("CreateIndependentMappings"), true);
+		return RunRemote(logger, testRootDir, config, RunVirtualFileTest);
+	}
+
 	#if PLATFORM_MAC
 	bool RunXCodeSelect(LoggerWithWriter& logger, SessionServer& session, const tchar* workingDir, const RunProcessFunction& runProcess)
 	{
diff --git a/Engine/Source/Programs/UnrealBuildAccelerator/TestApp/Private/PosixOS/UbaTestAppPosix.cpp b/Engine/Source/Programs/UnrealBuildAccelerator/TestApp/Private/PosixOS/UbaTestAppPosix.cpp
--- a/Engine/Source/Programs/UnrealBuildAccelerator/TestApp/Private/PosixOS/UbaTestAppPosix.cpp
+++ b/Engine/Source/Programs/UnrealBuildAccelerator/TestApp/Private/PosixOS/UbaTestAppPosix.cpp
@@ -300,6 +300,19 @@ int main(int argc, char* argv[])
 
 		return 0;
 	}
+	else if (strncmp(argv[1], "-virtualFile=", 13) == 0)
+	{
+		const char* file = argv[1]+13;
+		int rh = open(file, O_RDONLY);
+		if (rh == -1)
+			return LogError("Failed to open file %s", file);
+		char data[4] = {};
+		if (read(rh, data, 3) != 3)
+			return LogError("Failed to read 3 bytes from file %s", file);
+		close(rh);
+		if (memcmp(data, "FOO", 3) != 0)
+			return LogError("File %s has wrong content", file);
+	}
 	#if PLATFORM_MAC
 	else if (strcmp(argv[1], "-xcode-select") == 0)
 	{
diff --git a/Engine/Source/Programs/UnrealBuildAccelerator/TestApp/Private/Windows/UbaTestAppWin.cpp b/Engine/Source/Programs/UnrealBuildAccelerator/TestApp/Private/Windows/UbaTestAppWin.cpp
--- a/Engine/Source/Programs/UnrealBuildAccelerator/TestApp/Private/Windows/UbaTestAppWin.cpp
+++ b/Engine/Source/Programs/UnrealBuildAccelerator/TestApp/Private/Windows/UbaTestAppWin.cpp
@@ -322,6 +322,28 @@ int wmain(int argc, wchar_t* argv[])
 		}
 		wprintf(L"%s\n", str);
 	}
+	else if (wcsncmp(argv[1], L"-virtualFile", 12) == 0)
+	{
+		const wchar_t* file = L"VirtualFile.in";
+		HANDLE rh = CreateFileW(file, GENERIC_READ, FILE_SHARE_READ, NULL, OPEN_EXISTING, 0, NULL);
+		if (rh == INVALID_HANDLE_VALUE)
+			return LogError(L"Failed to open file %s", file);
+		char data[4] = {};
+		DWORD bytesRead;
+		if (!ReadFile(rh, data, 4, &bytesRead, NULL) || bytesRead != 3)
+			return LogError(L"Failed to read 3 bytes from file %s", file);
+		CloseHandle(rh);
+		if (memcmp(data, "FOO", 3) != 0)
+			return LogError(L"File %s has wrong content", file);
+
+		file = L"VirtualFile.out";
+		HANDLE wh = CreateFileW(file, GENERIC_WRITE, FILE_SHARE_WRITE, NULL, CREATE_ALWAYS, 0, NULL);
+		if (wh == INVALID_HANDLE_VALUE)
+			return LogError(L"Failed to open file %s for write", file);
+		if (!WriteFile(wh, "BAR", 3, NULL, NULL))
+			return LogError(L"Failed to write 3 bytes to file %s", file);
+		CloseHandle(wh);
+	}
 	else if (wcsncmp(argv[1], L"-readwrite=", 11) == 0)
 	{
 		int callIndex = _wtoi(argv[1] + 11);(void)callIndex;
