[UBA]
* Moved last oodle stuff from UbaStorage.cpp and into UbaStorageUtils.cpp (and a tiny bit into UbaPlatform.cpp)
* Added call to InitMemory in UbaAgent and UbaCacheService

diff --git a/Engine/Source/Programs/UnrealBuildAccelerator/Agent/Private/UbaAgent.cpp b/Engine/Source/Programs/UnrealBuildAccelerator/Agent/Private/UbaAgent.cpp
--- a/Engine/Source/Programs/UnrealBuildAccelerator/Agent/Private/UbaAgent.cpp
+++ b/Engine/Source/Programs/UnrealBuildAccelerator/Agent/Private/UbaAgent.cpp
@@ -728,6 +728,8 @@ namespace uba
 			}
 		}
 
+		InitMemory();
+
 		if (useExceptionHandler)
 			AddExceptionHandler();
 
diff --git a/Engine/Source/Programs/UnrealBuildAccelerator/CacheService/Private/UbaCacheService.cpp b/Engine/Source/Programs/UnrealBuildAccelerator/CacheService/Private/UbaCacheService.cpp
--- a/Engine/Source/Programs/UnrealBuildAccelerator/CacheService/Private/UbaCacheService.cpp
+++ b/Engine/Source/Programs/UnrealBuildAccelerator/CacheService/Private/UbaCacheService.cpp
@@ -269,7 +269,7 @@ namespace uba
 		#endif
 		logger.Info(TC("UbaCacheService v%s(%u)%s (Workers: %u, Rootdir: \"%s\", StoreCapacity: %s, Expiration: %s)"), Version, CacheNetworkVersion, dbgStr, maxWorkerCount, g_rootDir.data, BytesToText(storageCapacity).str, TimeToText(MsToTime(expirationTimeSeconds)*1000, true).str);
 
-		u64 maintenanceReserveSizeMb = 128;
+		u64 maintenanceReserveSizeMb = 196;
 
 		if (SupportsHugePages())
 		{
@@ -420,6 +420,8 @@ namespace uba
 
 		logger.Info(TC(""));
 
+		InitMemory();
+
 		if (signalHandlerEnabled)
 			AddExceptionHandler();
 
diff --git a/Engine/Source/Programs/UnrealBuildAccelerator/Common/Private/UbaStorage.cpp b/Engine/Source/Programs/UnrealBuildAccelerator/Common/Private/UbaStorage.cpp
--- a/Engine/Source/Programs/UnrealBuildAccelerator/Common/Private/UbaStorage.cpp
+++ b/Engine/Source/Programs/UnrealBuildAccelerator/Common/Private/UbaStorage.cpp
@@ -10,7 +10,6 @@
 #include "UbaDirectoryIterator.h"
 #include "UbaStorageUtils.h"
 #include "UbaWorkManager.h"
-#include <oodle2.h>
 
 #define UBA_CHECK_BUFFER_SLOTS 0
 
@@ -454,215 +453,8 @@ namespace uba
 
 	bool StorageImpl::DecompressMemoryToMemory(const u8* compressedData, u64 compressedSize, u8* writeData, u64 decompressedSize, const tchar* readHint, const tchar* writeHint)
 	{
-		UBA_ASSERTF(compressedData, TC("DecompressMemoryToMemory got readmem nullptr (%s)"), readHint);
-		UBA_ASSERTF(writeData, TC("DecompressMemoryToMemory got writemem nullptr (%s)"), writeHint);
-
-		StorageStats& stats = Stats();
-
-		if (decompressedSize > BufferSlotSize * 4) // Arbitrary size threshold. We want to at least catch the pch here
-		{
-			struct WorkRec
-			{
-				Logger* logger;
-				const tchar* hint;
-				Atomic<u64> refCount;
-				const u8* readPos = nullptr;
-				u8* writePos = nullptr;
-
-				Futex lock;
-				u64 decompressedSize = 0;
-				u64 decompressedLeft = 0;
-				u64 written = 0;
-				Event done;
-				Atomic<bool> error;
-			};
-
-			WorkRec* rec = new WorkRec();
-			rec->logger = &m_logger;
-			rec->hint = readHint;
-			rec->readPos = compressedData;
-			rec->writePos = writeData;
-			rec->decompressedSize = decompressedSize;
-			rec->decompressedLeft = decompressedSize;
-			rec->done.Create(true);
-			rec->refCount = 2;
-
-			auto work = [rec](const WorkContext&)
-				{
-					u64 lastWritten = 0;
-					while (true)
-					{
-						SCOPED_FUTEX(rec->lock, lock);
-						rec->written += lastWritten;
-						if (!rec->decompressedLeft)
-						{
-							if (rec->written == rec->decompressedSize)
-								rec->done.Set();
-							lock.Leave();
-							if (!--rec->refCount)
-								delete rec;
-							return;
-						}
-						const u8* readPos = rec->readPos;
-						u8* writePos = rec->writePos;
-						u32 compressedBlockSize = ((const u32*)readPos)[0];
-						u32 decompressedBlockSize = ((const u32*)readPos)[1];
-
-						if (decompressedBlockSize == 0 || decompressedBlockSize > rec->decompressedSize)
-						{
-							bool f = false;
-							if (rec->error.compare_exchange_strong(f, true))
-								rec->logger->Error(TC("Decompressed block size %u is invalid. Decompressed file is %u (%s)"), decompressedBlockSize, rec->decompressedSize, rec->hint);
-							if (!--rec->refCount)
-								delete rec;
-							rec->done.Set();
-							return;
-						}
-
-						readPos += sizeof(u32) * 2;
-						rec->decompressedLeft -= decompressedBlockSize;
-						rec->readPos = readPos + compressedBlockSize;
-						rec->writePos += decompressedBlockSize;
-						lock.Leave();
-
-						void* decoderMem = nullptr;
-						u64 decoderMemSize = 0;
-						OO_SINTa decompLen = OodleLZ_Decompress(readPos, (OO_SINTa)compressedBlockSize, writePos, (OO_SINTa)decompressedBlockSize, OodleLZ_FuzzSafe_Yes, OodleLZ_CheckCRC_No, OodleLZ_Verbosity_None, NULL, 0, NULL, NULL, decoderMem, decoderMemSize);
-						if (decompLen != decompressedBlockSize)
-						{
-							bool f = false;
-							if (rec->error.compare_exchange_strong(f, true))
-								rec->logger->Error(TC("Expecting to be able to decompress %u bytes to %u bytes but got %llu (%s)"), compressedBlockSize, decompressedBlockSize, decompLen, rec->hint);
-							if (!--rec->refCount)
-								delete rec;
-							rec->done.Set();
-							return;
-						}
-						lastWritten = u64(decompLen);
-					}
-				};
-
-			u32 workCount = u32(decompressedSize / BufferSlotSize) + 1;
-			u32 workerCount = Min(workCount, m_workManager.GetWorkerCount() - 1); // We are a worker ourselves
-			workerCount = Min(workerCount, MaxWorkItemsPerAction); // Cap this to not starve other things
-			rec->refCount += workerCount;
-			m_workManager.AddWork(work, workerCount, TC("DecompressMemToMem"));
-
-			TimerScope ts(stats.decompressToMem);
-			{
-				TrackWorkScope tws;
-				work({tws});
-			}
-			rec->done.IsSet();
-			bool success = !rec->error;
-			if (!success)
-				while (rec->refCount > 1)
-					Sleep(10);
-
-			if (!--rec->refCount)
-				delete rec;
-			return success;
-		}
-		else
-		{
-			const u8* readPos = compressedData;
-			u8* writePos = writeData;
-
-			u64 left = decompressedSize;
-			while (left)
-			{
-				u32 compressedBlockSize = ((const u32*)readPos)[0];
-				if (!compressedBlockSize)
-					break;
-				u32 decompressedBlockSize = ((const u32*)readPos)[1];
-				if (decompressedBlockSize == 0 || decompressedBlockSize > left)
-					return m_logger.Error(TC("Decompressed block size %u is invalid. Decompressed file is %u (%s -> %s)"), decompressedBlockSize, decompressedSize, readHint, writeHint);
-				readPos += sizeof(u32) * 2;
-
-				void* decoderMem = nullptr;
-				u64 decoderMemSize = 0;
-				TimerScope ts(stats.decompressToMem);
-				OO_SINTa decompLen = OodleLZ_Decompress(readPos, (OO_SINTa)compressedBlockSize, writePos, (OO_SINTa)decompressedBlockSize, OodleLZ_FuzzSafe_Yes, OodleLZ_CheckCRC_No, OodleLZ_Verbosity_None, NULL, 0, NULL, NULL, decoderMem, decoderMemSize);
-				if (decompLen != decompressedBlockSize)
-					return m_logger.Error(TC("Expecting to be able to decompress %u to %u bytes at pos %llu but got %llu. File compressed size: %llu (%s -> %s)"), compressedBlockSize, decompressedBlockSize, decompressedSize - left, decompLen, compressedSize, readHint, writeHint);
-				writePos += decompressedBlockSize;
-				readPos += compressedBlockSize;
-				left -= decompressedBlockSize;
-			}
-		}
-		return true;
-	}
-
-	bool StorageImpl::DecompressMemoryToFile(u8* compressedData, FileAccessor& destination, u64 decompressedSize, bool useNoBuffering)
-	{
-		StorageStats& stats = Stats();
-		u8* readPos = compressedData;
-
-		u8* slot = m_bufferSlots.Pop();
-		auto _ = MakeGuard([&](){ m_bufferSlots.Push(slot); });
-
-		u64 left = decompressedSize;
-		u64 overflow = 0;
-		while (left)
-		{
-			u32 compressedBlockSize = ((u32*)readPos)[0];
-			if (!compressedBlockSize)
-				break;
-			u32 decompressedBlockSize = ((u32*)readPos)[1];
-
-			readPos += sizeof(u32)*2;
-
-			OO_SINTa decompLen;
-			{
-				TimerScope ts(stats.decompressToMem);
-				decompLen = OodleLZ_Decompress(readPos, (OO_SINTa)compressedBlockSize, slot + overflow, (OO_SINTa)decompressedBlockSize);
-			}
-			UBA_ASSERT(decompLen == decompressedBlockSize);
-			
-			u64 available = overflow + u64(decompLen);
-
-			if (left - available > 0 && available < BufferSlotHalfSize)
-			{
-				overflow += u64(decompLen);
-				readPos += compressedBlockSize;
-				continue;
-			}
-
-			if (useNoBuffering)
-			{
-				u64 writeSize = AlignUp(available - 4096 + 1, 4096);
-
-				if (!destination.Write(slot, writeSize))
-					return false;
-
-				overflow = available - writeSize;
-				readPos += compressedBlockSize;
-				left -= writeSize;
-
-				if (overflow == left)
-				{
-					if (!destination.Write(slot + writeSize, 4096))
-						return false;
-					break;
-				}
-
-				memcpy(slot, slot + writeSize, overflow);
-			}
-			else
-			{
-				u64 writeSize = available;
-				if (!destination.Write(slot, writeSize))
-					return false;
-				readPos += compressedBlockSize;
-				left -= writeSize;
-				overflow = 0;
-			}
-		}
-
-		if (useNoBuffering)
-			if (!SetEndOfFile(m_logger, destination.GetFileName(), destination.GetHandle(), decompressedSize))
-				return false;
-		return true;
+		DecompressWriter w { m_logger, m_bufferSlots, m_workManager, Stats() };
+		return w.DecompressMemoryToMemory(compressedData, compressedSize, writeData, decompressedSize, readHint, writeHint);
 	}
 
 	bool StorageImpl::VerifyExisting(bool& outReturnValue, ScopedWriteLock& entryLock, const CasKey& casKey, CasEntry& casEntry, StringBufferBase& casFile)
@@ -1021,18 +813,6 @@ namespace uba
 		return m_logger.Warning(TC("Failed to delete %s (%s)"), isRunningName.data, LastErrorToText(lastError).data);
 	}
 
-	#if UBA_USE_MIMALLOC
-	void* Oodle_MallocAligned(OO_SINTa bytes, OO_S32 alignment)
-	{
-		return mi_malloc_aligned(bytes, alignment);
-	}
-
-	void Oodle_Free(void* ptr)
-	{
-		mi_free(ptr);
-	}
-	#endif
-
 	StorageImpl::StorageImpl(const StorageCreateInfo& info, const tchar* logPrefix)
 	:	m_logger(info.writer, logPrefix)
 	,	m_workManager(info.workManager)
@@ -1067,10 +847,6 @@ namespace uba
 
 		m_casCompressor = info.casCompressor;
 		m_casCompressionLevel = info.casCompressionLevel;
-
-		#if UBA_USE_MIMALLOC
-		OodleCore_Plugins_SetAllocators(Oodle_MallocAligned, Oodle_Free);
-		#endif
 	}
 
 	StorageImpl::~StorageImpl()
@@ -2385,8 +2161,11 @@ namespace uba
 					if (!destinationFile.CreateWrite(allowRead, writeFlags | fileAttributes, decompressedSize, m_tempPath.data))
 						return false;
 					if (decompressedSize)
-						if (!DecompressMemoryToFile(readData, destinationFile, decompressedSize, useNoBuffering))
+					{
+						DecompressWriter w { m_logger, m_bufferSlots, m_workManager, Stats() };
+						if (!w.DecompressMemoryToFile(readData, destinationFile, decompressedSize, useNoBuffering))
 							return false;
+					}
 				}
 				else
 				{
@@ -2646,75 +2425,8 @@ namespace uba
 
 	bool StorageImpl::DecompressFileToMemory(const tchar* fileName, FileHandle fileHandle, u8* dest, u64 decompressedSize, const tchar* writeHint, u64 fileStartOffset)
 	{
-		if (decompressedSize > BufferSlotSize*4) // Arbitrary size threshold. We want to at least catch the pch here
-		{
-			u64 compressedSize;
-			if (!uba::GetFileSizeEx(compressedSize, fileHandle))
-				return m_logger.Error(TC("GetFileSize failed for %s (%s)"), fileName, LastErrorToText().data);
-			FileMappingHandle fileMapping = uba::CreateFileMappingW(m_logger, fileHandle, PAGE_READONLY, compressedSize, fileName);
-			if (!fileMapping.IsValid())
-				return m_logger.Error(TC("Failed to create file mapping for %s (%s)"), fileName, LastErrorToText().data);
-			auto fmg = MakeGuard([&]() { CloseFileMapping(m_logger, fileMapping, fileName); });
-			u8* fileData = MapViewOfFile(m_logger, fileMapping, FILE_MAP_READ, 0, compressedSize);
-			if (!fileData)
-				return m_logger.Error(TC("Failed to map view of file mapping for %s (%s)"), fileName, LastErrorToText().data);
-			auto udg = MakeGuard([&]() { UnmapViewOfFile(m_logger, fileData, compressedSize, fileName); });
-			
-			u8* readPos = fileData + 8 + fileStartOffset;
-			if (!DecompressMemoryToMemory(readPos, compressedSize, dest, decompressedSize, fileName, writeHint))
-				return false;
-		}
-		else
-		{
-			StorageStats& stats = Stats();
-			u8* slot = m_bufferSlots.Pop();
-			auto _ = MakeGuard([&]() { m_bufferSlots.Push(slot); });
-
-			void* decoderMem = slot + BufferSlotHalfSize;
-			u64 decoderMemSize = BufferSlotHalfSize;
-
-			u64 bytesRead = 8; // We know size has already been read
-
-			u8* readBuffer = slot;
-			u8* writePos = dest;
-			u64 left = decompressedSize;
-			while (left)
-			{
-				u32 sizes[2];
-				if (!ReadFile(m_logger, fileName, fileHandle, sizes, sizeof(u32) * 2))
-				{
-					u64 compressedSize;
-					if (!uba::GetFileSizeEx(compressedSize, fileHandle))
-						return m_logger.Error(TC("GetFileSize failed for %s (%s)"), fileName, LastErrorToText().data);
-					if (bytesRead + 8 > compressedSize)
-						return m_logger.Error(TC("File %s corrupt. Tried to read 8 bytes. File is smaller than expected (Read: %llu, Size: %llu)"), fileName, bytesRead, compressedSize);
-					return false;
-				}
-				u32 compressedBlockSize = sizes[0];
-				u32 decompressedBlockSize = sizes[1];
-
-				bytesRead += 8;
-
-				if (!ReadFile(m_logger, fileName, fileHandle, readBuffer, compressedBlockSize))
-				{
-					u64 compressedSize;
-					if (!uba::GetFileSizeEx(compressedSize, fileHandle))
-						return m_logger.Error(TC("GetFileSize failed for %s (%s)"), fileName, LastErrorToText().data);
-					if (bytesRead + compressedBlockSize > compressedSize)
-						return m_logger.Error(TC("File %s corrupt. Compressed block size (%u) is larger than what is left of file (%llu)"), fileName, compressedBlockSize, compressedSize - bytesRead);
-					return false;
-				}
-				bytesRead += compressedBlockSize;
-
-				TimerScope ts(stats.decompressToMem);
-				OO_SINTa decompLen = OodleLZ_Decompress(readBuffer, (OO_SINTa)compressedBlockSize, writePos, (OO_SINTa)decompressedBlockSize, OodleLZ_FuzzSafe_Yes, OodleLZ_CheckCRC_No, OodleLZ_Verbosity_None, NULL, 0, NULL, NULL, decoderMem, decoderMemSize);
-				if (decompLen != decompressedBlockSize)
-					return m_logger.Error(TC("Failed to decompress data from file %s at pos %llu"), fileName, decompressedSize - left);
-				writePos += decompressedBlockSize;
-				left -= decompressedBlockSize;
-			}
-		}
-		return true;
+		DecompressWriter w { m_logger, m_bufferSlots, m_workManager, Stats() };
+		return w.DecompressFileToMemory(fileName, fileHandle, dest, decompressedSize, writeHint, fileStartOffset);
 	}
 
 	bool StorageImpl::CreateDirectory(const tchar* dir)
diff --git a/Engine/Source/Programs/UnrealBuildAccelerator/Common/Private/UbaStorageUtils.cpp b/Engine/Source/Programs/UnrealBuildAccelerator/Common/Private/UbaStorageUtils.cpp
--- a/Engine/Source/Programs/UnrealBuildAccelerator/Common/Private/UbaStorageUtils.cpp
+++ b/Engine/Source/Programs/UnrealBuildAccelerator/Common/Private/UbaStorageUtils.cpp
@@ -9,6 +9,11 @@
 
 namespace uba
 {
+	////////////////////////////////////////////////////////////////////////////////////////////////////
+
+	constexpr u32 MaxWorkItemsPerAction2 = 128; // Cap this to not starve other things
+
+
 	#define OODLE_COMPRESSORS \
 		OODLE_COMPRESSOR(Selkie) \
 		OODLE_COMPRESSOR(Mermaid) \
@@ -27,6 +32,8 @@ namespace uba
 		OODLE_COMPRESSION_LEVEL(Optimal4) \
 		OODLE_COMPRESSION_LEVEL(Optimal5) \
 
+	////////////////////////////////////////////////////////////////////////////////////////////////////
+
 	u8 GetCompressor(const tchar* str)
 	{
 		#define OODLE_COMPRESSOR(x) if (Equals(str, TC(#x))) return u8(OodleLZ_Compressor_##x);
@@ -46,7 +53,6 @@ namespace uba
 
 	CasKey CalculateCasKey(u8* fileMem, u64 fileSize, bool storeCompressed, WorkManager* workManager, const tchar* hint)
 	{
-		constexpr u32 MaxWorkItemsPerAction2 = 128; // Cap this to not starve other things
 		CasKeyHasher hasher;
 
 		if (fileSize == 0)
@@ -313,6 +319,8 @@ namespace uba
 		return sendEndMessage();
 	}
 
+	////////////////////////////////////////////////////////////////////////////////////////////////////
+
 	bool FileSender::SendFileCompressed(const CasKey& casKey, const tchar* fileName, const u8* sourceMem, u64 sourceSize, const tchar* hint)
 	{
 		UBA_ASSERT(casKey != CasKeyZero);
@@ -473,6 +481,8 @@ namespace uba
 		return sendEndMessage();
 	}
 
+	////////////////////////////////////////////////////////////////////////////////////////////////////
+
 	bool FileFetcher::RetrieveFile(Logger& logger, NetworkClient& client, const CasKey& casKey, const tchar* destination, bool writeCompressed, MemoryBlock* destinationMem)
 	{
 		TimerScope ts(m_stats.recvCas);
@@ -808,6 +818,8 @@ namespace uba
 		return true;
 	}
 
+	////////////////////////////////////////////////////////////////////////////////////////////////////
+
 	CompressWriter::CompressWriter(Logger& l, BufferSlots& bs, WorkManager& wm, StorageStats& s, u8 cc, u8 ccl, bool avof)
 	:	m_logger(l)
 	,	m_bufferSlots(bs)
@@ -864,8 +876,6 @@ namespace uba
 
 	bool CompressWriter::CompressFromMem(LinearWriter& destination, u32 workCount, const u8* uncompressedData, u64 fileSize, u64 maxUncompressedBlock, u64& totalWritten)
 	{
-		constexpr u32 MaxWorkItemsPerAction = 128; // Cap this to not starve other things
-
 		#ifndef __clang_analyzer__
 
 		struct WorkRec
@@ -974,7 +984,7 @@ namespace uba
 		};
 
 		u32 workerCount = Min(workCount, m_workManager.GetWorkerCount());
-		workerCount = Min(workerCount, MaxWorkItemsPerAction);
+		workerCount = Min(workerCount, MaxWorkItemsPerAction2);
 
 		rec.refCount = workerCount + 1; // We need to keep refcount up 1 to make sure it is not deleted before we read rec->written
 		m_workManager.AddWork(work, workerCount-1, TC("Compress")); // We are a worker ourselves
@@ -1098,11 +1108,15 @@ namespace uba
 		return true;
 	}
 
+	////////////////////////////////////////////////////////////////////////////////////////////////////
+
 	LinearWriterFile::LinearWriterFile(FileAccessor& f) : file(f) {}
 	bool LinearWriterFile::Write(const void* data, u64 dataLen) { written += dataLen; return file.Write(data, dataLen); }
 	u64 LinearWriterFile::GetWritten() { return written; }
 	const tchar* LinearWriterFile::GetHint() { return file.GetFileName(); }
 
+	////////////////////////////////////////////////////////////////////////////////////////////////////
+
 	LinearWriterMem::LinearWriterMem(u8* d, u64 c) : pos(d), size(0), capacity(c) {}
 	u64 LinearWriterMem::GetWritten() { return size; }
 	bool LinearWriterMem::Write(const void* data, u64 dataLen)
@@ -1116,4 +1130,294 @@ namespace uba
 		return true;
 	}
 	const tchar* LinearWriterMem::GetHint() { return TC(""); }
-}
+
+	////////////////////////////////////////////////////////////////////////////////////////////////////
+
+	bool DecompressWriter::DecompressMemoryToMemory(const u8* compressedData, u64 compressedSize, u8* writeData, u64 decompressedSize, const tchar* readHint, const tchar* writeHint)
+	{
+		UBA_ASSERTF(compressedData, TC("DecompressMemoryToMemory got readmem nullptr (%s)"), readHint);
+		UBA_ASSERTF(writeData, TC("DecompressMemoryToMemory got writemem nullptr (%s)"), writeHint);
+
+		StorageStats& stats = m_stats;
+
+		if (decompressedSize > BufferSlotSize * 4) // Arbitrary size threshold. We want to at least catch the pch here
+		{
+			struct WorkRec
+			{
+				Logger* logger;
+				const tchar* hint;
+				Atomic<u64> refCount;
+				const u8* readPos = nullptr;
+				u8* writePos = nullptr;
+
+				Futex lock;
+				u64 decompressedSize = 0;
+				u64 decompressedLeft = 0;
+				u64 written = 0;
+				Event done;
+				Atomic<bool> error;
+			};
+
+			WorkRec* rec = new WorkRec();
+			rec->logger = &m_logger;
+			rec->hint = readHint;
+			rec->readPos = compressedData;
+			rec->writePos = writeData;
+			rec->decompressedSize = decompressedSize;
+			rec->decompressedLeft = decompressedSize;
+			rec->done.Create(true);
+			rec->refCount = 2;
+
+			auto work = [rec](const WorkContext&)
+				{
+					u64 lastWritten = 0;
+					while (true)
+					{
+						SCOPED_FUTEX(rec->lock, lock);
+						rec->written += lastWritten;
+						if (!rec->decompressedLeft)
+						{
+							if (rec->written == rec->decompressedSize)
+								rec->done.Set();
+							lock.Leave();
+							if (!--rec->refCount)
+								delete rec;
+							return;
+						}
+						const u8* readPos = rec->readPos;
+						u8* writePos = rec->writePos;
+						u32 compressedBlockSize = ((const u32*)readPos)[0];
+						u32 decompressedBlockSize = ((const u32*)readPos)[1];
+
+						if (decompressedBlockSize == 0 || decompressedBlockSize > rec->decompressedSize)
+						{
+							bool f = false;
+							if (rec->error.compare_exchange_strong(f, true))
+								rec->logger->Error(TC("Decompressed block size %u is invalid. Decompressed file is %u (%s)"), decompressedBlockSize, rec->decompressedSize, rec->hint);
+							if (!--rec->refCount)
+								delete rec;
+							rec->done.Set();
+							return;
+						}
+
+						readPos += sizeof(u32) * 2;
+						rec->decompressedLeft -= decompressedBlockSize;
+						rec->readPos = readPos + compressedBlockSize;
+						rec->writePos += decompressedBlockSize;
+						lock.Leave();
+
+						void* decoderMem = nullptr;
+						u64 decoderMemSize = 0;
+						OO_SINTa decompLen = OodleLZ_Decompress(readPos, (OO_SINTa)compressedBlockSize, writePos, (OO_SINTa)decompressedBlockSize, OodleLZ_FuzzSafe_Yes, OodleLZ_CheckCRC_No, OodleLZ_Verbosity_None, NULL, 0, NULL, NULL, decoderMem, decoderMemSize);
+						if (decompLen != decompressedBlockSize)
+						{
+							bool f = false;
+							if (rec->error.compare_exchange_strong(f, true))
+								rec->logger->Error(TC("Expecting to be able to decompress %u bytes to %u bytes but got %llu (%s)"), compressedBlockSize, decompressedBlockSize, decompLen, rec->hint);
+							if (!--rec->refCount)
+								delete rec;
+							rec->done.Set();
+							return;
+						}
+						lastWritten = u64(decompLen);
+					}
+				};
+
+			u32 workCount = u32(decompressedSize / BufferSlotSize) + 1;
+			u32 workerCount = Min(workCount, m_workManager.GetWorkerCount() - 1); // We are a worker ourselves
+			workerCount = Min(workerCount, MaxWorkItemsPerAction2); // Cap this to not starve other things
+			rec->refCount += workerCount;
+			m_workManager.AddWork(work, workerCount, TC("DecompressMemToMem"));
+
+			TimerScope ts(stats.decompressToMem);
+			{
+				TrackWorkScope tws;
+				work({tws});
+			}
+			rec->done.IsSet();
+			bool success = !rec->error;
+			if (!success)
+				while (rec->refCount > 1)
+					Sleep(10);
+
+			if (!--rec->refCount)
+				delete rec;
+			return success;
+		}
+		else
+		{
+			const u8* readPos = compressedData;
+			u8* writePos = writeData;
+
+			u64 left = decompressedSize;
+			while (left)
+			{
+				u32 compressedBlockSize = ((const u32*)readPos)[0];
+				if (!compressedBlockSize)
+					break;
+				u32 decompressedBlockSize = ((const u32*)readPos)[1];
+				if (decompressedBlockSize == 0 || decompressedBlockSize > left)
+					return m_logger.Error(TC("Decompressed block size %u is invalid. Decompressed file is %u (%s -> %s)"), decompressedBlockSize, decompressedSize, readHint, writeHint);
+				readPos += sizeof(u32) * 2;
+
+				void* decoderMem = nullptr;
+				u64 decoderMemSize = 0;
+				TimerScope ts(stats.decompressToMem);
+				OO_SINTa decompLen = OodleLZ_Decompress(readPos, (OO_SINTa)compressedBlockSize, writePos, (OO_SINTa)decompressedBlockSize, OodleLZ_FuzzSafe_Yes, OodleLZ_CheckCRC_No, OodleLZ_Verbosity_None, NULL, 0, NULL, NULL, decoderMem, decoderMemSize);
+				if (decompLen != decompressedBlockSize)
+					return m_logger.Error(TC("Expecting to be able to decompress %u to %u bytes at pos %llu but got %llu. File compressed size: %llu (%s -> %s)"), compressedBlockSize, decompressedBlockSize, decompressedSize - left, decompLen, compressedSize, readHint, writeHint);
+				writePos += decompressedBlockSize;
+				readPos += compressedBlockSize;
+				left -= decompressedBlockSize;
+			}
+		}
+		return true;
+	}
+
+	bool DecompressWriter::DecompressMemoryToFile(u8* compressedData, FileAccessor& destination, u64 decompressedSize, bool useNoBuffering)
+	{
+		StorageStats& stats = m_stats;
+		u8* readPos = compressedData;
+
+		u8* slot = m_bufferSlots.Pop();
+		auto _ = MakeGuard([&](){ m_bufferSlots.Push(slot); });
+
+		u64 left = decompressedSize;
+		u64 overflow = 0;
+		while (left)
+		{
+			u32 compressedBlockSize = ((u32*)readPos)[0];
+			if (!compressedBlockSize)
+				break;
+			u32 decompressedBlockSize = ((u32*)readPos)[1];
+
+			readPos += sizeof(u32)*2;
+
+			OO_SINTa decompLen;
+			{
+				TimerScope ts(stats.decompressToMem);
+				decompLen = OodleLZ_Decompress(readPos, (OO_SINTa)compressedBlockSize, slot + overflow, (OO_SINTa)decompressedBlockSize);
+			}
+			UBA_ASSERT(decompLen == decompressedBlockSize);
+			
+			u64 available = overflow + u64(decompLen);
+
+			if (left - available > 0 && available < BufferSlotHalfSize)
+			{
+				overflow += u64(decompLen);
+				readPos += compressedBlockSize;
+				continue;
+			}
+
+			if (useNoBuffering)
+			{
+				u64 writeSize = AlignUp(available - 4096 + 1, 4096);
+
+				if (!destination.Write(slot, writeSize))
+					return false;
+
+				overflow = available - writeSize;
+				readPos += compressedBlockSize;
+				left -= writeSize;
+
+				if (overflow == left)
+				{
+					if (!destination.Write(slot + writeSize, 4096))
+						return false;
+					break;
+				}
+
+				memcpy(slot, slot + writeSize, overflow);
+			}
+			else
+			{
+				u64 writeSize = available;
+				if (!destination.Write(slot, writeSize))
+					return false;
+				readPos += compressedBlockSize;
+				left -= writeSize;
+				overflow = 0;
+			}
+		}
+
+		if (useNoBuffering)
+			if (!SetEndOfFile(m_logger, destination.GetFileName(), destination.GetHandle(), decompressedSize))
+				return false;
+		return true;
+	}
+
+	bool DecompressWriter::DecompressFileToMemory(const tchar* fileName, FileHandle fileHandle, u8* dest, u64 decompressedSize, const tchar* writeHint, u64 fileStartOffset)
+	{
+		if (decompressedSize > BufferSlotSize*4) // Arbitrary size threshold. We want to at least catch the pch here
+		{
+			u64 compressedSize;
+			if (!uba::GetFileSizeEx(compressedSize, fileHandle))
+				return m_logger.Error(TC("GetFileSize failed for %s (%s)"), fileName, LastErrorToText().data);
+			FileMappingHandle fileMapping = uba::CreateFileMappingW(m_logger, fileHandle, PAGE_READONLY, compressedSize, fileName);
+			if (!fileMapping.IsValid())
+				return m_logger.Error(TC("Failed to create file mapping for %s (%s)"), fileName, LastErrorToText().data);
+			auto fmg = MakeGuard([&]() { CloseFileMapping(m_logger, fileMapping, fileName); });
+			u8* fileData = MapViewOfFile(m_logger, fileMapping, FILE_MAP_READ, 0, compressedSize);
+			if (!fileData)
+				return m_logger.Error(TC("Failed to map view of file mapping for %s (%s)"), fileName, LastErrorToText().data);
+			auto udg = MakeGuard([&]() { UnmapViewOfFile(m_logger, fileData, compressedSize, fileName); });
+			
+			u8* readPos = fileData + 8 + fileStartOffset;
+			if (!DecompressMemoryToMemory(readPos, compressedSize, dest, decompressedSize, fileName, writeHint))
+				return false;
+		}
+		else
+		{
+			StorageStats& stats = m_stats;
+			u8* slot = m_bufferSlots.Pop();
+			auto _ = MakeGuard([&]() { m_bufferSlots.Push(slot); });
+
+			void* decoderMem = slot + BufferSlotHalfSize;
+			u64 decoderMemSize = BufferSlotHalfSize;
+
+			u64 bytesRead = 8; // We know size has already been read
+
+			u8* readBuffer = slot;
+			u8* writePos = dest;
+			u64 left = decompressedSize;
+			while (left)
+			{
+				u32 sizes[2];
+				if (!ReadFile(m_logger, fileName, fileHandle, sizes, sizeof(u32) * 2))
+				{
+					u64 compressedSize;
+					if (!uba::GetFileSizeEx(compressedSize, fileHandle))
+						return m_logger.Error(TC("GetFileSize failed for %s (%s)"), fileName, LastErrorToText().data);
+					if (bytesRead + 8 > compressedSize)
+						return m_logger.Error(TC("File %s corrupt. Tried to read 8 bytes. File is smaller than expected (Read: %llu, Size: %llu)"), fileName, bytesRead, compressedSize);
+					return false;
+				}
+				u32 compressedBlockSize = sizes[0];
+				u32 decompressedBlockSize = sizes[1];
+
+				bytesRead += 8;
+
+				if (!ReadFile(m_logger, fileName, fileHandle, readBuffer, compressedBlockSize))
+				{
+					u64 compressedSize;
+					if (!uba::GetFileSizeEx(compressedSize, fileHandle))
+						return m_logger.Error(TC("GetFileSize failed for %s (%s)"), fileName, LastErrorToText().data);
+					if (bytesRead + compressedBlockSize > compressedSize)
+						return m_logger.Error(TC("File %s corrupt. Compressed block size (%u) is larger than what is left of file (%llu)"), fileName, compressedBlockSize, compressedSize - bytesRead);
+					return false;
+				}
+				bytesRead += compressedBlockSize;
+
+				TimerScope ts(stats.decompressToMem);
+				OO_SINTa decompLen = OodleLZ_Decompress(readBuffer, (OO_SINTa)compressedBlockSize, writePos, (OO_SINTa)decompressedBlockSize, OodleLZ_FuzzSafe_Yes, OodleLZ_CheckCRC_No, OodleLZ_Verbosity_None, NULL, 0, NULL, NULL, decoderMem, decoderMemSize);
+				if (decompLen != decompressedBlockSize)
+					return m_logger.Error(TC("Failed to decompress data from file %s at pos %llu"), fileName, decompressedSize - left);
+				writePos += decompressedBlockSize;
+				left -= decompressedBlockSize;
+			}
+		}
+		return true;
+	}
+
+	////////////////////////////////////////////////////////////////////////////////////////////////////
+}
\ No newline at end of file
diff --git a/Engine/Source/Programs/UnrealBuildAccelerator/Common/Public/UbaStorage.h b/Engine/Source/Programs/UnrealBuildAccelerator/Common/Public/UbaStorage.h
--- a/Engine/Source/Programs/UnrealBuildAccelerator/Common/Public/UbaStorage.h
+++ b/Engine/Source/Programs/UnrealBuildAccelerator/Common/Public/UbaStorage.h
@@ -187,7 +187,6 @@ namespace uba
 		void CasEntryAccessed(const CasKey& casKey);
 		virtual bool IsDisallowedPath(const tchar* fileName);
 		virtual bool DecompressMemoryToMemory(const u8* compressedData, u64 compressedSize, u8* writeData, u64 decompressedSize, const tchar* readHint, const tchar* writeHint) override;
-		bool DecompressMemoryToFile(u8* compressedData, FileAccessor& destination, u64 decompressedSize, bool useNoBuffering);
 
 		void CasEntryAccessed(CasEntry& entry);
 		void CasEntryWritten(CasEntry& entry, u64 size);
diff --git a/Engine/Source/Programs/UnrealBuildAccelerator/Common/Public/UbaStorageUtils.h b/Engine/Source/Programs/UnrealBuildAccelerator/Common/Public/UbaStorageUtils.h
--- a/Engine/Source/Programs/UnrealBuildAccelerator/Common/Public/UbaStorageUtils.h
+++ b/Engine/Source/Programs/UnrealBuildAccelerator/Common/Public/UbaStorageUtils.h
@@ -13,6 +13,8 @@ namespace uba
 	struct BufferSlots;
 	struct StorageStats;
 
+	////////////////////////////////////////////////////////////////////////////////////////////////////
+
 	static constexpr u8 DefaultCompressor = OodleLZ_Compressor_Kraken;
 	static constexpr u8 DefaultCompressionLevel = OodleLZ_CompressionLevel_SuperFast;
 
@@ -23,6 +25,8 @@ namespace uba
 	bool SendBatchMessages(Logger& logger, NetworkClient& client, u16 fetchId, u8* slot, u64 capacity, u64 left, u32 messageMaxSize, u32& readIndex, u32& responseSize, const Function<bool()>& runInWaitFunc = {}, const tchar* hint = TC(""), u32* outError = nullptr);
 	bool SendFile(Logger& logger, NetworkClient& client, const CasKey& casKey, const u8* sourceMem, u64 sourceSize, const tchar* hint);
 
+	////////////////////////////////////////////////////////////////////////////////////////////////////
+
 	struct FileSender
 	{
 		bool SendFileCompressed(const CasKey& casKey, const tchar* fileName, const u8* sourceMem, u64 sourceSize, const tchar* hint);
@@ -38,6 +42,8 @@ namespace uba
 		u64 m_bytesSent = 0;
 	};
 
+	////////////////////////////////////////////////////////////////////////////////////////////////////
+
 	struct FileFetcher
 	{
 		bool RetrieveFile(Logger& logger, NetworkClient& client, const CasKey& casKey, const tchar* destination, bool writeCompressed, MemoryBlock* destinationMem = nullptr);
@@ -52,6 +58,8 @@ namespace uba
 		u64 bytesReceived = 0;
 	};
 
+	////////////////////////////////////////////////////////////////////////////////////////////////////
+
 	struct LinearWriter
 	{
 		virtual bool Write(const void* data, u64 dataLen) = 0;
@@ -59,6 +67,8 @@ namespace uba
 		virtual const tchar* GetHint() = 0;
 	};
 
+	////////////////////////////////////////////////////////////////////////////////////////////////////
+
 	struct CompressWriter
 	{
 		CompressWriter(Logger& l, BufferSlots& bs, WorkManager& wm, StorageStats& s, u8 cc = DefaultCompressor, u8 ccl = DefaultCompressionLevel, bool avof = false);
@@ -77,6 +87,8 @@ namespace uba
 		bool m_asyncUnmapViewOfFile;
 	};
 
+	////////////////////////////////////////////////////////////////////////////////////////////////////
+
 	struct LinearWriterFile final : public LinearWriter
 	{
 		LinearWriterFile(FileAccessor& f);
@@ -87,6 +99,8 @@ namespace uba
 		u64 written = 0;
 	};
 
+	////////////////////////////////////////////////////////////////////////////////////////////////////
+
 	struct LinearWriterMem final : public LinearWriter
 	{
 		LinearWriterMem(u8* d, u64 c);
@@ -97,4 +111,20 @@ namespace uba
 		u64 size;
 		u64 capacity;
 	};
+
+	////////////////////////////////////////////////////////////////////////////////////////////////////
+
+	struct DecompressWriter
+	{
+		bool DecompressMemoryToMemory(const u8* compressedData, u64 compressedSize, u8* writeData, u64 decompressedSize, const tchar* readHint, const tchar* writeHint);
+		bool DecompressMemoryToFile(u8* compressedData, FileAccessor& destination, u64 decompressedSize, bool useNoBuffering);
+		bool DecompressFileToMemory(const tchar* fileName, FileHandle fileHandle, u8* dest, u64 decompressedSize, const tchar* writeHint, u64 fileStartOffset);
+
+		Logger& m_logger;
+		BufferSlots& m_bufferSlots;
+		WorkManager& m_workManager;
+		StorageStats& m_stats;
+	};
+
+	////////////////////////////////////////////////////////////////////////////////////////////////////
 }
\ No newline at end of file
diff --git a/Engine/Source/Programs/UnrealBuildAccelerator/Core/Private/UbaPlatform.cpp b/Engine/Source/Programs/UnrealBuildAccelerator/Core/Private/UbaPlatform.cpp
--- a/Engine/Source/Programs/UnrealBuildAccelerator/Core/Private/UbaPlatform.cpp
+++ b/Engine/Source/Programs/UnrealBuildAccelerator/Core/Private/UbaPlatform.cpp
@@ -716,12 +716,25 @@ namespace uba
 #endif
 	}
 
+	#if UBA_USE_MIMALLOC
+	void* Oodle_MallocAligned(OO_SINTa bytes, OO_S32 alignment)
+	{
+		return mi_malloc_aligned(bytes, alignment);
+	}
+
+	void Oodle_Free(void* ptr)
+	{
+		mi_free(ptr);
+	}
+	#endif
+
 	void InitMemory()
 	{
 		#if UBA_USE_MIMALLOC
 		mi_option_set(mi_option_purge_delay, 100);
 		//mi_option_set(mi_option_eager_commit, 1);
 		//mi_option_set(mi_option_purge_decommits, 0);
+		OodleCore_Plugins_SetAllocators(Oodle_MallocAligned, Oodle_Free);
 		#endif
 	}
 
