[UBA]
* Removed this capture in CompressWriter since it can go out of scope while work objects are still left (but  no work is left)

diff --git a/Engine/Source/Programs/UnrealBuildAccelerator/Common/Private/UbaStorageUtils.cpp b/Engine/Source/Programs/UnrealBuildAccelerator/Common/Private/UbaStorageUtils.cpp
--- a/Engine/Source/Programs/UnrealBuildAccelerator/Common/Private/UbaStorageUtils.cpp
+++ b/Engine/Source/Programs/UnrealBuildAccelerator/Common/Private/UbaStorageUtils.cpp
@@ -906,7 +906,9 @@ namespace uba
 				aligned_free(events);
 				#endif
 			}
+			Logger* logger;
 			LinearWriter* destination;
+			BufferSlots* bufferSlots;
 			Atomic<u64> refCount;
 			Atomic<u64> compressCounter;
 			Event* events;
@@ -915,15 +917,21 @@ namespace uba
 			u64 workCount = 0;
 			u64 maxUncompressedBlock = 0;
 			u64 fileSize = 0;
+			u8 casCompressor = 0;
+			u8 casCompressionLevel = 0;
 			bool error = false;
 		};
 
 		WorkRec* recPtr = new WorkRec(workCount);
 		WorkRec& rec = *recPtr;
+		rec.logger = &m_logger;
 		rec.destination = &destination;
+		rec.bufferSlots = &m_bufferSlots;
 		rec.uncompressedData = uncompressedData;
 		rec.maxUncompressedBlock = maxUncompressedBlock;
 		rec.fileSize = fileSize;
+		rec.casCompressor = m_casCompressor;
+		rec.casCompressionLevel = m_casCompressionLevel;
 
 		for (u32 i=0; i!=workCount; ++i)
 			rec.events[i].Create(true);
@@ -932,13 +940,13 @@ namespace uba
 
 		auto& kernelStats = KernelStats::GetCurrent();
 
-		auto work = [recPtr, this, &kernelStats](const WorkContext&)
+		auto work = [recPtr, &kernelStats](const WorkContext&)
 		{
 			auto& rec = *recPtr;
 			KernelStatsScope kss(kernelStats);
 
-			u8* slot = m_bufferSlots.Pop();
-			auto _ = MakeGuard([&](){ m_bufferSlots.Push(slot); });
+			u8* slot = rec.bufferSlots->Pop();
+			auto _ = MakeGuard([&](){ rec.bufferSlots->Push(slot); });
 			u8* compressSlotBuffer = slot + BufferSlotHalfSize;
 			while (true)
 			{
@@ -957,10 +965,10 @@ namespace uba
 					void* scratchMem = slot;
 					u64 scratchSize = BufferSlotHalfSize;
 					TimerScope kts(kernelStats.memoryCompress);
-					compressedBlockSize = OodleLZ_Compress((OodleLZ_Compressor)m_casCompressor, uncompressedDataSlot, uncompressedBlockSize, compressSlotBuffer + 8, (OodleLZ_CompressionLevel)m_casCompressionLevel, NULL, NULL, NULL, scratchMem, scratchSize);
+					compressedBlockSize = OodleLZ_Compress((OodleLZ_Compressor)rec.casCompressor, uncompressedDataSlot, uncompressedBlockSize, compressSlotBuffer + 8, (OodleLZ_CompressionLevel)rec.casCompressionLevel, NULL, NULL, NULL, scratchMem, scratchSize);
 					if (compressedBlockSize == OODLELZ_FAILED)
 					{
-						m_logger.Error(TC("Failed to compress %llu bytes for %s"), u64(uncompressedBlockSize), rec.destination->GetHint());
+						rec.logger->Error(TC("Failed to compress %llu bytes for %s"), u64(uncompressedBlockSize), rec.destination->GetHint());
 						rec.error = true;
 						return;
 					}
@@ -1036,7 +1044,7 @@ namespace uba
 				auto udg = MakeGuard([&]()
 					{
 						if (m_asyncUnmapViewOfFile)
-							m_workManager.AddWork([this, uncompressedData, fileSize, f = TString(from)](const WorkContext&) { UnmapViewOfFile(m_logger, uncompressedData, fileSize, f.c_str()); }, 1, TC("UnmapFile"));
+							m_workManager.AddWork([uncompressedData, fileSize, f = TString(from), loggerPtr = &m_logger](const WorkContext&) { UnmapViewOfFile(*loggerPtr, uncompressedData, fileSize, f.c_str()); }, 1, TC("UnmapFile"));
 						else
 							UnmapViewOfFile(m_logger, uncompressedData, fileSize, from);
 					});
