[UBA]
* Added more code validation around deferred cas creation..

diff --git a/Engine/Source/Programs/UnrealBuildAccelerator/Common/Private/UbaStorage.cpp b/Engine/Source/Programs/UnrealBuildAccelerator/Common/Private/UbaStorage.cpp
--- a/Engine/Source/Programs/UnrealBuildAccelerator/Common/Private/UbaStorage.cpp
+++ b/Engine/Source/Programs/UnrealBuildAccelerator/Common/Private/UbaStorage.cpp
@@ -536,32 +536,43 @@ namespace uba
 		if (!casEntry.verified && VerifyExisting(verifyReturnValue, entryLock, casKey, casEntry, casFile))
 			return verifyReturnValue;
 
-		if (deferCreation)
 		{
-			SCOPED_WRITE_LOCK(m_deferredCasCreationLookupLock, deferredLock);
-			auto res = m_deferredCasCreationLookup.try_emplace(casKey);
-			if (res.second)
+			SCOPED_FUTEX(m_deferredCasCreationLookupLock, deferredLock);
+			if (deferCreation)
 			{
-				DeferedCasCreation& dcc = res.first->second;
-				dcc.fileName = fileName;
-				dcc.fileSize = fileSize;
-				dcc.lastWriteTime = lastWriteTime;
+				auto res = m_deferredCasCreationLookup.try_emplace(casKey);
+				if (res.second)
+				{
+					DeferedCasCreation& dcc = res.first->second;
+					dcc.fileName = fileName;
+					dcc.fileSize = fileSize;
+					dcc.lastWriteTime = lastWriteTime;
+				}
+				auto res2 = m_deferredCasCreationLookupByName.try_emplace(fileNameKey, casKey);
+				if (res2.second)
+					res.first->second.names.push_back(res2.first);
+				else if (res2.first->second != casKey)
+					return m_logger.Error(TC("Same file %s is registered with different caskey (%s vs %s)"), fileName, CasKeyString(res2.first->second).str, CasKeyString(casKey).str);
+				return true;
+			}
+			else
+			{
+				auto findIt = m_deferredCasCreationLookupByName.find(fileNameKey);
+				if (findIt != m_deferredCasCreationLookupByName.end() && findIt->second != casKey)
+					return m_logger.Error(TC("Found deferred entry for %s with different cas key while writing cas. A ReportFileWrite is missing somewhere"), fileName);
+
+				auto findIt2 = m_deferredCasCreationLookup.find(casKey);
+				if (findIt2 != m_deferredCasCreationLookup.end())
+				{
+					#if UBA_DEBUG
+					return m_logger.Error(TC("This should not happen. (%s)"), fileName);
+					#endif
+					for (auto& it2 : findIt2->second.names)
+						m_deferredCasCreationLookupByName.erase(it2);
+					m_deferredCasCreationLookup.erase(findIt2);
+				}
 			}
-			auto res2 = m_deferredCasCreationLookupByName.try_emplace(fileNameKey, casKey);
-			if (res2.second)
-				res.first->second.names.push_back(res2.first);
-			else if (res2.first->second != casKey)
-				return m_logger.Error(TC("Same file %s is registered with different caskey (%s vs %s)"), fileName, CasKeyString(res2.first->second).str, CasKeyString(casKey).str);
-			return true;
-		}
-		#if UBA_DEBUG
-		else
-		{
-			SCOPED_READ_LOCK(m_deferredCasCreationLookupLock, deferredLock);
-			auto findIt = m_deferredCasCreationLookupByName.find(fileNameKey);
-			UBA_ASSERTF(findIt == m_deferredCasCreationLookupByName.end() || findIt->second == casKey, TC("Found deferred entry for %s with different cas key while writing cas. A ReportFileWrite is missing somewhere"), fileName);
 		}
-		#endif
 
 		casEntry.verified = true;
 		casEntry.exists = false;
@@ -1768,7 +1779,7 @@ namespace uba
 		if (casEntry.verified && casEntry.exists)
 			return true;
 
-		SCOPED_WRITE_LOCK(m_deferredCasCreationLookupLock, deferredLock);
+		SCOPED_FUTEX(m_deferredCasCreationLookupLock, deferredLock);
 		auto findIt = m_deferredCasCreationLookup.find(casKey);
 		if (findIt == m_deferredCasCreationLookup.end())
 			return false;
@@ -1794,7 +1805,7 @@ namespace uba
 			return m_logger.Error(TC("Failed to write deferred cas %s (%s)"), CasKeyString(casKey).str, deferredCreation.data);
 
 		if (res.readSize != deferredFileSize) // TODO: Should also check last written time
-			m_logger.Warning(TC("Deferred cas %s was deferred with file size %llu but was read with size %llu (%s)"), CasKeyString(casKey).str, deferredFileSize, res.readSize, deferredCreation.data);
+			m_logger.Warning(TC("Deferred cas %s was deferred with source file size %llu but was read with source file size %llu (%s)"), CasKeyString(casKey).str, deferredFileSize, res.readSize, deferredCreation.data);
 
 		bool hasMapping = false;
 		if (res.mappingHandle.IsValid())
@@ -2290,7 +2301,7 @@ namespace uba
 	void StorageImpl::ReportFileWrite(StringKey fileNameKey, const tchar* fileName)
 	{
 		// If a defered cas creation is queued up while the source file is about to be modified we need to flush out the cas creation before modifying the file
-		SCOPED_READ_LOCK(m_deferredCasCreationLookupLock, deferredLock);
+		SCOPED_FUTEX_READ(m_deferredCasCreationLookupLock, deferredLock);
 		auto findIt = m_deferredCasCreationLookupByName.find(fileNameKey);
 		if (findIt == m_deferredCasCreationLookupByName.end())
 			return;
diff --git a/Engine/Source/Programs/UnrealBuildAccelerator/Common/Private/UbaStorageServer.cpp b/Engine/Source/Programs/UnrealBuildAccelerator/Common/Private/UbaStorageServer.cpp
--- a/Engine/Source/Programs/UnrealBuildAccelerator/Common/Private/UbaStorageServer.cpp
+++ b/Engine/Source/Programs/UnrealBuildAccelerator/Common/Private/UbaStorageServer.cpp
@@ -817,7 +817,7 @@ namespace uba
 
 				{
 					// Check if there is a deferred creation of this file, if there is we can just as well just use the local file
-					SCOPED_WRITE_LOCK(m_deferredCasCreationLookupLock, deferredLock);
+					SCOPED_FUTEX(m_deferredCasCreationLookupLock, deferredLock);
 					auto findIt = m_deferredCasCreationLookup.find(casKey);
 					if (findIt != m_deferredCasCreationLookup.end())
 					{
diff --git a/Engine/Source/Programs/UnrealBuildAccelerator/Common/Public/UbaStorage.h b/Engine/Source/Programs/UnrealBuildAccelerator/Common/Public/UbaStorage.h
--- a/Engine/Source/Programs/UnrealBuildAccelerator/Common/Public/UbaStorage.h
+++ b/Engine/Source/Programs/UnrealBuildAccelerator/Common/Public/UbaStorage.h
@@ -277,7 +277,7 @@ namespace uba
 
 		FileMappingBuffer m_casDataBuffer;
 
-		ReaderWriterLock m_deferredCasCreationLookupLock;
+		Futex m_deferredCasCreationLookupLock;
 		using NameToCas = UnorderedMap<StringKey, CasKey>;
 		struct DeferedCasCreation { StringKey fileNameKey; TString fileName; List<NameToCas::iterator> names; u64 fileSize = InvalidValue; u64 lastWriteTime = 0;};
 		UnorderedMap<CasKey, DeferedCasCreation> m_deferredCasCreationLookup;
