[UBA]
* Renamed some variables to make it easier to understand the code.
size of file/mapping that can be both compressed or uncompressed is now called "storedSize"
size of file/mapping that is never compressed or contains any header is called "contentSize"
* Fixed so storedSize is set for file mappings that are created directly from files

diff --git a/Engine/Source/Programs/UnrealBuildAccelerator/Common/Private/UbaSession.cpp b/Engine/Source/Programs/UnrealBuildAccelerator/Common/Private/UbaSession.cpp
--- a/Engine/Source/Programs/UnrealBuildAccelerator/Common/Private/UbaSession.cpp
+++ b/Engine/Source/Programs/UnrealBuildAccelerator/Common/Private/UbaSession.cpp
@@ -348,7 +348,7 @@ namespace uba
 					return false;
 
 				shouldEvictFromMemory = false; // Can't evict without properly update filemappingtable.. the file on disk does now not match what was registered for write
-				writtenSize = res.size;
+				writtenSize = res.storedSize;
 			}
 			else
 			{
@@ -432,14 +432,14 @@ namespace uba
 			entry.handled = true;
 			entry.mapping = mh;
 			entry.mappingOffset = 0;
-			entry.size = file.mappingWritten;
+			entry.contentSize = file.mappingWritten;
 			entry.lastWriteTime = file.lastWriteTime;
 			entry.isDir = false;
 			entry.success = true;
 			entry.canBeFreed = true;
 			entry.usedCount = 0;
 			entry.usedCountBeforeFree = g_globalRules.GetUsedCountBeforeFree(file.name);
-			entry.writtenSize = writtenSize;
+			entry.storedSize = writtenSize;
 			entry.isInvisible = rules.IsInvisible(file.name);
 
 			#if UBA_DEBUG_TRACK_MAPPING
@@ -506,7 +506,7 @@ namespace uba
 			if (deleteInternalMapping && !entry.isInvisible)
 				return m_logger.Error(TC("Trying to delete mapping that is not invisible (%s)"), filePath.data);
 			mapping = entry.mapping;
-			size = entry.size;
+			size = entry.contentSize;
 			if (deleteInternalMapping)
 			{
 				if (!entry.isInvisible)
@@ -800,7 +800,7 @@ namespace uba
 			return entry.success ? m_fileMappingTableSize : 0;
 		}
 
-		entry.size = fileSize;
+		entry.contentSize = fileSize;
 		entry.isDir = false;
 		entry.success = true;
 		entry.mapping = {};
@@ -848,7 +848,7 @@ namespace uba
 
 			if (!entry.success)
 				return false;
-			out.size = entry.size;
+			out.size = entry.contentSize;
 			if (entry.mapping.IsValid())
 			{
 				Storage::GetMappingString(out.name, entry.mapping, entry.mappingOffset);
@@ -885,7 +885,8 @@ namespace uba
 		}
 		auto _ = MakeGuard([&](){ uba::CloseFile(fileName, fileHandle); });
 
-		u64 size = 0;
+		u64 storedSize = InvalidValue;
+		u64 contentSize = 0;
 		u64 lastWriteTime = 0;
 		u64 fileStartOffset = 0;
 
@@ -894,7 +895,7 @@ namespace uba
 		{
 			if (isCompressedCas)
 			{
-				if (!ReadFile(m_logger, fileName, fileHandle, &size, sizeof(u64)))
+				if (!ReadFile(m_logger, fileName, fileHandle, &contentSize, sizeof(u64)))
 					return m_logger.Error(TC("Failed to read first bytes from file %s (%s)"), fileName, LastErrorToText().data);
 			}
 			else
@@ -903,7 +904,8 @@ namespace uba
 				if (!GetFileBasicInformationByHandle(info, m_logger, fileName, fileHandle))
 					return false;
 
-				size = info.size;
+				storedSize = info.size;
+				contentSize = info.size;
 				lastWriteTime = info.lastWriteTime;
 
 				if (m_readIntermediateFilesCompressed && info.size > sizeof(CompressedFileHeader) && g_globalRules.FileCanBeCompressed(ToView(fileName)))
@@ -915,7 +917,7 @@ namespace uba
 					{
 						fileStartOffset = sizeof(CompressedFileHeader);
 						isCompressed = true;
-						if (!ReadFile(m_logger, fileName, fileHandle, &size, sizeof(u64)))
+						if (!ReadFile(m_logger, fileName, fileHandle, &contentSize, sizeof(u64)))
 							return m_logger.Error(TC("Failed to read first bytes from file %s (%s)"), fileName, LastErrorToText().data);
 					}
 					else
@@ -927,7 +929,7 @@ namespace uba
 			}
 		}
 
-		if (isDir || size == 0)
+		if (isDir || contentSize == 0)
 		{
 			if (isDir)
 				out.name.Append(TCV("$d"));
@@ -936,8 +938,8 @@ namespace uba
 		}
 		else
 		{
-			if (size > m_fileMappingBuffer.GetFileMappingCapacity())
-				return m_logger.Error(TC("File %s has a size (%llu) that is too large to fit in mapping buffer (%s)"), fileName, size, hint);
+			if (contentSize > m_fileMappingBuffer.GetFileMappingCapacity())
+				return m_logger.Error(TC("File %s has a size (%llu) that is too large to fit in mapping buffer (%s)"), fileName, contentSize, hint);
 
 			FileMappingHandle mapping;
 			u64 mappingOffset = 0;
@@ -946,16 +948,16 @@ namespace uba
 			MappedView mappedView;
 
 			auto ownedGuard1 = MakeGuard([&]() { CloseFileMapping(m_logger, mapping, fileName); });
-			auto ownedGuard2 = MakeGuard([&]() { m_workManager.AddWork([this, mappingMemory, size, fn = TString(fileName)](const WorkContext&) { UnmapViewOfFile(m_logger, mappingMemory, size, fn.c_str()); }, 1, TC("UnmapViewOfFile")); });
+			auto ownedGuard2 = MakeGuard([&]() { m_workManager.AddWork([this, mappingMemory, contentSize, fn = TString(fileName)](const WorkContext&) { UnmapViewOfFile(m_logger, mappingMemory, contentSize, fn.c_str()); }, 1, TC("UnmapViewOfFile")); });
 			auto viewGuard = MakeGuard([&]() { m_fileMappingBuffer.UnmapView(mappedView, fileName); });
 
 			if (canBeFreed)
 			{
 				viewGuard.Cancel();
-				mapping = CreateMemoryMappingW(m_logger, PAGE_READWRITE, size, nullptr, fileName);
+				mapping = CreateMemoryMappingW(m_logger, PAGE_READWRITE, contentSize, nullptr, fileName);
 				if (!mapping.IsValid())
 					return false;
-				mappingMemory = MapViewOfFile(m_logger, mapping, FILE_MAP_WRITE, 0, size);
+				mappingMemory = MapViewOfFile(m_logger, mapping, FILE_MAP_WRITE, 0, contentSize);
 				if (!mappingMemory)
 					return false;
 			}
@@ -965,7 +967,7 @@ namespace uba
 				UBA_ASSERT(!entry.canBeFreed);
 				ownedGuard1.Cancel();
 				ownedGuard2.Cancel();
-				mappedView = m_fileMappingBuffer.AllocAndMapView(MappedView_Transient, size, alignment, fileName, false);
+				mappedView = m_fileMappingBuffer.AllocAndMapView(MappedView_Transient, contentSize, alignment, fileName, false);
 				if (!mappedView.memory)
 					return false;
 				mapping = mappedView.handle;
@@ -975,12 +977,12 @@ namespace uba
 
 			if (isCompressed)
 			{
-				if (!m_storage.DecompressFileToMemory(fileName, fileHandle, mappingMemory, size, TC("GetOrCreateMemoryMapFromFile"), fileStartOffset))
+				if (!m_storage.DecompressFileToMemory(fileName, fileHandle, mappingMemory, contentSize, TC("GetOrCreateMemoryMapFromFile"), fileStartOffset))
 					return false;
 			}
 			else
 			{
-				if (!ReadFile(m_logger, fileName, fileHandle, mappingMemory, size))
+				if (!ReadFile(m_logger, fileName, fileHandle, mappingMemory, contentSize))
 					return false;
 			}
 
@@ -1009,7 +1011,7 @@ namespace uba
 			BinaryWriter writer(m_fileMappingTableMem, m_fileMappingTableSize);
 			writer.WriteStringKey(fileNameKey);
 			writer.WriteString(out.name);
-			writer.Write7BitEncoded(size);
+			writer.Write7BitEncoded(contentSize);
 			m_fileMappingTableSize = (u32)writer.GetPosition();
 		}
 
@@ -1019,10 +1021,11 @@ namespace uba
 		#endif
 
 		entry.isDir = isDir;
-		entry.size = size;
+		entry.contentSize = contentSize;
+		entry.storedSize = storedSize;
 		entry.lastWriteTime = lastWriteTime;
 		
-		out.size = size;
+		out.size = contentSize;
 		return true;
 	}
 
@@ -1045,7 +1048,7 @@ namespace uba
 			entryCs.Leave();
 			if (!entry.success)
 				return false;
-			out.size = entry.size;
+			out.size = entry.contentSize;
 			if (entry.mapping.IsValid())
 				Storage::GetMappingString(out.name, entry.mapping, entry.mappingOffset);
 			else
@@ -1061,30 +1064,30 @@ namespace uba
 		if (!mappedViewRead.handle.IsValid())
 			return false;
 
-		u64 size = InvalidValue;
+		u64 contentSize = InvalidValue;
 
 		if (mappedViewRead.isCompressed)
 		{
 			auto mvrg = MakeGuard([&](){ m_storage.UnmapView(mappedViewRead, fileName); });
 			u8* readMemory = mappedViewRead.memory;
-			size = *(u64*)readMemory;
+			contentSize = *(u64*)readMemory;
 			readMemory += 8;
 
-			if (size == 0)
+			if (contentSize == 0)
 			{
 				out.name.Append(TCV("$f"));
 			}
 			else
 			{
-				if (size > m_fileMappingBuffer.GetFileMappingCapacity())
-					return m_logger.Error(TC("File %s has a size (%llu) that is too large to fit in mapping buffer (GetOrCreateMemoryMapFromView)"), fileName, size);
+				if (contentSize > m_fileMappingBuffer.GetFileMappingCapacity())
+					return m_logger.Error(TC("File %s has a size (%llu) that is too large to fit in mapping buffer (GetOrCreateMemoryMapFromView)"), fileName, contentSize);
 
-				auto mappedViewWrite = m_fileMappingBuffer.AllocAndMapView(MappedView_Transient, size, alignment, fileName);
+				auto mappedViewWrite = m_fileMappingBuffer.AllocAndMapView(MappedView_Transient, contentSize, alignment, fileName);
 				if (!mappedViewWrite.memory)
 					return false;
 				auto unmapGuard = MakeGuard([&](){ m_fileMappingBuffer.UnmapView(mappedViewWrite, fileName); });
 
-				if (!m_storage.DecompressMemoryToMemory(readMemory, mappedViewRead.size, mappedViewWrite.memory, size, fileName, TC("TransientMapping")))
+				if (!m_storage.DecompressMemoryToMemory(readMemory, mappedViewRead.size, mappedViewWrite.memory, contentSize, fileName, TC("TransientMapping")))
 					return false;
 				unmapGuard.Execute();
 
@@ -1100,7 +1103,7 @@ namespace uba
 			entry.mappingOffset = mappedViewRead.offset;
 			Storage::GetMappingString(out.name, mappedViewRead.handle, mappedViewRead.offset);
 			entry.mapping = mappedViewRead.handle;
-			size = mappedViewRead.size;
+			contentSize = mappedViewRead.size;
 		}
 		entry.success = true;
 
@@ -1109,14 +1112,14 @@ namespace uba
 			BinaryWriter writer(m_fileMappingTableMem, m_fileMappingTableSize);
 			writer.WriteStringKey(fileNameKey);
 			writer.WriteString(out.name);
-			writer.Write7BitEncoded(size);
+			writer.Write7BitEncoded(contentSize);
 			m_fileMappingTableSize = (u32)writer.GetPosition();
 		}
 
 		entry.isDir = false;
-		entry.size = size;
+		entry.contentSize = contentSize;
 		
-		out.size = size;
+		out.size = contentSize;
 
 		#if UBA_DEBUG_TRACK_MAPPING
 		entry.name = fileName;
@@ -1743,9 +1746,9 @@ namespace uba
 					return false;
 				out.handle = entry.mapping;
 				out.offset = entry.mappingOffset;
-				out.size = entry.size;
+				out.size = entry.contentSize;
 				out.lastWriteTime = entry.lastWriteTime;
-				out.writtenSize = entry.writtenSize;
+				out.storedSize = entry.storedSize;
 				out.createIndependentMapping = entry.createIndependentMapping;
 				return true;
 			});
@@ -3541,7 +3544,7 @@ namespace uba
 
 		entry.mapping = mappingHandle;
 		entry.mappingOffset = mappingOffset;
-		entry.size = mappingSize;
+		entry.contentSize = mappingSize;
 		entry.handled = true;
 		entry.lastWriteTime = 0; // TODO: Take lastwritetime of source file?
 		entry.createIndependentMapping = true; // We want independent mapping in casdb so they can be deleted
@@ -3574,7 +3577,7 @@ namespace uba
 		FileMappingEntry& entry = findIt->second;
 		lookupLock.Leave();
 		SCOPED_FUTEX(entry.lock, entryCs);
-		outSize = entry.size;
+		outSize = entry.contentSize;
 		return true;
 	}
 
diff --git a/Engine/Source/Programs/UnrealBuildAccelerator/Common/Private/UbaSessionClient.cpp b/Engine/Source/Programs/UnrealBuildAccelerator/Common/Private/UbaSessionClient.cpp
--- a/Engine/Source/Programs/UnrealBuildAccelerator/Common/Private/UbaSessionClient.cpp
+++ b/Engine/Source/Programs/UnrealBuildAccelerator/Common/Private/UbaSessionClient.cpp
@@ -627,12 +627,12 @@ namespace uba
 						if (!m_storage.RetrieveCasFile(result, casKey, fileName.data, &m_fileMappingBuffer, memoryMapAlignment, allowProxy))
 							return m_logger.Error(TC("Error retrieving cas entry %s (%s)"), CasKeyString(casKey).str, fileName.data);
 						entry.success = true;
-						entry.size = result.size;
+						entry.contentSize = result.size;
 						entry.mapping = result.view.handle;
 						entry.mappingOffset = result.view.offset;
 					}
 
-					fileSize = entry.size;
+					fileSize = entry.contentSize;
 					if (entry.mapping.IsValid())
 						Storage::GetMappingString(newName, entry.mapping, entry.mappingOffset);
 					else
diff --git a/Engine/Source/Programs/UnrealBuildAccelerator/Common/Private/UbaSessionServer.cpp b/Engine/Source/Programs/UnrealBuildAccelerator/Common/Private/UbaSessionServer.cpp
--- a/Engine/Source/Programs/UnrealBuildAccelerator/Common/Private/UbaSessionServer.cpp
+++ b/Engine/Source/Programs/UnrealBuildAccelerator/Common/Private/UbaSessionServer.cpp
@@ -2538,7 +2538,7 @@ namespace uba
 								Storage::GetMappingString(mappingName, entry.mapping, entry.mappingOffset);
 							else
 								mappingName.Append(TCV("Not valid"));
-							logStr.Appendf(TC(" Success: %s Size: %u IsDir: %s Mapping name: %s Mapping offset: %u"), ToString(entry.success), entry.size, ToString(entry.isDir), mappingName.data, entry.mappingOffset);
+							logStr.Appendf(TC(" Success: %s Size: %u IsDir: %s Mapping name: %s Mapping offset: %u"), ToString(entry.success), entry.contentSize, ToString(entry.isDir), mappingName.data, entry.mappingOffset);
 						}
 						else
 						{
diff --git a/Engine/Source/Programs/UnrealBuildAccelerator/Common/Private/UbaStorage.cpp b/Engine/Source/Programs/UnrealBuildAccelerator/Common/Private/UbaStorage.cpp
--- a/Engine/Source/Programs/UnrealBuildAccelerator/Common/Private/UbaStorage.cpp
+++ b/Engine/Source/Programs/UnrealBuildAccelerator/Common/Private/UbaStorage.cpp
@@ -326,7 +326,7 @@ namespace uba
 	bool StorageImpl::WriteCompressedFile(WriteResult& out, const tchar* from, FileHandle readHandle, u8* readMem, u64 fileSize, const tchar* to, const void* header, u64 headerSize, u64 lastWriteTime)
 	{
 		CompressWriter w(m_logger, m_bufferSlots, m_workManager, Stats(), m_casCompressor, m_casCompressionLevel, m_asyncUnmapViewOfFile);
-		return w.CompressToFile(out.size, from, readHandle, readMem, fileSize, to, header, headerSize, lastWriteTime, m_tempPath.data);
+		return w.CompressToFile(out.storedSize, from, readHandle, readMem, fileSize, to, header, headerSize, lastWriteTime, m_tempPath.data);
 	}
 
 	bool StorageImpl::WriteCasFileNoCheck(WriteResult& out, const StringKey& fileNameKey, const tchar* fileName, const CasKey& casKey, const tchar* casFile, bool storeCompressed)
@@ -343,13 +343,13 @@ namespace uba
 				u8* fileMem = scope.MapView(TC("WriteCompressed"), externalMapping.handle, externalMapping.offset, externalMapping.size, fileName);
 				if (!fileMem)
 					return false;
-				out.readSize = externalMapping.writtenSize;
+				out.sourceStoredSize = externalMapping.storedSize;
 				out.isIndependent = externalMapping.createIndependentMapping;
 				if (externalMapping.createIndependentMapping)
 				{
 					CompressWriter w(m_logger, m_bufferSlots, m_workManager, Stats(), m_casCompressor, m_casCompressionLevel, m_asyncUnmapViewOfFile);
 					out.mappingOffset = 0;
-					return w.CompressToMapping(out.mappingHandle, out.size, fileMem, externalMapping.size, casFile);
+					return w.CompressToMapping(out.mappingHandle, out.storedSize, fileMem, externalMapping.size, casFile);
 				}
 				else
 					return WriteCompressedFile(out, fileName, InvalidFileHandle, fileMem, externalMapping.size, casFile, nullptr, 0);
@@ -364,7 +364,7 @@ namespace uba
 		u64 readFileSize = 0;
 		if (!uba::GetFileSizeEx(readFileSize, readHandle))
 			return m_logger.Error(TC("[WriteCasFileNoCheck] GetFileSize failed for %s (%s)"), fileName, LastErrorToText().data);
-		out.readSize = readFileSize;
+		out.sourceStoredSize = readFileSize;
 
 		bool readIsCompressed = false;
 		if (readFileSize >= sizeof(CompressedFileHeader) && g_globalRules.FileCanBeCompressed(ToView(fileName)))
@@ -414,7 +414,7 @@ namespace uba
 		if (!destinationFile.Close())
 			return false;
 
-		out.size = readFileSize;
+		out.storedSize = readFileSize;
 		return true;
 	}
 
@@ -536,42 +536,43 @@ namespace uba
 		if (!casEntry.verified && VerifyExisting(verifyReturnValue, entryLock, casKey, casEntry, casFile))
 			return verifyReturnValue;
 
+		if (deferCreation)
 		{
 			SCOPED_FUTEX(m_deferredCasCreationLookupLock, deferredLock);
-			if (deferCreation)
+			auto res = m_deferredCasCreationLookup.try_emplace(casKey);
+			if (res.second)
 			{
-				auto res = m_deferredCasCreationLookup.try_emplace(casKey);
-				if (res.second)
-				{
-					DeferedCasCreation& dcc = res.first->second;
-					dcc.fileName = fileName;
-					dcc.fileSize = fileSize;
-					dcc.lastWriteTime = lastWriteTime;
-				}
-				auto res2 = m_deferredCasCreationLookupByName.try_emplace(fileNameKey, casKey);
-				if (res2.second)
-					res.first->second.names.push_back(res2.first);
-				else if (res2.first->second != casKey)
-					return m_logger.Error(TC("Same file %s is registered with different caskey (%s vs %s)"), fileName, CasKeyString(res2.first->second).str, CasKeyString(casKey).str);
-				return true;
+				DeferedCasCreation& dcc = res.first->second;
+				dcc.fileName = fileName;
+				dcc.fileSize = fileSize;
+				dcc.lastWriteTime = lastWriteTime;
 			}
-			else
-			{
-				auto findIt = m_deferredCasCreationLookupByName.find(fileNameKey);
-				if (findIt != m_deferredCasCreationLookupByName.end() && findIt->second != casKey)
-					return m_logger.Error(TC("Found deferred entry for %s with different cas key while writing cas. A ReportFileWrite is missing somewhere"), fileName);
+			auto res2 = m_deferredCasCreationLookupByName.try_emplace(fileNameKey, casKey);
+			if (res2.second)
+				res.first->second.names.push_back(res2.first);
+			else if (res2.first->second != casKey)
+				return m_logger.Error(TC("Same file %s is registered with different caskey (%s vs %s)"), fileName, CasKeyString(res2.first->second).str, CasKeyString(casKey).str);
+			return true;
+		}
+		else
+		{
+			#if 1 // UBA_DEBUG // This should never happen
+			SCOPED_FUTEX(m_deferredCasCreationLookupLock, deferredLock);
+			auto findIt = m_deferredCasCreationLookupByName.find(fileNameKey);
+			if (findIt != m_deferredCasCreationLookupByName.end() && findIt->second != casKey)
+				return m_logger.Error(TC("Found deferred entry for %s with different cas key while writing cas. A ReportFileWrite is missing somewhere"), fileName);
 
-				auto findIt2 = m_deferredCasCreationLookup.find(casKey);
-				if (findIt2 != m_deferredCasCreationLookup.end())
-				{
-					#if UBA_DEBUG
-					return m_logger.Error(TC("This should not happen. (%s)"), fileName);
-					#endif
-					for (auto& it2 : findIt2->second.names)
-						m_deferredCasCreationLookupByName.erase(it2);
-					m_deferredCasCreationLookup.erase(findIt2);
-				}
+			auto findIt2 = m_deferredCasCreationLookup.find(casKey);
+			if (findIt2 != m_deferredCasCreationLookup.end())
+			{
+				#if UBA_DEBUG
+				return m_logger.Error(TC("This should not happen. (%s)"), fileName);
+				#endif
+				for (auto& it2 : findIt2->second.names)
+					m_deferredCasCreationLookupByName.erase(it2);
+				m_deferredCasCreationLookup.erase(findIt2);
 			}
+			#endif
 		}
 
 		casEntry.verified = true;
@@ -583,7 +584,7 @@ namespace uba
 
 		casEntry.exists = true;
 		entryLock.Leave();
-		CasEntryWritten(casEntry, res.size);
+		CasEntryWritten(casEntry, res.storedSize);
 		return true;
 	}
 
@@ -1525,7 +1526,7 @@ namespace uba
 		ExternalFileMapping externalMapping;
 
 		u64 fileSize = 0;
-		u64 writtenSize = 0;
+		u64 storedSize = 0;
 		u64 lastWriteTime = 0;
 
 		FileHandle fileHandle = InvalidFileHandle;
@@ -1534,7 +1535,7 @@ namespace uba
 		if (m_externalFileMappingsProvider && m_externalFileMappingsProvider(externalMapping, fileNameKey, fileName))
 		{
 			fileSize = externalMapping.size;
-			writtenSize = externalMapping.writtenSize;
+			storedSize = externalMapping.storedSize;
 			lastWriteTime = externalMapping.lastWriteTime;
 		}
 		else
@@ -1554,7 +1555,7 @@ namespace uba
 				return m_logger.Error(TC("GetFileInformationByHandle failed on %s"), fileName);
 			}
 			fileSize = info.size;
-			writtenSize = fileSize;
+			storedSize = fileSize;
 			lastWriteTime = info.lastWriteTime;
 		}
 
@@ -1563,14 +1564,14 @@ namespace uba
 			if (casKeyOverride != CasKeyZero && casKeyOverride != fileEntry.casKey)
 			{
 				fileEntry.casKey = casKeyOverride;
-				if (!AddCasFile(fileNameKey, fileName, fileEntry.casKey, deferCreation, writtenSize, lastWriteTime))
+				if (!AddCasFile(fileNameKey, fileName, fileEntry.casKey, deferCreation, storedSize, lastWriteTime))
 					return false;
 				out = fileEntry.casKey;
 				return true;
 			}
 			if (fileSize == fileEntry.size && (lastWriteTime && lastWriteTime == fileEntry.lastWritten) && fileEntry.casKey != CasKeyInvalid)
 			{
-				if (!AddCasFile(fileNameKey, fileName, fileEntry.casKey, deferCreation, writtenSize, lastWriteTime))
+				if (!AddCasFile(fileNameKey, fileName, fileEntry.casKey, deferCreation, storedSize, lastWriteTime))
 					return false;
 				out = fileEntry.casKey;
 				return true;
@@ -1622,7 +1623,7 @@ namespace uba
 		if (fileEntry.casKey == CasKeyZero)
 			return false;
 
-		if (!AddCasFile(fileNameKey, fileName, fileEntry.casKey, deferCreation, writtenSize, lastWriteTime))
+		if (!AddCasFile(fileNameKey, fileName, fileEntry.casKey, deferCreation, storedSize, lastWriteTime))
 			return false;
 
 		out = fileEntry.casKey;
@@ -1804,15 +1805,15 @@ namespace uba
 		if (!WriteCasFile(res, deferredCreation.data, casKey))
 			return m_logger.Error(TC("Failed to write deferred cas %s (%s)"), CasKeyString(casKey).str, deferredCreation.data);
 
-		if (res.readSize != deferredFileSize) // TODO: Should also check last written time
-			m_logger.Warning(TC("Deferred cas %s was deferred with source file size %llu but was read with source file size %llu (%s)"), CasKeyString(casKey).str, deferredFileSize, res.readSize, deferredCreation.data);
+		if (res.sourceStoredSize != deferredFileSize) // TODO: Should also check last written time
+			m_logger.Warning(TC("Deferred cas %s was deferred with source file size %llu but was read with source file size %llu (%s)"), CasKeyString(casKey).str, deferredFileSize, res.sourceStoredSize, deferredCreation.data);
 
 		bool hasMapping = false;
 		if (res.mappingHandle.IsValid())
 		{
 			casEntry.mappingHandle = res.mappingHandle;
 			casEntry.mappingOffset = res.mappingOffset;
-			casEntry.mappingSize = res.size;
+			casEntry.mappingSize = res.storedSize;
 			casEntry.mappingIsIndependent = res.isIndependent;
 			hasMapping = true;
 		}
@@ -1821,7 +1822,7 @@ namespace uba
 		entryLock.Leave();
 
 		if (!hasMapping)
-			CasEntryWritten(casEntry, res.size);
+			CasEntryWritten(casEntry, res.storedSize);
 		return true;
 	}
 
@@ -1868,7 +1869,7 @@ namespace uba
 			return false;
 		casEntry.exists = true;
 		entryLock.Leave();
-		CasEntryWritten(casEntry, res.size);
+		CasEntryWritten(casEntry, res.storedSize);
 		return true;
 	}
 
diff --git a/Engine/Source/Programs/UnrealBuildAccelerator/Common/Public/UbaSession.h b/Engine/Source/Programs/UnrealBuildAccelerator/Common/Public/UbaSession.h
--- a/Engine/Source/Programs/UnrealBuildAccelerator/Common/Public/UbaSession.h
+++ b/Engine/Source/Programs/UnrealBuildAccelerator/Common/Public/UbaSession.h
@@ -276,9 +276,9 @@ namespace uba
 			Futex lock;
 			FileMappingHandle mapping;
 			u64 mappingOffset = 0;
-			u64 size = 0;
+			u64 contentSize = 0; // Size of uncompressed file or mapping
 			u64 lastWriteTime = 0;
-			u64 writtenSize = InvalidValue;
+			u64 storedSize = InvalidValue; // Size stored on disk. Can be compressed or uncompressed
 			u16 refCount = 0;
 			u8 usedCount = 0;
 			u8 usedCountBeforeFree = 0;
diff --git a/Engine/Source/Programs/UnrealBuildAccelerator/Common/Public/UbaStorage.h b/Engine/Source/Programs/UnrealBuildAccelerator/Common/Public/UbaStorage.h
--- a/Engine/Source/Programs/UnrealBuildAccelerator/Common/Public/UbaStorage.h
+++ b/Engine/Source/Programs/UnrealBuildAccelerator/Common/Public/UbaStorage.h
@@ -74,13 +74,21 @@ namespace uba
 		{
 			FileMappingHandle mappingHandle;
 			u64 mappingOffset = InvalidValue;
-			u64 size = InvalidValue;
-			u64 readSize = InvalidValue;
+			u64 storedSize = InvalidValue; // Stored size of the destination
+			u64 sourceStoredSize = InvalidValue; // Stored size of the source used to write compressed file
 			bool isIndependent = false;
 		};
 		virtual bool WriteCompressedFile(WriteResult& out, const tchar* from, FileHandle readHandle, u8* readMem, u64 fileSize, const tchar* toFile, const void* header, u64 headerSize, u64 lastWriteTime = 0) = 0;
 
-		struct ExternalFileMapping { FileMappingHandle handle; u64 offset = 0; u64 size = 0; u64 lastWriteTime = 0; u64 writtenSize = InvalidValue; bool createIndependentMapping = false; };
+		struct ExternalFileMapping
+		{
+			FileMappingHandle handle;
+			u64 offset = 0;
+			u64 size = 0;
+			u64 lastWriteTime = 0;
+			u64 storedSize = InvalidValue;
+			bool createIndependentMapping = false;
+		};
 		using ExternalFileMappingsProvider = Function<bool(ExternalFileMapping& out, StringKey fileNameKey, const tchar* fileName)>;
 		virtual void RegisterExternalFileMappingsProvider(ExternalFileMappingsProvider&& provider) {}
 	};
@@ -233,7 +241,7 @@ namespace uba
 			CasKey key;
 			CasEntry* prevAccessed = nullptr;
 			CasEntry* nextAccessed = nullptr;
-			u64 size = 0;
+			u64 size = 0; // Stored size of cas entry. Can be both compressed or uncompressed
 			u16 readCount = 0; // This is set while file is being read over network
 			bool verified = false; // This flag needs to be set for below flags to be reliable. if this is false below flags are assumptions
 			bool exists = false; // File exists on disk
