[UBA]
* Fixed bug where .manifest files used internally between lld-link.exe and mt.exe were not tagged as "output files" so they did not end up in shared memory
* Added so ApplicationRules::IsOutput takes in systemtemp path

diff --git a/Engine/Source/Programs/UnrealBuildAccelerator/Core/Private/UbaApplicationRules.cpp b/Engine/Source/Programs/UnrealBuildAccelerator/Core/Private/UbaApplicationRules.cpp
--- a/Engine/Source/Programs/UnrealBuildAccelerator/Core/Private/UbaApplicationRules.cpp
+++ b/Engine/Source/Programs/UnrealBuildAccelerator/Core/Private/UbaApplicationRules.cpp
@@ -82,7 +82,7 @@ namespace uba
 	class ApplicationRulesClExe : public ApplicationRulesVC
 	{
 	public:
-		virtual bool IsOutputFile(const StringView& fileName) const override
+		virtual bool IsOutputFile(StringView fileName, StringView systemTemp) const override
 		{
 			return fileName.EndsWith(TCV(".obj"))
 				|| fileName.EndsWith(TCV(".dep.json"))
@@ -142,7 +142,7 @@ namespace uba
 	{
 		using Super = ApplicationRulesVC;
 	public:
-		virtual bool IsOutputFile(const StringView& fileName) const override
+		virtual bool IsOutputFile(StringView fileName, StringView systemTemp) const override
 		{
 			return fileName.EndsWith(TCV(".lib"))
 				|| fileName.EndsWith(TCV(".exp"))
@@ -239,11 +239,6 @@ namespace uba
 	{
 		using Super = ApplicationRulesVcLink;
 
-		virtual bool KeepInMemory(StringView fileName, StringView systemTemp, bool isRunningRemote, bool isWrite) const override
-		{
-			return fileName.EndsWith(TCV(".manifest")) && fileName.Contains(systemTemp);
-		}
-
 		virtual u64 FileTypeMaxSize(const StringView& file, bool isSystemOrTempFile) const override
 		{
 			if (file.Contains(TCV(".pdb.tmp")))
@@ -251,12 +246,14 @@ namespace uba
 			return ApplicationRulesVcLink::FileTypeMaxSize(file, isSystemOrTempFile);
 		}
 
-		virtual bool IsOutputFile(const StringView& fileName) const override
+		virtual bool IsOutputFile(StringView fileName, StringView systemTemp) const override
 		{
 			return false//Super::IsOutputFile(fileName)
 				|| fileName.Contains(TCV(".exe.tmp"))
 				|| fileName.Contains(TCV(".dll.tmp"))
-				|| fileName.Contains(TCV(".pdb.tmp"));
+				|| fileName.Contains(TCV(".pdb.tmp"))
+				|| (fileName.EndsWith(TCV(".manifest")) && fileName.Contains(systemTemp)) // Used between lld-link.exe and mt.exe (test compiling AutoRTFMEngineTests)
+				;
 		}
 	};
 
@@ -271,7 +268,7 @@ namespace uba
 			return ApplicationRulesVcLink::FileTypeMaxSize(file, isSystemOrTempFile);
 		}
 
-		virtual bool IsOutputFile(const StringView& fileName) const override
+		virtual bool IsOutputFile(StringView fileName, StringView systemTemp) const override
 		{
 			return false//Super::IsOutputFile(fileName)
 				|| fileName.EndsWith(TCV(".exe"))
@@ -315,7 +312,7 @@ namespace uba
 			return true;
 		}
 
-		virtual bool IsOutputFile(const StringView& fileName) const override
+		virtual bool IsOutputFile(StringView fileName, StringView systemTemp) const override
 		{
 			return fileName.EndsWith(TCV(".c.d"))
 				|| fileName.EndsWith(TCV(".h.d"))
@@ -383,9 +380,9 @@ namespace uba
 			return DependencyCrawlerType_MsvcCompiler;
 		}
 
-		virtual bool IsOutputFile(const StringView& fileName) const override
+		virtual bool IsOutputFile(StringView fileName, StringView systemTemp) const override
 		{
-			return Super::IsOutputFile(fileName) || fileName.Contains(TCV("response-")); // file used by internally spawned clang-cl.exe. Note this is only generated when command line is very long
+			return Super::IsOutputFile(fileName, systemTemp) || fileName.Contains(TCV("response-")); // file used by internally spawned clang-cl.exe. Note this is only generated when command line is very long
 		}
 
 		virtual bool KeepInMemory(StringView fileName, StringView systemTemp, bool isRunningRemote, bool isWrite) const override
@@ -398,7 +395,7 @@ namespace uba
 	{
 		using Super = ApplicationRulesClang;
 
-		virtual bool IsOutputFile(const StringView& fileName) const override
+		virtual bool IsOutputFile(StringView fileName, StringView systemTemp) const override
 		{
 			return fileName.Contains(TCV(".tmp")); // both .so.tmp and .tmp123456
 		}
@@ -418,14 +415,14 @@ namespace uba
 	{
 		using Super = ApplicationRulesClang;
 
-		virtual bool IsOutputFile(const StringView& fileName) const override
+		virtual bool IsOutputFile(StringView fileName, StringView systemTemp) const override
 		{
 			return fileName.Contains(TCV(".temp-stream-"));
 		}
 
 		virtual u64 FileTypeMaxSize(const StringView& file, bool isSystemOrTempFile) const override
 		{
-			if (IsOutputFile(file))
+			if (file.Contains(TCV(".temp-stream-")))
 				return 14ull * 1024 * 1024 * 1024; // This is ridiculous (needed for asan targets)
 			return Super::FileTypeMaxSize(file, isSystemOrTempFile);
 		}
@@ -433,7 +430,7 @@ namespace uba
 
 	class ApplicationRulesDumpSymsExe : public ApplicationRulesClang
 	{
-		virtual bool IsOutputFile(const StringView& fileName) const override
+		virtual bool IsOutputFile(StringView fileName, StringView systemTemp) const override
 		{
 			return false; // fileName.EndsWith(TCV(".psym")); With psym as output file the BreakpadSymbolEncoder fails to output a .sym file
 		}
@@ -466,7 +463,7 @@ namespace uba
 	{
 		using Super = ApplicationRules;
 
-		virtual bool IsOutputFile(const StringView& fileName) const override
+		virtual bool IsOutputFile(StringView fileName, StringView systemTemp) const override
 		{
 			return fileName.Contains(TCV("thinlto-")); // Used by a clang based platform's link time optimization pass. Shared from lto process back to linker process
 		}
@@ -487,9 +484,9 @@ namespace uba
 	{
 		using Super = ApplicationRulesClangPlusPlusExe;
 
-		virtual bool IsOutputFile(const StringView& fileName) const override
+		virtual bool IsOutputFile(StringView fileName, StringView systemTemp) const override
 		{
-			return fileName.Contains(TCV(".self")) || Super::IsOutputFile(fileName);
+			return fileName.Contains(TCV(".self")) || Super::IsOutputFile(fileName, systemTemp);
 		}
 
 		virtual bool KeepTempOutputFile(const StringView& fileName) const override
@@ -508,7 +505,7 @@ namespace uba
 	{
 		using Super = ApplicationRules;
 
-		virtual bool IsOutputFile(const StringView& fileName) const override
+		virtual bool IsOutputFile(StringView fileName, StringView systemTemp) const override
 		{
 			return fileName.Contains(TCV(".self"));
 		}
@@ -529,7 +526,7 @@ namespace uba
 	{
 		using Super = ApplicationRules;
 
-		virtual bool IsOutputFile(const StringView& fileName) const override
+		virtual bool IsOutputFile(StringView fileName, StringView systemTemp) const override
 		{
 			return fileName.EndsWith(TCV(".split"));
 		}
@@ -575,7 +572,7 @@ namespace uba
 			return true;
 		}
 
-		virtual bool IsOutputFile(const StringView& fileName) const override
+		virtual bool IsOutputFile(StringView fileName, StringView systemTemp) const override
 		{
 			return fileName.Contains(TCV(".generated.dummy"))
 				|| fileName.EndsWith(TCV(".ispc.bc"))
@@ -611,7 +608,7 @@ namespace uba
 
 	class ApplicationRulesUBTDll : public ApplicationRules
 	{
-		virtual bool IsOutputFile(const StringView& fileName) const override
+		virtual bool IsOutputFile(StringView fileName, StringView systemTemp) const override
 		{
 			return false;
 			// TODO: These does not work when UnrealBuildTool creates these files multiple times in a row (building multiple targets)
@@ -624,7 +621,7 @@ namespace uba
 
 	class ApplicationRulesPVSStudio : public ApplicationRules
 	{
-		virtual bool IsOutputFile(const StringView& fileName) const override
+		virtual bool IsOutputFile(StringView fileName, StringView systemTemp) const override
 		{
 			return fileName.EndsWith(TCV(".PVS-Studio.log"))
 				|| fileName.EndsWith(TCV(".pvslog"))
@@ -659,7 +656,7 @@ namespace uba
 			return file.Contains(TCV(".uba."));
 		}
 
-		virtual bool IsOutputFile(const StringView& fileName) const override
+		virtual bool IsOutputFile(StringView fileName, StringView systemTemp) const override
 		{
 			return fileName.EndsWith(TCV(".tmp")) || fileName.EndsWith(TCV(".out"));
 		}
@@ -677,7 +674,7 @@ namespace uba
 
 	class ApplicationRulesUbaObjTool : public ApplicationRules
 	{
-		virtual bool IsOutputFile(const StringView& fileName) const override
+		virtual bool IsOutputFile(StringView fileName, StringView systemTemp) const override
 		{
 			return fileName.EndsWith(TCV(".obj"))
 				|| fileName.EndsWith(TCV(".exp"));
@@ -691,7 +688,7 @@ namespace uba
 
 	class ApplicationRulesUbaTestApp : public ApplicationRules
 	{
-		virtual bool IsOutputFile(const StringView& fileName) const override
+		virtual bool IsOutputFile(StringView fileName, StringView systemTemp) const override
 		{
 			return fileName.EndsWith(TCV(".out"));
 		}
diff --git a/Engine/Source/Programs/UnrealBuildAccelerator/Core/Public/UbaApplicationRules.h b/Engine/Source/Programs/UnrealBuildAccelerator/Core/Public/UbaApplicationRules.h
--- a/Engine/Source/Programs/UnrealBuildAccelerator/Core/Public/UbaApplicationRules.h
+++ b/Engine/Source/Programs/UnrealBuildAccelerator/Core/Public/UbaApplicationRules.h
@@ -64,7 +64,7 @@ namespace uba
 		}
 
 		// Outputfile means that it is kept in memory and then sent back to session process which can decide to write it to disk or send it over network
-		virtual bool IsOutputFile(const StringView& fileName) const
+		virtual bool IsOutputFile(StringView fileName, StringView systemTemp) const
 		{
 			return false;
 		}
diff --git a/Engine/Source/Programs/UnrealBuildAccelerator/Detours/Private/UbaDetoursShared.cpp b/Engine/Source/Programs/UnrealBuildAccelerator/Detours/Private/UbaDetoursShared.cpp
--- a/Engine/Source/Programs/UnrealBuildAccelerator/Detours/Private/UbaDetoursShared.cpp
+++ b/Engine/Source/Programs/UnrealBuildAccelerator/Detours/Private/UbaDetoursShared.cpp
@@ -561,7 +561,7 @@ namespace uba
 		{
 			// If file is output file we accept wrong size because size is not supposed to be used anyway.
 			// We don't want to trigger unnecessary download/decompress of file
-			if (!g_rules->IsOutputFile(fileName))
+			if (!g_rules->IsOutputFile(fileName, g_systemTemp))
 			{
 				StringBuffer<> temp;
 				u32 closeId;
diff --git a/Engine/Source/Programs/UnrealBuildAccelerator/Detours/Private/Windows/UbaDetoursFunctionsKernelBase.inl b/Engine/Source/Programs/UnrealBuildAccelerator/Detours/Private/Windows/UbaDetoursFunctionsKernelBase.inl
--- a/Engine/Source/Programs/UnrealBuildAccelerator/Detours/Private/Windows/UbaDetoursFunctionsKernelBase.inl
+++ b/Engine/Source/Programs/UnrealBuildAccelerator/Detours/Private/Windows/UbaDetoursFunctionsKernelBase.inl
@@ -1588,7 +1588,7 @@ bool Shared_MoveFile(const tchar* func, LPCWSTR lpExistingFileName, LPCWSTR lpNe
 		FileInfo& sourceInfo = it->second;
 		lock.Leave();
 
-		if (g_allowOutputFiles && g_rules->IsOutputFile(dest))
+		if (g_allowOutputFiles && g_rules->IsOutputFile(dest, g_systemTemp))
 		{
 			sourceInfo.deleted = true;
 			UBA_ASSERT(!sourceInfo.memoryFile->isLocalOnly);
diff --git a/Engine/Source/Programs/UnrealBuildAccelerator/Detours/Private/Windows/UbaDetoursFunctionsNtDll.inl b/Engine/Source/Programs/UnrealBuildAccelerator/Detours/Private/Windows/UbaDetoursFunctionsNtDll.inl
--- a/Engine/Source/Programs/UnrealBuildAccelerator/Detours/Private/Windows/UbaDetoursFunctionsNtDll.inl
+++ b/Engine/Source/Programs/UnrealBuildAccelerator/Detours/Private/Windows/UbaDetoursFunctionsNtDll.inl
@@ -858,7 +858,7 @@ NTSTATUS NTAPI Shared_NtCreateFile(bool IsCreateFunc, PHANDLE hFileHandle, ACCES
 	bool isThrowAway = g_rules->IsThrowAway(fileName, g_runningRemote);
 	bool keepInMemory = g_allowKeepFilesInMemory && g_rules->KeepInMemory(fileName, g_systemTemp, g_runningRemote, isWrite);
 	
-	keepInMemory = keepInMemory || ((isWrite || isDeleteOnClose) && g_rules->IsOutputFile(fileName) && g_allowKeepFilesInMemory) || (isWrite && isThrowAway);
+	keepInMemory = keepInMemory || ((isWrite || isDeleteOnClose) && g_rules->IsOutputFile(fileName, g_systemTemp) && g_allowKeepFilesInMemory) || (isWrite && isThrowAway);
 
 #if UBA_DEBUG_LOG_ENABLED
 	const wchar_t* isWriteStr = isWrite ? L" WRITE" : L""; (void)isWriteStr;
@@ -956,7 +956,7 @@ NTSTATUS NTAPI Shared_NtCreateFile(bool IsCreateFunc, PHANDLE hFileHandle, ACCES
 
 				// If file is an output file we still allllow this path and accept wrong (compressed) size.
 				// This is a bit hacky but we don't want to transfer and decompress file just to get the size
-				if (allowEarlyOut && !useContent && !isWriteAttributes && (!CouldBeCompressedFile(fileName) || g_rules->IsOutputFile(fileName)))
+				if (allowEarlyOut && !useContent && !isWriteAttributes && (!CouldBeCompressedFile(fileName) || g_rules->IsOutputFile(fileName, g_systemTemp)))
 				{
 					auto dh = new DetouredHandle(HandleType_File);
 					dh->fileObject = new FileObject();
@@ -1292,7 +1292,7 @@ NTSTATUS NTAPI Shared_NtCreateFile(bool IsCreateFunc, PHANDLE hFileHandle, ACCES
 			}
 			else
 			{
-				bool isOutput = (isWrite || isDeleteOnClose) && g_rules->IsOutputFile(fileName);
+				bool isOutput = (isWrite || isDeleteOnClose) && g_rules->IsOutputFile(fileName, g_systemTemp);
 				if (isOutput && (CreateDisposition & FILE_OPEN_IF) == 0)
 				{
 					UBA_ASSERTF(false, TC("Trying to open %s with openif. This is not supported"), fileName.data);
@@ -1540,7 +1540,7 @@ NTSTATUS NTAPI Detoured_NtClose(HANDLE handle)
 		mappingWritten = mf->writtenSize;
 
 		u32 orginalNameLen = TStrlen(fi.originalName);
-		if ((fo->deleteOnClose || IsWrite(fo->desiredAccess, 0)) && g_rules->IsOutputFile({fi.originalName, orginalNameLen}) && !g_rules->IsThrowAway(StringView(fi.originalName, orginalNameLen), g_runningRemote))
+		if ((fo->deleteOnClose || IsWrite(fo->desiredAccess, 0)) && g_rules->IsOutputFile({fi.originalName, orginalNameLen}, g_systemTemp) && !g_rules->IsThrowAway(StringView(fi.originalName, orginalNameLen), g_runningRemote))
 		{
 			// Need to report this file to host so it can be tracked in directory table
 			if (!mf->isReported)
