[UBA]
* Fixed bug in CompressWriter::CompressToFile where header size was not included in stored size

diff --git a/Engine/Source/Programs/UnrealBuildAccelerator/Common/Private/UbaStorageUtils.cpp b/Engine/Source/Programs/UnrealBuildAccelerator/Common/Private/UbaStorageUtils.cpp
--- a/Engine/Source/Programs/UnrealBuildAccelerator/Common/Private/UbaStorageUtils.cpp
+++ b/Engine/Source/Programs/UnrealBuildAccelerator/Common/Private/UbaStorageUtils.cpp
@@ -848,7 +848,7 @@ namespace uba
 				return m_logger.Error(TC("Failed to set file time on filehandle for %s"), toFile);
 		if (!destinationFile.Close())
 			return false;
-		outCompressedSize = destination.GetWritten();
+		outCompressedSize = destination.GetWritten() + headerSize;
 		return true;
 	}
 
