[UBA]
* Fixed race condition in CompressWriter::CompressFromMem where it could end up accessing deleted memory during shutdown

diff --git a/Engine/Source/Programs/UnrealBuildAccelerator/Common/Private/UbaStorageUtils.cpp b/Engine/Source/Programs/UnrealBuildAccelerator/Common/Private/UbaStorageUtils.cpp
--- a/Engine/Source/Programs/UnrealBuildAccelerator/Common/Private/UbaStorageUtils.cpp
+++ b/Engine/Source/Programs/UnrealBuildAccelerator/Common/Private/UbaStorageUtils.cpp
@@ -945,9 +945,11 @@ namespace uba
 			auto& rec = *recPtr;
 			KernelStatsScope kss(kernelStats);
 
-			u8* slot = rec.bufferSlots->Pop();
-			auto _ = MakeGuard([&](){ rec.bufferSlots->Push(slot); });
-			u8* compressSlotBuffer = slot + BufferSlotHalfSize;
+			BufferSlots* bufferSlots = nullptr;
+			u8* slot = nullptr;
+			u8* compressSlotBuffer = nullptr;
+			auto _ = MakeGuard([&](){ if (slot) bufferSlots->Push(slot); });
+
 			while (true)
 			{
 				u64 index = rec.compressCounter++;
@@ -957,6 +959,13 @@ namespace uba
 						delete recPtr;
 					return;
 				}
+				if (!slot)
+				{
+					bufferSlots = rec.bufferSlots;
+					slot = bufferSlots->Pop();
+					compressSlotBuffer = slot + BufferSlotHalfSize;
+				}
+
 				u64 startOffset = rec.maxUncompressedBlock*index;
 				const u8* uncompressedDataSlot = rec.uncompressedData + startOffset;
 				OO_SINTa uncompressedBlockSize = (OO_SINTa)Min(rec.maxUncompressedBlock, rec.fileSize - startOffset);
