[UBA]
* Cleaner fix for CompressWriter::CompressFromMem which also prevents a memory leak if there is a failure compressing segment

diff --git a/Engine/Source/Programs/UnrealBuildAccelerator/Common/Private/UbaStorageUtils.cpp b/Engine/Source/Programs/UnrealBuildAccelerator/Common/Private/UbaStorageUtils.cpp
--- a/Engine/Source/Programs/UnrealBuildAccelerator/Common/Private/UbaStorageUtils.cpp
+++ b/Engine/Source/Programs/UnrealBuildAccelerator/Common/Private/UbaStorageUtils.cpp
@@ -945,24 +945,26 @@ namespace uba
 			auto& rec = *recPtr;
 			KernelStatsScope kss(kernelStats);
 
-			BufferSlots* bufferSlots = nullptr;
 			u8* slot = nullptr;
 			u8* compressSlotBuffer = nullptr;
-			auto _ = MakeGuard([&](){ if (slot) bufferSlots->Push(slot); });
+
+			auto exitGuard = MakeGuard([&]()
+				{
+					if (slot)
+						rec.bufferSlots->Push(slot);
+					if (!--rec.refCount)
+						delete recPtr;
+				});
 
 			while (true)
 			{
 				u64 index = rec.compressCounter++;
 				if (index >= rec.workCount)
-				{
-					if (!--rec.refCount)
-						delete recPtr;
 					return;
-				}
+
 				if (!slot)
 				{
-					bufferSlots = rec.bufferSlots;
-					slot = bufferSlots->Pop();
+					slot = rec.bufferSlots->Pop();
 					compressSlotBuffer = slot + BufferSlotHalfSize;
 				}
 
