[UBA]
* Fixed so special ubacopy path is invalidating cached file infos. Without this fix we could have a different file on disk than the host provides to the helpers (helpers getting an old version of the file)
* Added detour logging in special ubacopy path (even though it is not detoured) just to make it easier to search in log files

diff --git a/Engine/Source/Programs/UnrealBuildAccelerator/Common/Private/UbaProcess.cpp b/Engine/Source/Programs/UnrealBuildAccelerator/Common/Private/UbaProcess.cpp
--- a/Engine/Source/Programs/UnrealBuildAccelerator/Common/Private/UbaProcess.cpp
+++ b/Engine/Source/Programs/UnrealBuildAccelerator/Common/Private/UbaProcess.cpp
@@ -1,11 +1,12 @@
 // Copyright Epic Games, Inc. All Rights Reserved.
 
 #include "UbaProcess.h"
+#include "UbaApplicationRules.h"
 #include "UbaFileAccessor.h"
 #include "UbaProtocol.h"
 #include "UbaProcessStats.h"
 #include "UbaProcessUtils.h"
-#include "UbaApplicationRules.h"
+#include "UbaStorage.h"
 
 #if PLATFORM_WINDOWS
 #include "UbaDetoursPayload.h"
@@ -673,6 +674,8 @@ namespace uba
 		if (!Equals(m_startInfo.application, TC("ubacopy")))
 			return false;
 		const tchar* fromFileBegin = m_startInfo.arguments;
+		if (*fromFileBegin == '\"')
+			++fromFileBegin;
 		const tchar* fromFileEnd = TStrchr(fromFileBegin, '\"');
 		UBA_ASSERT(fromFileEnd);
 		if (!fromFileEnd)
@@ -701,11 +704,41 @@ namespace uba
 		temp.Clear().Append(toFileBegin, toFileEnd - toFileBegin);
 		FixPath(temp.data, workDir.data, workDir.count, toName);
 
+		UBA_ASSERTF(m_session.m_shouldWriteToDisk, TC("NotImplemented - Copy code path does not support keeping files in memory. (%s -> %s)"), fromName.data, toName.data);
+
 		DWORD oldAttributes = GetFileAttributes(toName.data);
 		if (oldAttributes != INVALID_FILE_ATTRIBUTES && (oldAttributes & FILE_ATTRIBUTE_READONLY))
 			SetFileAttributes(toName.data, oldAttributes & (~FILE_ATTRIBUTE_READONLY));
 
-		if (!CopyFileW(fromName.data, toName.data, false))
+		StringKey toKey = ToStringKeyLower(toName);
+
+		#if UBA_DEBUG_LOG_ENABLED
+		FileAccessor logFile(m_session.m_logger, m_startInfo.logFile);
+		bool logSuccess = *m_startInfo.logFile && logFile.CreateWrite();
+		auto writeLogLine = [&](const StringBufferBase& line)
+			{
+				if (!logSuccess)
+					return;
+				char buf[2*1024];
+				if (u32 written = line.Parse(buf, sizeof_array(buf)))
+					logFile.Write(buf, written - 1);
+				logFile.Write("\n", 1);
+			};
+		writeLogLine(StringBuffer<>().Appendf(TC("ProcessId: %u"), m_id));
+		writeLogLine(StringBuffer<>().Appendf(TC("CmdLine: %s %s"), m_startInfo.application, m_startInfo.arguments));
+		auto logGuard = MakeGuard([&]() { if (logSuccess) logFile.Close(); });
+		#endif
+
+		// We need to report before writing to flush out potential deferred cas creations
+		m_session.m_storage.ReportFileWrite(toKey, toName.data);
+
+		bool copySuccess = CopyFileW(fromName.data, toName.data, false);
+
+		#if UBA_DEBUG_LOG_ENABLED
+		writeLogLine(StringBuffer<>().Appendf(TC("T      CopyFileW %s to %s -> %s"), fromName.data, toName.data, (copySuccess?TC("true"):TC("true"))));
+		#endif
+
+		if (!copySuccess)
 		{
 			m_exitCode = GetLastError();
 			temp.Clear().Appendf(TC("Failed to copy %s to %s (%s)"), fromName.data, toName.data, LastErrorToText(m_exitCode).data);
@@ -713,9 +746,14 @@ namespace uba
 			return true;
 		}
 		
-		SetFileAttributes(toName.data, DefaultAttributes());
+		bool attributesSuccess = SetFileAttributes(toName.data, DefaultAttributes());(void)attributesSuccess;
+
+		#if UBA_DEBUG_LOG_ENABLED
+		writeLogLine(StringBuffer<>().Appendf(TC("T      SetFileAttributes %s -> %s"), toName.data, (attributesSuccess?TC("true"):TC("true"))));
+		#endif
+
+		m_session.m_storage.InvalidateCachedFileInfo(toKey);
 
-		StringKey toKey = ToStringKeyLower(toName);
 		m_session.RegisterCreateFileForWrite(toKey, toName, true);
 
 		//m_writtenFiles.try_emplace(name);
diff --git a/Engine/Source/Programs/UnrealBuildAccelerator/Common/Private/UbaProcessStartInfoHolder.cpp b/Engine/Source/Programs/UnrealBuildAccelerator/Common/Private/UbaProcessStartInfoHolder.cpp
--- a/Engine/Source/Programs/UnrealBuildAccelerator/Common/Private/UbaProcessStartInfoHolder.cpp
+++ b/Engine/Source/Programs/UnrealBuildAccelerator/Common/Private/UbaProcessStartInfoHolder.cpp
@@ -116,8 +116,8 @@ namespace uba
 			if (!commands.StartsWith(TC("copy /Y \"")))
 				return false;
 
-			const tchar* fromFileBegin = commands.data + 9;
-			const tchar* fromFileEnd = TStrchr(fromFileBegin, '\"');
+			const tchar* fromFileBegin = commands.data + 8;
+			const tchar* fromFileEnd = TStrchr(fromFileBegin+1, '\"');
 			if (!fromFileEnd)
 				return false;
 			const tchar* toFileBegin = TStrchr(fromFileEnd + 1, '\"');
diff --git a/Engine/Source/Programs/UnrealBuildAccelerator/Common/Private/UbaSessionServer.cpp b/Engine/Source/Programs/UnrealBuildAccelerator/Common/Private/UbaSessionServer.cpp
--- a/Engine/Source/Programs/UnrealBuildAccelerator/Common/Private/UbaSessionServer.cpp
+++ b/Engine/Source/Programs/UnrealBuildAccelerator/Common/Private/UbaSessionServer.cpp
@@ -1162,7 +1162,7 @@ namespace uba
 			process.m_writtenFiles.try_emplace(destination.ToString(), casKey);
 		}
 
-		// We need to report before writing to flush out potential 
+		// We need to report before writing to flush out potential deferred cas creations
 		m_storage.ReportFileWrite(destinationKey, destination.data);
 
 		if (shouldWriteToDisk)
