[UBA]
* Fixed bug where old memory file was not deleted after child process had overwritten it (fixes binary produced when building AutoRTFMEngineTests

diff --git a/Engine/Source/Programs/UnrealBuildAccelerator/Detours/Private/UbaDetoursFileMappingTable.cpp b/Engine/Source/Programs/UnrealBuildAccelerator/Detours/Private/UbaDetoursFileMappingTable.cpp
--- a/Engine/Source/Programs/UnrealBuildAccelerator/Detours/Private/UbaDetoursFileMappingTable.cpp
+++ b/Engine/Source/Programs/UnrealBuildAccelerator/Detours/Private/UbaDetoursFileMappingTable.cpp
@@ -195,6 +195,19 @@ namespace uba
 			if (!info.name || !backedName.Equals(info.name))
 				info.name = g_mappedFileTable.m_memoryBlock.Strdup(backedName).data;
 
+			if (auto mf = info.memoryFile) // Memory file is now wrong.. we need to delete it (or should we create a new one with the new mapping?)
+			{
+				if (!mf->isLocalOnly)
+				{
+					DEBUG_LOG(TC("Unmapping old memoryfile for: %s"), info.originalName);
+					UnmapViewOfFile(mf->baseAddress);
+					CloseHandle(mf->mappingHandle.mh);
+					CloseHandle(mf->mappingHandle.fh);
+				}
+				// delete mf; // Let them leak
+				info.memoryFile = nullptr;
+			}
+
 			DEBUG_LOG(TC("GOT WRITTEN FILE: %s (BackedFile: %s Size: %llu)"), info.originalName, info.name, info.size);
 		}
 	}
