[UBA]
* Fixed posix compile error

diff --git a/Engine/Source/Programs/UnrealBuildAccelerator/Detours/Private/UbaDetoursFileMappingTable.cpp b/Engine/Source/Programs/UnrealBuildAccelerator/Detours/Private/UbaDetoursFileMappingTable.cpp
--- a/Engine/Source/Programs/UnrealBuildAccelerator/Detours/Private/UbaDetoursFileMappingTable.cpp
+++ b/Engine/Source/Programs/UnrealBuildAccelerator/Detours/Private/UbaDetoursFileMappingTable.cpp
@@ -197,6 +197,7 @@ namespace uba
 
 			if (auto mf = info.memoryFile) // Memory file is now wrong.. we need to delete it (or should we create a new one with the new mapping?)
 			{
+				#if PLATFORM_WINDOWS
 				if (!mf->isLocalOnly)
 				{
 					DEBUG_LOG(TC("Unmapping old memoryfile for: %s"), info.originalName);
@@ -204,6 +205,7 @@ namespace uba
 					CloseHandle(mf->mappingHandle.mh);
 					CloseHandle(mf->mappingHandle.fh);
 				}
+				#endif
 				// delete mf; // Let them leak
 				info.memoryFile = nullptr;
 			}
