[UBA]
* Fixed bug where process exited function was never called if process was never started (binary was not found)

diff --git a/Engine/Source/Programs/UnrealBuildAccelerator/Common/Private/UbaProcess.cpp b/Engine/Source/Programs/UnrealBuildAccelerator/Common/Private/UbaProcess.cpp
--- a/Engine/Source/Programs/UnrealBuildAccelerator/Common/Private/UbaProcess.cpp
+++ b/Engine/Source/Programs/UnrealBuildAccelerator/Common/Private/UbaProcess.cpp
@@ -219,6 +219,7 @@ namespace uba
 		if (!m_session.PrepareProcess(*this, m_parentProcess != nullptr, realApplication, realWorkingDir))
 		{
 			m_exitCode = 44324;
+			CallExitedFunc();
 			return false;
 		}
 
@@ -642,18 +643,7 @@ namespace uba
 			m_comMemory = {};
 		}
 
-		if (m_startInfo.exitedFunc)
-		{
-			auto exitedFunc = m_startInfo.exitedFunc;
-			auto userData = m_startInfo.userData;
-			m_startInfo.exitedFunc = nullptr;
-			m_startInfo.userData = nullptr;
-			ProcessHandle h;
-			h.m_process = this;
-			ProcessExitedResponse exitedResponse = ProcessExitedResponse_None;
-			exitedFunc(userData, h, exitedResponse);
-			h.m_process = nullptr;
-		}
+		CallExitedFunc();
 
 		UBA_ASSERT(m_refCount);
 
@@ -668,6 +658,21 @@ namespace uba
 		#endif
 	}
 
+	void ProcessImpl::CallExitedFunc()
+	{
+		if (!m_startInfo.exitedFunc)
+			return;
+		auto exitedFunc = m_startInfo.exitedFunc;
+		auto userData = m_startInfo.userData;
+		m_startInfo.exitedFunc = nullptr;
+		m_startInfo.userData = nullptr;
+		ProcessHandle h;
+		h.m_process = this;
+		ProcessExitedResponse exitedResponse = ProcessExitedResponse_None;
+		exitedFunc(userData, h, exitedResponse);
+		h.m_process = nullptr;
+	}
+
 	bool ProcessImpl::HandleSpecialApplication()
 	{
 	#if PLATFORM_WINDOWS
diff --git a/Engine/Source/Programs/UnrealBuildAccelerator/Common/Public/UbaProcess.h b/Engine/Source/Programs/UnrealBuildAccelerator/Common/Public/UbaProcess.h
--- a/Engine/Source/Programs/UnrealBuildAccelerator/Common/Public/UbaProcess.h
+++ b/Engine/Source/Programs/UnrealBuildAccelerator/Common/Public/UbaProcess.h
@@ -101,6 +101,7 @@ namespace uba
 		void WaitForParent();
 		void WaitForChildrenExit();
 		bool CancelWithError(bool dummy = true);
+		void CallExitedFunc();
 
 		ProcessStartInfoHolder m_startInfo;
 		Session& m_session;
