[UBA]
* Fixed failing linux unit test because cloud machine does not have "zip" installed.. switched to "cat" instead for system application test
* Adde special copy unit test for windows

diff --git a/Engine/Source/Programs/UnrealBuildAccelerator/Test/Private/UbaTest.cpp b/Engine/Source/Programs/UnrealBuildAccelerator/Test/Private/UbaTest.cpp
--- a/Engine/Source/Programs/UnrealBuildAccelerator/Test/Private/UbaTest.cpp
+++ b/Engine/Source/Programs/UnrealBuildAccelerator/Test/Private/UbaTest.cpp
@@ -5,7 +5,7 @@
 
 namespace uba
 {
-	bool CreateTestFile(LoggerWithWriter& logger, StringView testRootDir, StringView fileName, StringView content)
+	bool CreateTestFile(Logger& logger, StringView testRootDir, StringView fileName, StringView content)
 	{
 		StringBuffer<> testFileName;
 		testFileName.Append(testRootDir).EnsureEndsWithSlash().Append(fileName).FixPathSeparators();
@@ -24,13 +24,20 @@ namespace uba
 		return file.Close();
 	}
 
-	bool DeleteTestFile(LoggerWithWriter& logger, StringView testRootDir, StringView fileName)
+	bool DeleteTestFile(Logger& logger, StringView testRootDir, StringView fileName)
 	{
 		StringBuffer<> testFileName;
 		testFileName.Append(testRootDir).EnsureEndsWithSlash().Append(fileName).FixPathSeparators();
 		return DeleteFileW(testFileName.data);
 	}
 
+	bool FileExists(Logger& logger, StringView dir, StringView fileName)
+	{
+		StringBuffer<> testFileName;
+		testFileName.Append(dir).EnsureEndsWithSlash().Append(fileName).FixPathSeparators();
+		return FileExists(logger, testFileName.data);
+	}
+
 	bool WrappedMain(int argc, tchar* argv[])
 	{
 		AddExceptionHandler();
diff --git a/Engine/Source/Programs/UnrealBuildAccelerator/Test/Private/UbaTest.h b/Engine/Source/Programs/UnrealBuildAccelerator/Test/Private/UbaTest.h
--- a/Engine/Source/Programs/UnrealBuildAccelerator/Test/Private/UbaTest.h
+++ b/Engine/Source/Programs/UnrealBuildAccelerator/Test/Private/UbaTest.h
@@ -16,8 +16,9 @@ namespace uba
 	class SessionServer;
 	struct ProcessStartInfo;
 
-	bool CreateTestFile(LoggerWithWriter& logger, StringView testRootDir, StringView fileName, StringView content);
-	bool DeleteTestFile(LoggerWithWriter& logger, StringView testRootDir, StringView fileName);
+	bool CreateTestFile(Logger& logger, StringView testRootDir, StringView fileName, StringView content);
+	bool DeleteTestFile(Logger& logger, StringView testRootDir, StringView fileName);
+	bool FileExists(Logger& logger, StringView dir, StringView fileName);
 
 	using RunProcessFunction = Function<ProcessHandle(const ProcessStartInfo&)>;
 	using TestSessionFunction = Function<bool(LoggerWithWriter& logger, SessionServer& session, const tchar* workingDir, const RunProcessFunction& runProcess)>;
diff --git a/Engine/Source/Programs/UnrealBuildAccelerator/Test/Private/UbaTestAll.h b/Engine/Source/Programs/UnrealBuildAccelerator/Test/Private/UbaTestAll.h
--- a/Engine/Source/Programs/UnrealBuildAccelerator/Test/Private/UbaTestAll.h
+++ b/Engine/Source/Programs/UnrealBuildAccelerator/Test/Private/UbaTestAll.h
@@ -77,6 +77,7 @@ namespace uba
 		UBA_TEST(TestVirtualFile) \
 		UBA_TEST(TestRemoteVirtualFile) \
 		UBA_TEST(TestRemoteProcessSpecialCase1) \
+		UBA_TEST(TestSessionSpecialCopy) \
 
 
 #define UBA_LINUX_TESTS \
diff --git a/Engine/Source/Programs/UnrealBuildAccelerator/Test/Private/UbaTestSession.cpp b/Engine/Source/Programs/UnrealBuildAccelerator/Test/Private/UbaTestSession.cpp
--- a/Engine/Source/Programs/UnrealBuildAccelerator/Test/Private/UbaTestSession.cpp
+++ b/Engine/Source/Programs/UnrealBuildAccelerator/Test/Private/UbaTestSession.cpp
@@ -524,7 +524,7 @@ namespace uba
 		#if PLATFORM_WINDOWS
 		return TC("c:\\windows\\system32\\ping.exe");
 		#elif PLATFORM_LINUX
-		return TC("/usr/bin/zip");
+		return TC("/usr/bin/cat");
 		#else
 		return TC("/sbin/zip");
 		#endif
@@ -534,6 +534,8 @@ namespace uba
 	{
 		#if PLATFORM_WINDOWS
 		return TC("-n 1 localhost");
+		#elif PLATFORM_LINUX
+		return TC("--help");
 		#else
 		return TC("-help");
 		#endif
@@ -543,6 +545,8 @@ namespace uba
 	{
 		#if PLATFORM_WINDOWS
 		return TC("Pinging ");
+		#elif PLATFORM_LINUX
+		return TC("cat [OPTION]");
 		#else
 		return TC("zip [-options]");
 		#endif
@@ -872,4 +876,33 @@ namespace uba
 				return true;
 			});
 	}
+
+	bool TestSessionSpecialCopy(LoggerWithWriter& logger, const StringBufferBase& testRootDir)
+	{
+		Config config;
+		return SetupServerSession(logger, testRootDir, true, true, config,
+			[&](LoggerWithWriter& logger, const StringView& workingDir, SessionServer& sessionServer)
+			{
+				if (!CreateTestFile(logger, workingDir, TCV("File.h"), TCV("0")))
+					return false;
+
+				ProcessStartInfo processInfo;
+				processInfo.application = TC("cmd.exe");
+				processInfo.workingDir = workingDir.data;
+				processInfo.arguments = TC("/c copy /Y \"File.h\" \"File2.h\"");
+				ProcessHandle process = sessionServer.RunProcess(processInfo);
+				if (!process.IsValid())
+					return logger.Error(TC("Failed to start process"));
+				if (!process.WaitForExit(GetTimeoutTime()))
+					return logger.Error(TC("UbaTestApp did not exit in 10 seconds"));
+				if (!Equals(process.GetStartInfo().application, TC("ubacopy")))
+					return logger.Error(TC("Special copy was not used"));
+				u32 exitCode = process.GetExitCode();
+				if (exitCode != 0)
+					return logger.Error(TC("Special copy failed"));
+				if (!FileExists(logger, workingDir, TCV("File2.h")))
+					return false;
+				return true;
+			});
+	}
 }
diff --git a/Engine/Source/Programs/UnrealBuildAccelerator/TestApp/Private/Windows/UbaTestAppWin.cpp b/Engine/Source/Programs/UnrealBuildAccelerator/TestApp/Private/Windows/UbaTestAppWin.cpp
--- a/Engine/Source/Programs/UnrealBuildAccelerator/TestApp/Private/Windows/UbaTestAppWin.cpp
+++ b/Engine/Source/Programs/UnrealBuildAccelerator/TestApp/Private/Windows/UbaTestAppWin.cpp
@@ -344,9 +344,31 @@ int wmain(int argc, wchar_t* argv[])
 			return LogError(L"Failed to write 3 bytes to file %s", file);
 		CloseHandle(wh);
 	}
+	else if (wcsncmp(argv[1], L"-read=", 6) == 0)
+	{
+		int callIndex = _wtoi(argv[1] + 6);
+
+		wchar_t buf[32];
+		int read = ReadTestFile(buf, 32*2, L"SpecialFile1");
+		if (read == -1)
+			return -1;
+		buf[read*2] = 0;
+
+		int readIndex = _wtoi(buf);
+		if (callIndex != readIndex)
+			return LogError(L"callIndex different from readIndex");
+		return 0;
+	}
+	else if (wcsncmp(argv[1], L"-write=", 7) == 0)
+	{
+		int callIndex = _wtoi(argv[1] + 7);
+		wchar_t buf[32];
+		buf[0] = wchar_t('1' + callIndex);
+		return WriteTestFile(buf, 2, L"SpecialFile1");
+	}
 	else if (wcsncmp(argv[1], L"-readwrite=", 11) == 0)
 	{
-		int callIndex = _wtoi(argv[1] + 11);(void)callIndex;
+		int callIndex = _wtoi(argv[1] + 11);
 
 		wchar_t buf[32];
 		int read = ReadTestFile(buf, 32*2, L"SpecialFile1");
