[UBA]
* Reduced size of GetWritten message by using a bool to say if file is in temp (which it likely is) and if so, just write name of file instead of full path
* Reduced size of GetWritten message by changing mapping names to use base62 instead of hex
* Reduced size of GetWritten message by changing size to use 7bit encoding instead of u64
* Swapped file and mapping handle when writing to u64 on win64.. since we most of the time only have the mapping in there the value gets smaller and compress better
* Eventually gave up and added support for multi-message for GetWritten files :-/
* Fixed so WaitForMultipleObjects is not updating tables unless state has changed of a detoured process (which was the intention all along)

diff --git a/Engine/Source/Programs/UnrealBuildAccelerator/Common/Private/UbaProcess.cpp b/Engine/Source/Programs/UnrealBuildAccelerator/Common/Private/UbaProcess.cpp
--- a/Engine/Source/Programs/UnrealBuildAccelerator/Common/Private/UbaProcess.cpp
+++ b/Engine/Source/Programs/UnrealBuildAccelerator/Common/Private/UbaProcess.cpp
@@ -1003,14 +1003,38 @@ namespace uba
 
 	bool ProcessImpl::HandleUpdateTables(BinaryReader& reader, BinaryWriter& writer)
 	{
+		UBA_ASSERT(!m_shared.firstOverflowMessage);
 		writer.WriteU32(m_session.GetDirectoryTableSize());
 		writer.WriteU32(m_session.GetFileMappingSize());
 		return HandleGetWrittenFiles(reader, writer);
 	}
 
+	struct ProcessImpl::OverflowedWrittenFilesMessage
+	{
+		u8 buffer[CommunicationMemSize];
+		u64 written = 0;
+		OverflowedWrittenFilesMessage* next = nullptr;
+	};
+
 	bool ProcessImpl::HandleGetWrittenFiles(BinaryReader& reader, BinaryWriter& writer)
 	{
-		u32& writtenCount = *(u32*)writer.AllocWrite(4);
+		if (auto msg = m_shared.firstOverflowMessage)
+		{
+			writer.WriteBytes(msg->buffer, msg->written);
+			m_shared.firstOverflowMessage = msg->next;
+			delete msg;
+			return true;
+		}
+
+		auto writtenCount = (u32*)writer.AllocWrite(4);
+		auto overflowed = (u8*)writer.AllocWrite(1);
+		*overflowed = 0;
+
+		BinaryWriter* w = &writer;
+		
+		BinaryWriter tempWriter(0, 0, 0);
+		OverflowedWrittenFilesMessage* activeOverflowMessage = nullptr;
+
 		u32 count = 0;
 		SCOPED_FUTEX_READ(m_shared.writtenFilesLock, lock);
 		for (auto& kv : m_shared.writtenFiles)
@@ -1027,14 +1051,54 @@ namespace uba
 			else if (wf.name == wf.backedName)
 				continue;
 
-			writer.WriteStringKey(kv.first);
-			writer.WriteString(wf.name);
-			writer.WriteString(wf.backedName);
-			writer.WriteU64(wf.mappingHandle.ToU64());
-			writer.WriteU64(fileSize);
+			if (w->GetCapacityLeft() < 1024)
+			{
+				// Close previous
+				*writtenCount = count;
+				*overflowed = 1;
+				count = 0;
+
+				// Create new message
+				auto newMsg = new OverflowedWrittenFilesMessage();
+
+				// Update active message if there is one
+				if (auto oldMsg = activeOverflowMessage)
+				{
+					oldMsg->written = w->GetPosition();
+					oldMsg->next = newMsg;
+				}
+				else
+				{
+					w = &tempWriter;
+					m_shared.firstOverflowMessage = newMsg;
+				}
+				activeOverflowMessage = newMsg;
+
+				// Setup writer and variables
+				tempWriter.Reset(newMsg->buffer, 0, sizeof_array(OverflowedWrittenFilesMessage::buffer));
+				writtenCount = (u32*)w->AllocWrite(4);
+				overflowed = (u8*)w->AllocWrite(1);
+				*overflowed = 0;
+			}
+
+			// This is an attempt to reduce size because we hit message max size (65536). Will need to split up in multiple messages eventually
+			StringView fileName(wf.name);
+			bool isInTemp = fileName.StartsWith(m_session.m_tempPath);
+			if (isInTemp)
+				fileName = fileName.Skip(m_session.m_tempPath.count);
+			w->WriteStringKey(kv.first);
+			w->WriteBool(isInTemp);
+			w->WriteString(fileName);
+			w->WriteString(wf.backedName);
+			w->Write7BitEncoded(wf.mappingHandle.ToU64());
+			w->Write7BitEncoded(fileSize);
 			++count;
 		}
-		writtenCount = count;
+
+		if (activeOverflowMessage)
+			activeOverflowMessage->written = w->GetPosition();
+
+		*writtenCount = count;
 		return true;
 	}
 
diff --git a/Engine/Source/Programs/UnrealBuildAccelerator/Common/Private/UbaStorage.cpp b/Engine/Source/Programs/UnrealBuildAccelerator/Common/Private/UbaStorage.cpp
--- a/Engine/Source/Programs/UnrealBuildAccelerator/Common/Private/UbaStorage.cpp
+++ b/Engine/Source/Programs/UnrealBuildAccelerator/Common/Private/UbaStorage.cpp
@@ -22,7 +22,9 @@ namespace uba
 
 	void Storage::GetMappingString(StringBufferBase& out, FileMappingHandle mappingHandle, u64 offset)
 	{
-		out.Append('^').AppendHex(mappingHandle.ToU64()).Append('-').AppendHex(offset);
+		out.Append('^').AppendBase62(mappingHandle.ToU64());
+		if (offset)
+			out.Append('-').AppendBase62(offset);
 	}
 
 	u8* BufferSlots::Pop()
diff --git a/Engine/Source/Programs/UnrealBuildAccelerator/Common/Public/UbaProcess.h b/Engine/Source/Programs/UnrealBuildAccelerator/Common/Public/UbaProcess.h
--- a/Engine/Source/Programs/UnrealBuildAccelerator/Common/Public/UbaProcess.h
+++ b/Engine/Source/Programs/UnrealBuildAccelerator/Common/Public/UbaProcess.h
@@ -133,10 +133,13 @@ namespace uba
 		HANDLE m_accountingJobObject = NULL;
 		#endif
 
+		struct OverflowedWrittenFilesMessage;
+
 		struct Shared
 		{
 			Futex writtenFilesLock;
 			UnorderedMap<StringKey, WrittenFile> writtenFiles;
+			OverflowedWrittenFilesMessage* firstOverflowMessage = nullptr;
 		};
 
 		u32 m_nativeProcessId = 0;
diff --git a/Engine/Source/Programs/UnrealBuildAccelerator/Core/Private/UbaStringBuffer.cpp b/Engine/Source/Programs/UnrealBuildAccelerator/Core/Private/UbaStringBuffer.cpp
--- a/Engine/Source/Programs/UnrealBuildAccelerator/Core/Private/UbaStringBuffer.cpp
+++ b/Engine/Source/Programs/UnrealBuildAccelerator/Core/Private/UbaStringBuffer.cpp
@@ -216,6 +216,20 @@ namespace uba
 		return Append(buf);
 	}
 
+	StringBufferBase& StringBufferBase::AppendBase62(u64 v)
+	{
+		tchar temp[16];
+		tchar* it = temp + 16;
+		if (v == 0)
+			*--it = '0';
+		while (v > 0)
+		{
+			*--it = g_base64Chars[v % 62];
+			v /= 62;
+		}
+		return Append(it, 16 - u32(it - temp));
+	}
+
 	StringBufferBase& StringBufferBase::AppendValue(u64 v)
 	{
 		tchar buf[256];
diff --git a/Engine/Source/Programs/UnrealBuildAccelerator/Core/Public/UbaBinaryReaderWriter.h b/Engine/Source/Programs/UnrealBuildAccelerator/Core/Public/UbaBinaryReaderWriter.h
--- a/Engine/Source/Programs/UnrealBuildAccelerator/Core/Public/UbaBinaryReaderWriter.h
+++ b/Engine/Source/Programs/UnrealBuildAccelerator/Core/Public/UbaBinaryReaderWriter.h
@@ -42,7 +42,8 @@ namespace uba
 		inline u64 GetCapacityLeft() const { return u64(m_end - m_pos); }
 		inline u8* GetData() const { return m_begin; }
 		inline void ChangeData(u8* newData, u64 newCapacity = InvalidValue) { u64 offset = u64(m_pos - m_begin); m_begin = newData; m_pos = newData + offset; m_end = newData + newCapacity; }
-		inline BinaryWriter(u8* data, u64 offset = 0, u64 capacity = InvalidValue) { m_begin = data; m_pos = data + offset; m_end = data + capacity; }
+		inline void Reset(u8* data, u64 offset = 0, u64 capacity = InvalidValue) { m_begin = data; m_pos = data + offset; m_end = data + capacity; }
+		inline BinaryWriter(u8* data, u64 offset = 0, u64 capacity = InvalidValue) { Reset(data, offset, capacity); }
 
 		void Flush(bool waitOnResponse = true);
 		BinaryWriter();
diff --git a/Engine/Source/Programs/UnrealBuildAccelerator/Core/Public/UbaMemory.h b/Engine/Source/Programs/UnrealBuildAccelerator/Core/Public/UbaMemory.h
--- a/Engine/Source/Programs/UnrealBuildAccelerator/Core/Public/UbaMemory.h
+++ b/Engine/Source/Programs/UnrealBuildAccelerator/Core/Public/UbaMemory.h
@@ -276,6 +276,27 @@ namespace uba
 		return v;
 	}
 
+	constexpr tchar g_base64Chars[] = TC("0123456789ABCDEFGHIJKLMNOPQRSTUVWXYZabcdefghijklmnopqrstuvwxyz");
+
+	inline u64 StringToValueBase62(const tchar* str, u64 len)
+	{
+		u64 result = 0;
+		const tchar* it = str;
+		const tchar* end = it + len;
+		while (it != end)
+		{
+			tchar c = *it++;
+			result *= 62;
+			if ('0' <= c && c <= '9')
+				result += c - '0';
+			else if ('A' <= c && c <= 'Z')
+				result += c - 'A' + 10;
+			else // if ('a' <= c && c <= 'z')
+				result += c - 'a' + 36;
+		}
+		return result;
+	}
+
 	inline u32 RoundUpPow2(u64 value)
 	{
 		u32 v = u32(value);
diff --git a/Engine/Source/Programs/UnrealBuildAccelerator/Core/Public/UbaStringBuffer.h b/Engine/Source/Programs/UnrealBuildAccelerator/Core/Public/UbaStringBuffer.h
--- a/Engine/Source/Programs/UnrealBuildAccelerator/Core/Public/UbaStringBuffer.h
+++ b/Engine/Source/Programs/UnrealBuildAccelerator/Core/Public/UbaStringBuffer.h
@@ -35,6 +35,7 @@ namespace uba
 		StringBufferBase& AppendDir(const tchar* dir);
 		StringBufferBase& AppendFileName(const tchar* str);
 		StringBufferBase& AppendHex(u64 v);
+		StringBufferBase& AppendBase62(u64 v);
 		StringBufferBase& AppendValue(u64 v);
 		StringBufferBase& Append(const tchar* format, va_list& args);
 		StringBufferBase& Prepend(const StringView& view, u32 overwriteCount = 0);
diff --git a/Engine/Source/Programs/UnrealBuildAccelerator/Detours/Private/UbaDetoursFileMappingTable.cpp b/Engine/Source/Programs/UnrealBuildAccelerator/Detours/Private/UbaDetoursFileMappingTable.cpp
--- a/Engine/Source/Programs/UnrealBuildAccelerator/Detours/Private/UbaDetoursFileMappingTable.cpp
+++ b/Engine/Source/Programs/UnrealBuildAccelerator/Detours/Private/UbaDetoursFileMappingTable.cpp
@@ -40,7 +40,7 @@ namespace uba
 			FileInfo& info = insres.first->second;
 			if (!insres.second)
 			{
-				if (info.name && info.name[0] == '^' && !mappedFileName.Equals(info.name)) // Mapped file has been re-mapped.
+				if (info.name && info.name[0] == '^' && !mappedFileName.Equals(info.name, false)) // Mapped file has been re-mapped.
 				{
 					UBA_ASSERTF(!info.memoryFile, TC("Mapped file %s has changed mapping (%s to %s) while being in use"), info.originalName, info.name, mappedFileName.data);
 					info.name = m_memoryBlock.Strdup(mappedFileName).data;
@@ -163,16 +163,21 @@ namespace uba
 		g_directoryTable.ParseDirectoryTable(directoryTableSize);
 	}
 
-	void UpdateWrittenFiles(BinaryReader& reader)
+	// Returns true if done (no overflow)
+	bool UpdateWrittenFiles(BinaryReader& reader)
 	{
 		u32 count = reader.ReadU32();
+		u8 overflow = reader.ReadByte();
 		while (count--)
 		{
 			StringKey key = reader.ReadStringKey();
 			auto insres = g_mappedFileTable.m_lookup.try_emplace(key);
 			FileInfo& info = insres.first->second;
 
+			bool isInTemp = reader.ReadBool();
 			StringBuffer<> originalName;
+			if (isInTemp)
+				originalName.Append(g_systemTemp).Append(PathSeparator);
 			reader.ReadString(originalName);
 			if (!info.originalName || !originalName.Equals(info.originalName))
 				info.originalName = g_mappedFileTable.m_memoryBlock.Strdup(originalName).data;
@@ -180,8 +185,8 @@ namespace uba
 			StringBufferBase& backedName = originalName.Clear();
 			reader.ReadString(backedName);
 
-			FileMappingHandle mappingHandle = FileMappingHandle::FromU64(reader.ReadU64());
-			u64 fileSize = reader.ReadU64();
+			FileMappingHandle mappingHandle = FileMappingHandle::FromU64(reader.Read7BitEncoded());
+			u64 fileSize = reader.Read7BitEncoded();
 			info.fileNameKey = key;
 			info.size = fileSize;
 
@@ -190,7 +195,7 @@ namespace uba
 			#endif
 
 			if (mappingHandle.IsValid())
-				backedName.Clear().Append(':').AppendHex(mappingHandle.ToU64()).Append('-').AppendHex(0);
+				backedName.Clear().Append(':').AppendBase62(mappingHandle.ToU64());
 
 			if (!info.name || !backedName.Equals(info.name))
 				info.name = g_mappedFileTable.m_memoryBlock.Strdup(backedName).data;
@@ -212,33 +217,41 @@ namespace uba
 
 			DEBUG_LOG(TC("GOT WRITTEN FILE: %s (BackedFile: %s Size: %llu)"), info.originalName, info.name, info.size);
 		}
+		return overflow == 0;
 	}
 
 	void Rpc_UpdateTables()
 	{
 		u32 directoryTableSize;
 		u32 fileMappingTableSize;
+		bool done;
 		{
 			RPC_MESSAGE(UpdateTables, updateTables)
 			writer.Flush();
 			BinaryReader reader;
 			directoryTableSize = reader.ReadU32();
 			fileMappingTableSize = reader.ReadU32();
-			UpdateWrittenFiles(reader);
+			done = UpdateWrittenFiles(reader);
 
 			DEBUG_LOG_PIPE(L"UpdateTables", L"");
 		}
+		if (!done)
+			Rpc_GetWrittenFiles();
+
 		g_directoryTable.ParseDirectoryTable(directoryTableSize);
 		g_mappedFileTable.Parse(fileMappingTableSize);
 	}
 
 	void Rpc_GetWrittenFiles()
 	{
-		RPC_MESSAGE(GetWrittenFiles, updateTables)
-		writer.Flush();
-		BinaryReader reader;
-
-		UpdateWrittenFiles(reader);
+		while (true)
+		{
+			RPC_MESSAGE(GetWrittenFiles, updateTables)
+			writer.Flush();
+			BinaryReader reader;
+			if (UpdateWrittenFiles(reader))
+				break;
+		}
 	}
 
 	u32 Rpc_GetEntryOffset(const StringKey& entryNameKey, StringView entryName, bool checkIfDir)
diff --git a/Engine/Source/Programs/UnrealBuildAccelerator/Detours/Private/Windows/UbaDetoursFunctionsKernelBase.inl b/Engine/Source/Programs/UnrealBuildAccelerator/Detours/Private/Windows/UbaDetoursFunctionsKernelBase.inl
--- a/Engine/Source/Programs/UnrealBuildAccelerator/Detours/Private/Windows/UbaDetoursFunctionsKernelBase.inl
+++ b/Engine/Source/Programs/UnrealBuildAccelerator/Detours/Private/Windows/UbaDetoursFunctionsKernelBase.inl
@@ -3441,31 +3441,46 @@ DWORD Detoured_WaitForSingleObjectEx(HANDLE hHandle, DWORD dwMilliseconds, BOOL
 DWORD Detoured_WaitForMultipleObjects(DWORD nCount, const HANDLE* lpHandles, BOOL bWaitAll, DWORD dwMilliseconds)
 {
 	DETOURED_CALL(WaitForMultipleObjects);
+	
+	UBA_ASSERT(nCount != 0);
 
-	bool isProcess = false;
-	HANDLE* tempHandles = (HANDLE*)_malloca(nCount * sizeof(HANDLE));
-	if (!tempHandles)
+	bool hasProcess = false;
+	auto tempHandles = (HANDLE*)_malloca(nCount * sizeof(HANDLE));
+	auto isProcess = (bool*)_malloca(nCount * sizeof(bool));
+	if (!tempHandles || !isProcess)
 		FatalError(1355, L"Here to please static analyzer");
 
+#ifndef __clang_analyzer__
+	auto g = MakeGuard([&]() { _freea(tempHandles); _freea(isProcess); });
+#endif
+
 	for (u32 i = 0; i != nCount; ++i)
 	{
 		HANDLE hHandle = lpHandles[i];
+		isProcess[i] = false;
 		if (isDetouredHandle(hHandle))
 		{
 			DetouredHandle& dh = asDetouredHandle(hHandle);
 			hHandle = asDetouredHandle(hHandle).trueHandle;
-			isProcess |= dh.type == HandleType_Process;
+			if (dh.type == HandleType_Process)
+			{
+				hasProcess = true;
+				isProcess[i] = true;
+			}
 		}
 		tempHandles[i] = hHandle;
 	}
 
 	auto res = True_WaitForMultipleObjectsEx(nCount, tempHandles, bWaitAll, dwMilliseconds, false);
 
-#ifndef __clang_analyzer__
-	_freea(tempHandles);
-#endif
 
-	if (!isProcess || res != WAIT_OBJECT_0)
+	if (!hasProcess)
+		return res;
+
+	if (res < WAIT_OBJECT_0 || res > WAIT_OBJECT_0 + nCount)
+		return res;
+	
+	if (!isProcess[res - WAIT_OBJECT_0])
 		return res;
 
 	DEBUG_LOG_DETOURED(L"WaitForMultipleObjects", L"");
diff --git a/Engine/Source/Programs/UnrealBuildAccelerator/Detours/Private/Windows/UbaDetoursFunctionsNtDll.inl b/Engine/Source/Programs/UnrealBuildAccelerator/Detours/Private/Windows/UbaDetoursFunctionsNtDll.inl
--- a/Engine/Source/Programs/UnrealBuildAccelerator/Detours/Private/Windows/UbaDetoursFunctionsNtDll.inl
+++ b/Engine/Source/Programs/UnrealBuildAccelerator/Detours/Private/Windows/UbaDetoursFunctionsNtDll.inl
@@ -1198,14 +1198,17 @@ NTSTATUS NTAPI Shared_NtCreateFile(bool IsCreateFunc, PHANDLE hFileHandle, ACCES
 
 		const wchar_t* handleStr = lpFileName + 1;
 		const wchar_t* handleStrEnd = wcschr(handleStr, '-');
+		u64 mappingOffset = 0;
 		if (!handleStrEnd)
 		{
-			UBA_ASSERT(handleStrEnd);
-			return STATUS_UNSUCCESSFUL;
+			handleStrEnd = handleStr + wcslen(handleStr);
+		}
+		else
+		{
+			const wchar_t* mappingOffsetStr = handleStrEnd + 1;
+			mappingOffset = StringToValueBase62(mappingOffsetStr, wcslen(mappingOffsetStr));
 		}
-		HANDLE mappingHandle = FileMappingHandle::FromU64(StringToValue(handleStr, handleStrEnd - handleStr)).mh;
-		const wchar_t* mappingOffsetStr = handleStrEnd + 1;
-		u64 mappingOffset = StringToValue(mappingOffsetStr, wcslen(mappingOffsetStr));
+		HANDLE mappingHandle = FileMappingHandle::FromU64(StringToValueBase62(handleStr, handleStrEnd - handleStr)).mh;
 		info.trueFileMapOffset = mappingOffset;
 
 		info.isFileMap = true;
@@ -1250,14 +1253,8 @@ NTSTATUS NTAPI Shared_NtCreateFile(bool IsCreateFunc, PHANDLE hFileHandle, ACCES
 		if (!info.memoryFile)
 		{
 			const wchar_t* handleStr = lpFileName + 1;
-			const wchar_t* handleStrEnd = wcschr(handleStr, '-');
-			if (!handleStrEnd)
-			{
-				UBA_ASSERT(handleStrEnd);
-				return STATUS_UNSUCCESSFUL;
-			}
-			auto mappingHandle = FileMappingHandle::FromU64(StringToValue(handleStr, handleStrEnd - handleStr));
-			UBA_ASSERT(!StringToValue(handleStrEnd + 1, wcslen(handleStrEnd + 1)));
+			UBA_ASSERT(!wcschr(handleStr, '-'));
+			auto mappingHandle = FileMappingHandle::FromU64(StringToValueBase62(handleStr, wcslen(handleStr)));
 			UBA_ASSERT(info.size != InvalidValue);
 
 			TakeFileOwnershipFromSession(mappingHandle, info.size);
diff --git a/Engine/Source/Programs/UnrealBuildAccelerator/Test/Private/UbaTestAll.h b/Engine/Source/Programs/UnrealBuildAccelerator/Test/Private/UbaTestAll.h
--- a/Engine/Source/Programs/UnrealBuildAccelerator/Test/Private/UbaTestAll.h
+++ b/Engine/Source/Programs/UnrealBuildAccelerator/Test/Private/UbaTestAll.h
@@ -18,6 +18,7 @@ namespace uba
 
 
 #define UBA_ALLPLATFORM_TESTS \
+		UBA_TEST(TestEncodings) \
 		UBA_TEST(TestTime) \
 		UBA_TEST(TestEvents) \
 		UBA_TEST(TestPaths) \
diff --git a/Engine/Source/Programs/UnrealBuildAccelerator/Test/Private/UbaTestBasics.cpp b/Engine/Source/Programs/UnrealBuildAccelerator/Test/Private/UbaTestBasics.cpp
--- a/Engine/Source/Programs/UnrealBuildAccelerator/Test/Private/UbaTestBasics.cpp
+++ b/Engine/Source/Programs/UnrealBuildAccelerator/Test/Private/UbaTestBasics.cpp
@@ -23,6 +23,27 @@ namespace uba
 {
 	bool GetCpuTime(u64& outTotalTime, u64& outIdleTime);
 
+	bool TestEncodings(LoggerWithWriter& logger, const StringBufferBase& rootDir)
+	{
+		StringBuffer<> buf;
+
+		u64 v = 0;
+		buf.Clear().AppendBase62(v);
+		if (v != StringToValueBase62(buf.data, buf.count))
+			return false;
+
+		v = 5473789457;
+		buf.Clear().AppendBase62(v);
+		if (v != StringToValueBase62(buf.data, buf.count))
+			return false;
+
+		v = ~u64(0);
+		buf.Clear().AppendBase62(v);
+		if (v != StringToValueBase62(buf.data, buf.count))
+			return false;
+		return true;
+	}
+
 	bool TestTime(LoggerWithWriter& logger, const StringBufferBase& rootDir)
 	{
 #if 0
