[UBA]
* File missing from changelist 42953111

diff --git a/Engine/Source/Programs/UnrealBuildAccelerator/Common/Private/UbaFileMapping.cpp b/Engine/Source/Programs/UnrealBuildAccelerator/Common/Private/UbaFileMapping.cpp
--- a/Engine/Source/Programs/UnrealBuildAccelerator/Common/Private/UbaFileMapping.cpp
+++ b/Engine/Source/Programs/UnrealBuildAccelerator/Common/Private/UbaFileMapping.cpp
@@ -714,12 +714,16 @@ namespace uba
 #if PLATFORM_WINDOWS
 		const tchar* handleStr = str.data + 1;
 		const tchar* handleStrEnd = TStrchr(handleStr, '-');
-		if (!handleStrEnd)
-			return {};
+		u64 mappingOffset = 0;
+		if (handleStrEnd)
+		{
+			const tchar* mappingOffsetStr = handleStrEnd + 1;
+			mappingOffset = StringToValueBase62(mappingOffsetStr, wcslen(mappingOffsetStr));
+		}
+		else
+			handleStrEnd = handleStr + wcslen(handleStr);
 
-		FileMappingHandle h = FileMappingHandle::FromU64(StringToValue(handleStr, handleStrEnd - handleStr));
-		const tchar* mappingOffsetStr = handleStrEnd + 1;
-		u64 mappingOffset = StringToValue(mappingOffsetStr, wcslen(mappingOffsetStr));
+		FileMappingHandle h = FileMappingHandle::FromU64(StringToValueBase62(handleStr, handleStrEnd - handleStr));
 		
 		return MapView(h, mappingOffset, size, hint);
 #else
