[UBA]
* Added detouring of openat. It is needed to track file created by cp (copy)

diff --git a/Engine/Source/Programs/UnrealBuildAccelerator/Detours/Private/PosixOS/UbaDetoursFunctionsPosix.cpp b/Engine/Source/Programs/UnrealBuildAccelerator/Detours/Private/PosixOS/UbaDetoursFunctionsPosix.cpp
--- a/Engine/Source/Programs/UnrealBuildAccelerator/Detours/Private/PosixOS/UbaDetoursFunctionsPosix.cpp
+++ b/Engine/Source/Programs/UnrealBuildAccelerator/Detours/Private/PosixOS/UbaDetoursFunctionsPosix.cpp
@@ -61,6 +61,7 @@ using namespace uba;
 	DETOURED_FUNCTION(read) \
 	DETOURED_FUNCTION(pread) \
 	DETOURED_FUNCTION(open) \
+	DETOURED_FUNCTION(openat) \
 	DETOURED_FUNCTION(dup) \
 	DETOURED_FUNCTION(dup2) \
 	DETOURED_FUNCTION(close) \
@@ -121,6 +122,10 @@ using namespace uba;
 	DETOURED_FUNCTION_LINUX \
 	DETOURED_FUNCTION_MACOS \
 
+// Used by cp and might need detour
+// fadvise64
+// newfstatat
+// copy_file_range
 
 #if UBA_DETOUR_DEBUG && PLATFORM_LINUX
 #define DETOURED_FUNCTION_DEBUG \
@@ -620,7 +625,7 @@ int Shared_fstat(const char* funcName, int fd, struct stat* attr, const True_fst
 	if (findIt == g_fileHandles.end())
 	{
 		int res = trueFstat(fd, attr);
-		DEBUG_LOG_TRUE(funcName, "(%i) -> %i (%s)", fd, res, StrError(res, errno));
+		DEBUG_LOG_TRUE(funcName, "(%i) NODETOUR (size: %llu) -> %i (%s)", fd, attr->st_size, res, StrError(res, errno));
 		return res;
 	}
 
@@ -630,7 +635,7 @@ int Shared_fstat(const char* funcName, int fd, struct stat* attr, const True_fst
 	if (fo.desiredAccess & AccessFlag_Write)
 	{
 		int res = trueFstat(fd, attr);
-		DEBUG_LOG_TRUE(funcName, "(%i) (%s) -> %i (%s)", fd, fi.originalName, res, StrError(res, errno));
+		DEBUG_LOG_TRUE(funcName, "(%i) AFW (%s size: %llu) -> %i (%s)", fd, fi.originalName, attr->st_size, res, StrError(res, errno));
 		return res;
 	}
 
@@ -641,14 +646,14 @@ int Shared_fstat(const char* funcName, int fd, struct stat* attr, const True_fst
 	{
 		lock.Leave();
 		int res = trueFstat(fd, attr);
-		DEBUG_LOG_TRUE(funcName, "(%i) -> %i (%s)", fd, res, StrError(res, errno));
+		DEBUG_LOG_TRUE(funcName, "(%i) NOCACHE (%s size: %llu) -> %i (%s)", fd, fi.originalName, fileAttr.data.st_size, res, StrError(res, errno));
 		return res;
 	}
 
 
 	int res = fileAttr.lastError == 0 ? 0 : -1;
 
-	DEBUG_LOG_DETOURED(funcName, "(%i) (name: %s size: %llu id: %llu dev: %u)-> %i (%s)", fd, fi.originalName, fileAttr.data.st_size, fileAttr.data.st_ino, fileAttr.data.st_dev, res, StrError(res, fileAttr.lastError));
+	DEBUG_LOG_DETOURED(funcName, "(%i) (%s size: %llu id: %llu dev: %u)-> %i (%s)", fd, fi.originalName, fileAttr.data.st_size, fileAttr.data.st_ino, fileAttr.data.st_dev, res, StrError(res, fileAttr.lastError));
 
 	errno = fileAttr.lastError;
 	if (res == 0)
@@ -1138,6 +1143,24 @@ UBA_EXPORT int UBA_WRAPPER(open)(const char* file, int flags, ...)
 	return Shared_open("open", file, flags, mode, [](const char* realFile, int flags, int mode) { return TRUE_WRAPPER(open)(realFile, flags, mode); });
 }
 
+UBA_EXPORT int openat(int dirfd, const char* pathname, int flags, ...)
+{
+	va_list args;
+	va_start(args, flags);
+	int mode = va_arg(args, int);
+	va_end(args);
+	UBA_INIT_DETOUR(openat, dirfd, pathname, flags, mode);
+
+	if (dirfd == AT_FDCWD)
+		return Shared_open("openat", pathname, flags, mode, [](const char* realFile, int flags, int mode) { return TRUE_WRAPPER(open)(realFile, flags, mode); });
+
+	UBA_ASSERT(!g_runningRemote);
+	
+	int res = TRUE_WRAPPER(openat)(dirfd, pathname, flags, mode);
+	DEBUG_LOG_TRUE("openat", "(NODETOUR) %s -> %s", pathname, res);
+	return res;
+}
+
 #if UBA_DETOUR_DEBUG && PLATFORM_LINUX
 UBA_EXPORT ssize_t UBA_WRAPPER(write)(int fd, const void* buf, size_t count)
 {
@@ -2592,19 +2615,6 @@ UBA_EXPORT int __xstat64(int ver, const char* file, struct stat64* attr)
 	UBA_NOT_IMPLEMENTED(__xstat64);
 	return TRUE_WRAPPER(__xstat64)(ver, file, attr);
 }
-
-UBA_EXPORT int openat(int dirfd, const char* pathname, int flags, ...)
-{
-	UBA_NOT_IMPLEMENTED(openat);
-	return -1;
-}
-
-UBA_EXPORT int openat(int dirfd, const char* pathname, int flags, mode_t mode, ...)
-{
-	UBA_NOT_IMPLEMENTED(openat);
-	return -1;
-}
-
 UBA_EXPORT int openat64(int __fd, const char* __file, int __oflag, ...)
 {
 	UBA_NOT_IMPLEMENTED(openat64);
