[UBA]
* Fixed uba compile errors when compiling uba with clang ltcg

diff --git a/Engine/Source/Programs/UnrealBuildAccelerator/Common/Private/UbaNetworkBackendTcp.cpp b/Engine/Source/Programs/UnrealBuildAccelerator/Common/Private/UbaNetworkBackendTcp.cpp
--- a/Engine/Source/Programs/UnrealBuildAccelerator/Common/Private/UbaNetworkBackendTcp.cpp
+++ b/Engine/Source/Programs/UnrealBuildAccelerator/Common/Private/UbaNetworkBackendTcp.cpp
@@ -1370,7 +1370,7 @@ namespace uba
 			*prevStats = stats;
 		}
 
-		double rrate = sent ? (100.0 * retrans / sent) : 0;
+		double rrate = sent ? (100.0 * double(retrans) / double(sent)) : 0;
 		logger.Info(TC("%s  recv=%lld  sent=%lld  retrans=%lld  (%.2f%%)  inErr=%lld  RST=%lld"), temp.data, recv, sent, retrans, rrate, inerrs, outrsts);
 		#endif
 	}
