[UBA]
* Fixed macos compile error (missing UBA_WRAPPER around openat)

diff --git a/Engine/Source/Programs/UnrealBuildAccelerator/Detours/Private/PosixOS/UbaDetoursFunctionsPosix.cpp b/Engine/Source/Programs/UnrealBuildAccelerator/Detours/Private/PosixOS/UbaDetoursFunctionsPosix.cpp
--- a/Engine/Source/Programs/UnrealBuildAccelerator/Detours/Private/PosixOS/UbaDetoursFunctionsPosix.cpp
+++ b/Engine/Source/Programs/UnrealBuildAccelerator/Detours/Private/PosixOS/UbaDetoursFunctionsPosix.cpp
@@ -1143,7 +1143,7 @@ UBA_EXPORT int UBA_WRAPPER(open)(const char* file, int flags, ...)
 	return Shared_open("open", file, flags, mode, [](const char* realFile, int flags, int mode) { return TRUE_WRAPPER(open)(realFile, flags, mode); });
 }
 
-UBA_EXPORT int openat(int dirfd, const char* pathname, int flags, ...)
+UBA_EXPORT int UBA_WRAPPER(openat)(int dirfd, const char* pathname, int flags, ...)
 {
 	va_list args;
 	va_start(args, flags);
