[UBA]
* Increased debug log header buffer size since it is too small for some osx command lines

diff --git a/Engine/Source/Programs/UnrealBuildAccelerator/Detours/Private/PosixOS/UbaDetoursFunctionsPosix.cpp b/Engine/Source/Programs/UnrealBuildAccelerator/Detours/Private/PosixOS/UbaDetoursFunctionsPosix.cpp
--- a/Engine/Source/Programs/UnrealBuildAccelerator/Detours/Private/PosixOS/UbaDetoursFunctionsPosix.cpp
+++ b/Engine/Source/Programs/UnrealBuildAccelerator/Detours/Private/PosixOS/UbaDetoursFunctionsPosix.cpp
@@ -2703,7 +2703,8 @@ namespace uba
 	void LogHeader()
 	{
 		#if UBA_DEBUG_LOG_ENABLED
-		StringBuffer<32 * 1024> buf;
+		static StringBuffer<128 * 1024> buf;
+		buf.Clear();
 		#if PLATFORM_LINUX
 		int fd = TRUE_WRAPPER(open)("/proc/self/cmdline", O_RDONLY);
 		if (fd != -1)
