[UBA]
* Fixed CreateFile handling path where it has OPEN_EXISTING but also GENERIC_WRITE and application rule returns true for file being output file (storing it in memory). Code path will now open file for read and copy over content to memory file (before it failed). This code path is used for incremental linking which no-one at Epic has enabled so it has been flying under the radar this entire time

diff --git a/Engine/Source/Programs/UnrealBuildAccelerator/Detours/Private/Windows/UbaDetoursFunctionsNtDll.inl b/Engine/Source/Programs/UnrealBuildAccelerator/Detours/Private/Windows/UbaDetoursFunctionsNtDll.inl
--- a/Engine/Source/Programs/UnrealBuildAccelerator/Detours/Private/Windows/UbaDetoursFunctionsNtDll.inl
+++ b/Engine/Source/Programs/UnrealBuildAccelerator/Detours/Private/Windows/UbaDetoursFunctionsNtDll.inl
@@ -1278,26 +1278,68 @@ NTSTATUS NTAPI Shared_NtCreateFile(bool IsCreateFunc, PHANDLE hFileHandle, ACCES
 
 		if (!info.memoryFile)
 		{
+			bool isOutput = (isWrite || isDeleteOnClose) && g_rules->IsOutputFile(fileName, g_systemTemp);
+			bool isLocal = !isOutput;
+			MemoryFile* mf = new MemoryFile(isLocal, FileTypeMaxSize(fileName, isSystemOrTempFile), isThrowAway, fileName.data);
+
 			if (!isThrowAway && CreateDisposition == FILE_OPEN)
 			{
-				if (!info.memoryFile)
+				if (!isWrite)
+				{
+					*hFileHandle = INVALID_HANDLE_VALUE;
+					DEBUG_LOG_DETOURED(funcName, L"NOTEXISTS (%ls) -> Error", fileName.data);
+					return STATUS_OBJECT_NAME_NOT_FOUND;
+				}
+
+				// We need to open file for read first and then copy content over to a memory file since this is actually a write
+				// (We end up in this code path is used for incremental linking)
+
+				DEBUG_LOG_DETOURED(funcName, L"INTERNAL READ FOR MEMORYWRITE (%ls) (%u %u)", fileName.data, CreateDisposition, CreateOptions);
+
+				HANDLE fileHandle;
+				IO_STATUS_BLOCK ioStatusBlock;
+				lookupLock.Leave();
+				NTSTATUS res2 = Shared_NtCreateFile(IsCreateFunc, &fileHandle, FILE_GENERIC_READ, ObjectAttributes, &ioStatusBlock, NULL, FILE_ATTRIBUTE_NORMAL, FILE_SHARE_READ, FILE_OPEN, FILE_SEQUENTIAL_ONLY|FILE_SYNCHRONOUS_IO_NONALERT, NULL, 0);
+				lookupLock.Enter();
+				if (res2 != STATUS_SUCCESS)
 				{
 					*hFileHandle = INVALID_HANDLE_VALUE;
-					DEBUG_LOG_DETOURED(funcName, L"ALREADYEXISTS (%ls) -> Error", fileName.data);
-					return STATUS_OBJECT_NAME_EXISTS;//STATUS_OBJECT_NAME_NOT_FOUND;
+					DEBUG_LOG_DETOURED(funcName, L"NOTEXISTS (%ls) -> Error", fileName.data);
+					return res2;
+				}
+
+				LARGE_INTEGER fileSize2;
+				if (!GetFileSizeEx(fileHandle, &fileSize2))
+				{
+					UBA_ASSERTF(false, TC("%s"), fileName.data);
+					return STATUS_OBJECT_NAME_EXISTS;
+				}
+				DetouredHandle tempDh(HandleType_File, INVALID_HANDLE_VALUE);
+				mf->EnsureCommitted(tempDh, fileSize2.QuadPart);
+				u64 left = fileSize2.QuadPart;
+				u8* writePos = mf->baseAddress;
+				while (left)
+				{
+					DWORD toRead = (DWORD)Min(left, u64(~0u));
+					DWORD read;
+					if (!ReadFile(fileHandle, writePos, toRead, &read, NULL))
+					{
+						UBA_ASSERTF(false, TC("%s"), fileName.data);
+						return STATUS_OBJECT_NAME_EXISTS;
+					}
+					writePos += read;
+					left -= read;
 				}
+				NtClose(fileHandle);
+		
+				mf->writtenSize = fileSize2.QuadPart;
 			}
 			else
 			{
-				bool isOutput = (isWrite || isDeleteOnClose) && g_rules->IsOutputFile(fileName, g_systemTemp);
 				if (isOutput && (CreateDisposition & FILE_OPEN_IF) == 0)
 				{
 					UBA_ASSERTF(false, TC("Trying to open %s with openif. This is not supported"), fileName.data);
 				}
-
-				bool isLocal = !isOutput;
-				//UBA_ASSERTF(CreateDisposition != FILE_OPEN || Contains(fileName.data, L"vctip_"), TC("Unsupported disposition %u for file %s"), CreateDisposition, fileName.data);
-				info.memoryFile = new MemoryFile(isLocal, FileTypeMaxSize(fileName, isSystemOrTempFile), isThrowAway, fileName.data);
 			}
 
 			// TODO: Time should be in sync with host machine!
@@ -1305,9 +1347,11 @@ NTSTATUS NTAPI Shared_NtCreateFile(bool IsCreateFunc, PHANDLE hFileHandle, ACCES
 			SYSTEMTIME st;
 			GetSystemTime(&st);
 			SystemTimeToFileTime(&st, &ft);
-			info.memoryFile->fileTime = (u64&)ft;
-			info.memoryFile->volumeSerial = 1;
-			info.memoryFile->fileIndex = InterlockedDecrement(&g_memoryFileIndexCounter);
+			mf->fileTime = (u64&)ft;
+			mf->volumeSerial = 1;
+			mf->fileIndex = InterlockedDecrement(&g_memoryFileIndexCounter);
+
+			info.memoryFile = mf;
 		}
 		else
 		{
@@ -1371,7 +1415,7 @@ NTSTATUS NTAPI Shared_NtCreateFile(bool IsCreateFunc, PHANDLE hFileHandle, ACCES
 
 	UBA_ASSERT(info.originalName);
 	*hFileHandle = CreateFileHandle(*hFileHandle);
-	DEBUG_LOG_TRUE(funcName, L"%ls %llu (%ls)%s -> %ls", isWriteStr, uintptr_t(*hFileHandle), tempFileName.data, isDeleteOnClose ? TC(" DeleteOnClose") : TC(""), ToString(res).data);
+	DEBUG_LOG_TRUE(funcName, L"%ls %llu (%ls)%s (%u %u) -> %ls", isWriteStr, uintptr_t(*hFileHandle), tempFileName.data, isDeleteOnClose ? TC(" DeleteOnClose") : TC(""), CreateDisposition, CreateOptions, ToString(res).data);
 	return res;
 }
 
diff --git a/Engine/Source/Programs/UnrealBuildAccelerator/Test/Private/UbaTestSession.cpp b/Engine/Source/Programs/UnrealBuildAccelerator/Test/Private/UbaTestSession.cpp
--- a/Engine/Source/Programs/UnrealBuildAccelerator/Test/Private/UbaTestSession.cpp
+++ b/Engine/Source/Programs/UnrealBuildAccelerator/Test/Private/UbaTestSession.cpp
@@ -294,6 +294,9 @@ namespace uba
 				return logger.Error(TC("Failed to create dir %s"), dir.data);
 		}
 
+		if (!CreateTestFile(logger, ToView(workingDir), TCV("File4.out"), TCV("0")))
+			return false;
+
 		if (!RunTestApp(logger, workingDir, runProcess, TC("")))
 			return false;
 
diff --git a/Engine/Source/Programs/UnrealBuildAccelerator/TestApp/Private/Windows/UbaTestAppWin.cpp b/Engine/Source/Programs/UnrealBuildAccelerator/TestApp/Private/Windows/UbaTestAppWin.cpp
--- a/Engine/Source/Programs/UnrealBuildAccelerator/TestApp/Private/Windows/UbaTestAppWin.cpp
+++ b/Engine/Source/Programs/UnrealBuildAccelerator/TestApp/Private/Windows/UbaTestAppWin.cpp
@@ -185,6 +185,18 @@ int wmain(int argc, wchar_t* argv[])
 		}
 
 
+		{
+			HANDLE fh = CreateFileW(L"File4.out", GENERIC_WRITE, FILE_SHARE_READ, NULL, OPEN_EXISTING, 0, NULL);
+			if (fh == INVALID_HANDLE_VALUE)
+				return LogError(L"Failed to open File4.out for read with write permissions");
+			DWORD bytesRead;
+			char data[1];
+			if (!ReadFile(fh, data, 1, &bytesRead, NULL) || bytesRead != 1 || data[0] != '0')
+				return LogError(L"Failed to read one byte from File4.out");
+			CloseHandle(fh);
+		}
+
+
 		{
 			STARTUPINFOW si;
 			memset(&si, 0, sizeof(si));
