[UBA]
* Fixed static analysis warnings (no bug fixes and most of them just bogus and tool not being able to read code properly)

diff --git a/Engine/Source/Programs/UnrealBuildAccelerator/CacheService/Private/UbaCacheService.cpp b/Engine/Source/Programs/UnrealBuildAccelerator/CacheService/Private/UbaCacheService.cpp
--- a/Engine/Source/Programs/UnrealBuildAccelerator/CacheService/Private/UbaCacheService.cpp
+++ b/Engine/Source/Programs/UnrealBuildAccelerator/CacheService/Private/UbaCacheService.cpp
@@ -427,8 +427,6 @@ namespace uba
 
 		if (shouldCrash)
 		{
-			void* mem = new u8[1ull*1024*1024*1024];
-			memset(mem, 1, 1ull*1024*1024*1024);
 			#ifndef __clang_analyzer__
 			int* ptr = nullptr;
 			*ptr = 0;
diff --git a/Engine/Source/Programs/UnrealBuildAccelerator/Common/Private/UbaCacheClient.cpp b/Engine/Source/Programs/UnrealBuildAccelerator/Common/Private/UbaCacheClient.cpp
--- a/Engine/Source/Programs/UnrealBuildAccelerator/Common/Private/UbaCacheClient.cpp
+++ b/Engine/Source/Programs/UnrealBuildAccelerator/Common/Private/UbaCacheClient.cpp
@@ -402,7 +402,10 @@ namespace uba
 				}
 			}
 
+			#ifndef __clang_analyzer__ // Must be bug in analyzer
 			UBA_ASSERT(IsCompressed(casKey));
+			#endif
+
 			insres.first->second = bucket.sendCasKeyTable.Add(casKey, pathOffset, requiredCasTableSize);
 		}
 
diff --git a/Engine/Source/Programs/UnrealBuildAccelerator/Common/Private/UbaNetworkClient.cpp b/Engine/Source/Programs/UnrealBuildAccelerator/Common/Private/UbaNetworkClient.cpp
--- a/Engine/Source/Programs/UnrealBuildAccelerator/Common/Private/UbaNetworkClient.cpp
+++ b/Engine/Source/Programs/UnrealBuildAccelerator/Common/Private/UbaNetworkClient.cpp
@@ -426,13 +426,11 @@ namespace uba
 	void NetworkClient::ValidateNetwork(Logger& logger, bool full)
 	{
 		UnorderedMap<NetworkBackend*, Vector<void*>> backends;
-		u32 connectionsCount;
 		{
 			LogStallScope lss(logger, LogEntryType_Info, 1, TC("  Connections lock took %s"));
 			SCOPED_WRITE_LOCK(m_connectionsLock, lock);
 			lss.Leave();
 
-			connectionsCount = u32(m_connections.size());
 			u32 connectionIndex = 0;
 			for (auto& c : m_connections)
 			{
diff --git a/Engine/Source/Programs/UnrealBuildAccelerator/Common/Private/UbaObjectFileCoff.cpp b/Engine/Source/Programs/UnrealBuildAccelerator/Common/Private/UbaObjectFileCoff.cpp
--- a/Engine/Source/Programs/UnrealBuildAccelerator/Common/Private/UbaObjectFileCoff.cpp
+++ b/Engine/Source/Programs/UnrealBuildAccelerator/Common/Private/UbaObjectFileCoff.cpp
@@ -275,7 +275,7 @@ namespace uba
 				tmp = tmp.substr(0, comma - tmp.data());
 			}
 			StringKey key = ToStringKeyRaw(tmp.data(), tmp.size());
-			auto insres = m_exports.emplace(key, ExportInfo{tmp, isData, index++});
+			m_exports.emplace(key, ExportInfo{tmp, isData, index++});
 		}
 		return true;
 	}
diff --git a/Engine/Source/Programs/UnrealBuildAccelerator/Common/Private/UbaSessionServer.cpp b/Engine/Source/Programs/UnrealBuildAccelerator/Common/Private/UbaSessionServer.cpp
--- a/Engine/Source/Programs/UnrealBuildAccelerator/Common/Private/UbaSessionServer.cpp
+++ b/Engine/Source/Programs/UnrealBuildAccelerator/Common/Private/UbaSessionServer.cpp
@@ -345,9 +345,10 @@ namespace uba
 
 			u32 knownInputsBytes = u32((kiEnd - kiBegin)*sizeof(tchar));
 			void* knownInputsCopy = malloc(knownInputsBytes);
-			memcpy(knownInputsCopy, knownInputs, knownInputsBytes);
+			if (knownInputsCopy) // static analysis check
+				memcpy(knownInputsCopy, knownInputs, knownInputsBytes);
 
-			#ifdef __clang_analyzer__
+			#if defined(__clang_analyzer__)
 			free(knownInputsCopy); // doesn't seem to understand it is handed over
 			#endif
 
diff --git a/Engine/Source/Programs/UnrealBuildAccelerator/Common/Private/UbaStorage.cpp b/Engine/Source/Programs/UnrealBuildAccelerator/Common/Private/UbaStorage.cpp
--- a/Engine/Source/Programs/UnrealBuildAccelerator/Common/Private/UbaStorage.cpp
+++ b/Engine/Source/Programs/UnrealBuildAccelerator/Common/Private/UbaStorage.cpp
@@ -915,6 +915,7 @@ namespace uba
 
 	bool StorageImpl::GetFileEntry(FileEntry*& out, const tchar* file)
 	{
+		out = nullptr;
 		StringBuffer<> forKey;
 		if (!FixPath(file, nullptr, 0, forKey))
 			return m_logger.Error(TC("GetCasEntry - FixPath failed for %s"), file);
@@ -2451,7 +2452,7 @@ namespace uba
 	bool StorageImpl::DeleteCasForFile(const tchar* file)
 	{
 		FileEntry* fileEntry;
-		if (!GetFileEntry(fileEntry, file))
+		if (!GetFileEntry(fileEntry, file) || !fileEntry)
 			return false;
 
 		SCOPED_FUTEX(fileEntry->lock, entryLock);
diff --git a/Engine/Source/Programs/UnrealBuildAccelerator/Common/Private/UbaStorageUtils.cpp b/Engine/Source/Programs/UnrealBuildAccelerator/Common/Private/UbaStorageUtils.cpp
--- a/Engine/Source/Programs/UnrealBuildAccelerator/Common/Private/UbaStorageUtils.cpp
+++ b/Engine/Source/Programs/UnrealBuildAccelerator/Common/Private/UbaStorageUtils.cpp
@@ -962,7 +962,7 @@ namespace uba
 				if (index >= rec.workCount)
 					return;
 
-				if (!slot)
+				if (!compressSlotBuffer)
 				{
 					slot = rec.bufferSlots->Pop();
 					compressSlotBuffer = slot + BufferSlotHalfSize;
diff --git a/Engine/Source/Programs/UnrealBuildAccelerator/Core/Public/UbaProcessUtils.h b/Engine/Source/Programs/UnrealBuildAccelerator/Core/Public/UbaProcessUtils.h
--- a/Engine/Source/Programs/UnrealBuildAccelerator/Core/Public/UbaProcessUtils.h
+++ b/Engine/Source/Programs/UnrealBuildAccelerator/Core/Public/UbaProcessUtils.h
@@ -4,11 +4,16 @@
 
 #include "UbaPlatform.h"
 
-#if PLATFORM_WINDOWS
+#if defined(__clang_analyzer__)
+#define stackAlloc(x) malloc(x)
+#define stackFree(x) free(x)
+#elif PLATFORM_WINDOWS
 #include <malloc.h>
+#define stackAlloc(x) _malloca(x)
+#define stackFree(x) _freea(x)
 #else
-#define _malloca(x) alloca(x)
-#define _freea(x)
+#define stackAlloc(x) alloca(x)
+#define stackFree(x)
 #endif
 
 namespace uba
@@ -28,8 +33,9 @@ namespace uba
 		CharType currentChar = 0;
 		CharType lastChar = 0;
 		bool isBackslashOwned = false;
-		CharType* argWriteBuffer = static_cast<CharType*>( _malloca(sizeof(CharType) * (argumentsLen + 1)) );
-		auto g = MakeGuard([&]() { _freea(argWriteBuffer); });
+
+		CharType* argWriteBuffer = static_cast<CharType*>( stackAlloc(sizeof(CharType) * (argumentsLen + 1)) );
+		auto g = MakeGuard([&]() { stackFree(argWriteBuffer); });
 
 		for (const CharType* it = arguments; !isEnd; lastChar = currentChar, ++it)
 		{
diff --git a/Engine/Source/Programs/UnrealBuildAccelerator/Detours/Private/Windows/UbaDetoursFunctionsNtDll.inl b/Engine/Source/Programs/UnrealBuildAccelerator/Detours/Private/Windows/UbaDetoursFunctionsNtDll.inl
--- a/Engine/Source/Programs/UnrealBuildAccelerator/Detours/Private/Windows/UbaDetoursFunctionsNtDll.inl
+++ b/Engine/Source/Programs/UnrealBuildAccelerator/Detours/Private/Windows/UbaDetoursFunctionsNtDll.inl
@@ -1231,8 +1231,16 @@ NTSTATUS NTAPI Shared_NtCreateFile(bool IsCreateFunc, PHANDLE hFileHandle, ACCES
 		return STATUS_SUCCESS;
 	}
 
-	auto TakeFileOwnershipFromSession = [&](FileMappingHandle mappingHandle, u64 mappingHandleSize)
+	if (lpFileName[0] == ':') // It is a HANDLE from session process. A written file that is writable
+	{
+		if (!info.memoryFile)
 		{
+			const wchar_t* handleStr = lpFileName + 1;
+			UBA_ASSERT(!wcschr(handleStr, '-'));
+			auto mappingHandle = FileMappingHandle::FromU64(StringToValueBase62(handleStr, wcslen(handleStr)));
+			UBA_ASSERT(info.size != InvalidValue);
+			u64 mappingHandleSize = info.size;
+
 			info.memoryFile = new MemoryFile(nullptr, false);
 			MemoryFile& mf = *info.memoryFile;
 			HANDLE newHandle = 0;
@@ -1246,18 +1254,6 @@ NTSTATUS NTAPI Shared_NtCreateFile(bool IsCreateFunc, PHANDLE hFileHandle, ACCES
 			UBA_ASSERTF(mf.baseAddress, L"MapViewOfFile failed when opening temp file %ls (%u)", fileName.data, GetLastError());
 			mf.reserveSize = FileTypeMaxSize(fileName, isSystemOrTempFile);
 			mf.mappingHandle.mh = newHandle;
-		};
-
-	if (lpFileName[0] == ':') // It is a HANDLE from session process. A written file that is writable
-	{
-		if (!info.memoryFile)
-		{
-			const wchar_t* handleStr = lpFileName + 1;
-			UBA_ASSERT(!wcschr(handleStr, '-'));
-			auto mappingHandle = FileMappingHandle::FromU64(StringToValueBase62(handleStr, wcslen(handleStr)));
-			UBA_ASSERT(info.size != InvalidValue);
-
-			TakeFileOwnershipFromSession(mappingHandle, info.size);
 		}
 
 		if (CreateDisposition != FILE_OPEN && CreateDisposition != FILE_OPEN_IF)
diff --git a/Engine/Source/Programs/UnrealBuildAccelerator/ObjTool/Private/UbaObjTool.cpp b/Engine/Source/Programs/UnrealBuildAccelerator/ObjTool/Private/UbaObjTool.cpp
--- a/Engine/Source/Programs/UnrealBuildAccelerator/ObjTool/Private/UbaObjTool.cpp
+++ b/Engine/Source/Programs/UnrealBuildAccelerator/ObjTool/Private/UbaObjTool.cpp
@@ -252,7 +252,11 @@ namespace uba
 			CriticalSection cs;
 			Atomic<bool> success = true;
 
-			AllExternalImports& allExternalImports = *new AllExternalImports(16*1024); // Imports needed from the outside of the stripped obj files. Let it leak to speedup shutdown
+			AllExternalImports& allExternalImports = *new AllExternalImports(16*1024); // Imports needed from the outside of the stripped obj files.
+
+			#if defined(__clang_analyzer__)
+			auto sg = MakeGuard([&]() { delete &allExternalImports; }); // Let it leak to speedup shutdown
+			#endif
 
 			u32 workerCount = DefaultProcessorCount-1;
 			WorkManagerImpl workManager(workerCount, TC("UbaWrk/StrpObj"));
@@ -275,6 +279,10 @@ namespace uba
 			AllInternalImports& allInternalImports = *new AllInternalImports(512*1024); // Imports that the obj files has. Note these could be existing in the obj files and then we might need to create loopbacks
 			AllExports& allExports = *new AllExports(12*1024*1024); // Exports from all the obj files.
 
+			#if defined(__clang_analyzer__)
+			auto sg2 = MakeGuard([&]() { delete &allInternalImports; delete &allExports; }); // Let then leak to speedup shutdown
+			#endif
+
 			Map<TString, ObjectFile*> objectFiles;
 			auto g = MakeGuard([&]() { for (auto& kv : objectFiles) delete kv.second; });
 
diff --git a/Engine/Source/Programs/UnrealBuildAccelerator/Test/Private/UbaTest.h b/Engine/Source/Programs/UnrealBuildAccelerator/Test/Private/UbaTest.h
--- a/Engine/Source/Programs/UnrealBuildAccelerator/Test/Private/UbaTest.h
+++ b/Engine/Source/Programs/UnrealBuildAccelerator/Test/Private/UbaTest.h
@@ -11,6 +11,7 @@
 
 namespace uba
 {
+	class Logger;
 	class LoggerWithWriter;
 	class ProcessHandle;
 	class SessionServer;
diff --git a/Engine/Source/Programs/UnrealBuildAccelerator/Test/Private/UbaTestBasics.cpp b/Engine/Source/Programs/UnrealBuildAccelerator/Test/Private/UbaTestBasics.cpp
--- a/Engine/Source/Programs/UnrealBuildAccelerator/Test/Private/UbaTestBasics.cpp
+++ b/Engine/Source/Programs/UnrealBuildAccelerator/Test/Private/UbaTestBasics.cpp
@@ -326,11 +326,13 @@ namespace uba
 		FileAccessor fileHandle(logger, testFileName.data);
 		if (!fileHandle.CreateWrite(false, DefaultAttributes() | FILE_FLAG_OVERLAPPED, left))
 			return logger.Error(TC("Failed to create file for write"));
-		u8 buffer[512*1024];
+
+		constexpr u64 oddSize = 277872ull;
+		u8 buffer[oddSize];
 
 		while (left)
 		{
-			u64 toWrite = Min(left, 277872ull);
+			u64 toWrite = Min(left, oddSize);
 			if (!fileHandle.Write(buffer, toWrite))
 				return logger.Error(TC("Failed to create file for write"));
 			left -= toWrite;
diff --git a/Engine/Source/Programs/UnrealBuildAccelerator/Visualizer/Private/UbaVisualizer.cpp b/Engine/Source/Programs/UnrealBuildAccelerator/Visualizer/Private/UbaVisualizer.cpp
--- a/Engine/Source/Programs/UnrealBuildAccelerator/Visualizer/Private/UbaVisualizer.cpp
+++ b/Engine/Source/Programs/UnrealBuildAccelerator/Visualizer/Private/UbaVisualizer.cpp
@@ -1803,7 +1803,7 @@ namespace uba
 							}
 							else
 							{
-								stop = playTime;
+								//stop = playTime;
 								brush = m_processBrushes[0].inProgress;
 							}
 
