[UBA]
* Fixed (extremely unlikely) race condition. I don't think this will ever happen in reality but if it does it will potentially be really hard to find. It would happen if there are _lots_ of files written by a child processes and parent process has multiple threads that at the same time call updates-table rpc. Since the session can have a stored linked list of rpc message responses produced by GetWrittenFiles rpc it is important that all those messages end up in the same query as the first message sent from detoured process

diff --git a/Engine/Source/Programs/UnrealBuildAccelerator/Detours/Private/UbaDetoursFileMappingTable.cpp b/Engine/Source/Programs/UnrealBuildAccelerator/Detours/Private/UbaDetoursFileMappingTable.cpp
--- a/Engine/Source/Programs/UnrealBuildAccelerator/Detours/Private/UbaDetoursFileMappingTable.cpp
+++ b/Engine/Source/Programs/UnrealBuildAccelerator/Detours/Private/UbaDetoursFileMappingTable.cpp
@@ -220,13 +220,26 @@ namespace uba
 		return overflow == 0;
 	}
 
+	void Rpc_GetWrittenFilesNoLock()
+	{
+		while (true)
+		{
+			RPC_MESSAGE_NO_LOCK(GetWrittenFiles, updateTables)
+			writer.Flush();
+			BinaryReader reader;
+			if (UpdateWrittenFiles(reader))
+				break;
+		}
+	}
+
 	void Rpc_UpdateTables()
 	{
+		SCOPED_WRITE_LOCK(g_communicationLock, pcs);
 		u32 directoryTableSize;
 		u32 fileMappingTableSize;
 		bool done;
 		{
-			RPC_MESSAGE(UpdateTables, updateTables)
+			RPC_MESSAGE_NO_LOCK(UpdateTables, updateTables)
 			writer.Flush();
 			BinaryReader reader;
 			directoryTableSize = reader.ReadU32();
@@ -236,7 +249,9 @@ namespace uba
 			DEBUG_LOG_PIPE(L"UpdateTables", L"");
 		}
 		if (!done)
-			Rpc_GetWrittenFiles();
+			Rpc_GetWrittenFilesNoLock();
+
+		pcs.Leave();
 
 		g_directoryTable.ParseDirectoryTable(directoryTableSize);
 		g_mappedFileTable.Parse(fileMappingTableSize);
@@ -244,14 +259,8 @@ namespace uba
 
 	void Rpc_GetWrittenFiles()
 	{
-		while (true)
-		{
-			RPC_MESSAGE(GetWrittenFiles, updateTables)
-			writer.Flush();
-			BinaryReader reader;
-			if (UpdateWrittenFiles(reader))
-				break;
-		}
+		SCOPED_WRITE_LOCK(g_communicationLock, pcs);
+		Rpc_GetWrittenFilesNoLock();
 	}
 
 	u32 Rpc_GetEntryOffset(const StringKey& entryNameKey, StringView entryName, bool checkIfDir)
diff --git a/Engine/Source/Programs/UnrealBuildAccelerator/Detours/Private/UbaDetoursShared.h b/Engine/Source/Programs/UnrealBuildAccelerator/Detours/Private/UbaDetoursShared.h
--- a/Engine/Source/Programs/UnrealBuildAccelerator/Detours/Private/UbaDetoursShared.h
+++ b/Engine/Source/Programs/UnrealBuildAccelerator/Detours/Private/UbaDetoursShared.h
@@ -142,4 +142,10 @@ namespace uba
 		SCOPED_WRITE_LOCK(g_communicationLock, pcs); \
 		BinaryWriter writer; \
 		writer.WriteByte(MessageType_##messageName);
+
+	#define RPC_MESSAGE_NO_LOCK(messageName, timerName) \
+		DEBUG_LOG(TC("RPC_MESSAGE %s"), TC(#messageName)); \
+		TimerScope ts(g_stats.timerName); \
+		BinaryWriter writer; \
+		writer.WriteByte(MessageType_##messageName);
 }
