[UBA]
* Fixed bug in UbaCli where -checkcas, -checkpathtable etc crashed because they were using wrong work manager.
* Promoted Flush from WorkManagerImpl to WorkManager

diff --git a/Engine/Source/Programs/UnrealBuildAccelerator/Cli/Private/UbaCli.cpp b/Engine/Source/Programs/UnrealBuildAccelerator/Cli/Private/UbaCli.cpp
--- a/Engine/Source/Programs/UnrealBuildAccelerator/Cli/Private/UbaCli.cpp
+++ b/Engine/Source/Programs/UnrealBuildAccelerator/Cli/Private/UbaCli.cpp
@@ -546,7 +546,7 @@ namespace uba
 		bool exit = false;
 		if (deleteCas)
 		{
-			WorkManagerImpl workManager(0);
+			WorkManagerImpl workManager(maxProcessCount, TC("UbaWrk/DelCas"));
 			StorageImpl(StorageCreateInfo(g_rootDir.data, logWriter, workManager)).DeleteAllCas();
 			for (u32 i=0; i!=agentCount; ++i)
 			{
@@ -559,7 +559,7 @@ namespace uba
 
 		if (!addCas.empty())
 		{
-			WorkManagerImpl workManager(0);
+			WorkManagerImpl workManager(maxProcessCount);
 			StorageCreateInfo storageInfo(g_rootDir.data, logWriter, workManager);
 			storageInfo.casCapacityBytes = 0;
 			storageInfo.storeCompressed = storeCompressed;
@@ -572,26 +572,26 @@ namespace uba
 
 		if (checkCas)
 		{
-			WorkManagerImpl workManager(0);
+			WorkManagerImpl workManager(maxProcessCount, TC("UbaWrk/ChkCasC"));
 			StorageCreateInfo storageInfo(g_rootDir.data, logWriter, workManager);
 			storageInfo.casCapacityBytes = 0;
 			storageInfo.storeCompressed = storeCompressed;
 			StorageImpl storage(storageInfo);
-			if (!storage.CheckCasContent(DefaultProcessorCount))
+			if (!storage.CheckCasContent())
 				return false;
 			exit = true;
 		}
 
 		if (!checkFileTable.empty())
 		{
-			WorkManagerImpl workManager(0);
+			WorkManagerImpl workManager(maxProcessCount, TC("UbaWrk/ChkFTbl"));
 			StorageCreateInfo storageInfo(g_rootDir.data, logWriter, workManager);
 			storageInfo.casCapacityBytes = 0;
 			storageInfo.storeCompressed = storeCompressed;
 			StorageImpl storage(storageInfo);
 			if (!storage.LoadCasTable())
 				return false;
-			if (!storage.CheckFileTable(checkFileTable.data(), DefaultProcessorCount))
+			if (!storage.CheckFileTable(checkFileTable.data()))
 				return false;
 			exit = true;
 		}
@@ -621,7 +621,7 @@ namespace uba
 			if (!client.Init(cii))
 				return false;
 			bool success = true;
-			WorkManagerImpl workManager(DefaultProcessorCount, TC("UbaWrk/ChkCas2"));
+			WorkManagerImpl workManager(maxProcessCount, TC("UbaWrk/ChkCas2"));
 			storageServer.TraverseAllCasFiles([&](const CasKey& casKey, u64 size)
 				{
 					workManager.AddWork([&, casKey](const WorkContext&)
@@ -673,7 +673,7 @@ namespace uba
 		u64 testCompressOriginalSize = 0;
 		if (!testCompress.empty())
 		{
-			WorkManagerImpl workManager(DefaultProcessorCount, TC("UbaWrk/TstComp"));
+			WorkManagerImpl workManager(maxProcessCount, TC("UbaWrk/TstComp"));
 
 			FileAccessor fa(logger, testCompress.c_str());
 			if (!fa.OpenMemoryRead())
@@ -703,7 +703,7 @@ namespace uba
 
 		if (!testDecompress.empty())
 		{
-			WorkManagerImpl workManager(DefaultProcessorCount, TC("UbaWrk/TstDecm"));
+			WorkManagerImpl workManager(maxProcessCount, TC("UbaWrk/TstDecm"));
 
 			FileAccessor fa(logger, testDecompress.c_str());
 			if (!fa.OpenMemoryRead())
diff --git a/Engine/Source/Programs/UnrealBuildAccelerator/Common/Private/UbaNetworkServer.cpp b/Engine/Source/Programs/UnrealBuildAccelerator/Common/Private/UbaNetworkServer.cpp
--- a/Engine/Source/Programs/UnrealBuildAccelerator/Common/Private/UbaNetworkServer.cpp
+++ b/Engine/Source/Programs/UnrealBuildAccelerator/Common/Private/UbaNetworkServer.cpp
@@ -983,6 +983,11 @@ namespace uba
 		return m_maxWorkerCount;
 	}
 
+	bool NetworkServer::FlushWork(u32 timeoutMs)
+	{
+		return m_logger.Error(TC("NetworkServer::FlushWork not implemented"));
+	}
+
 	MutableLogger& NetworkServer::GetLogger()
 	{
 		return m_logger;
diff --git a/Engine/Source/Programs/UnrealBuildAccelerator/Common/Private/UbaStorage.cpp b/Engine/Source/Programs/UnrealBuildAccelerator/Common/Private/UbaStorage.cpp
--- a/Engine/Source/Programs/UnrealBuildAccelerator/Common/Private/UbaStorage.cpp
+++ b/Engine/Source/Programs/UnrealBuildAccelerator/Common/Private/UbaStorage.cpp
@@ -1222,7 +1222,7 @@ namespace uba
 		return true;
 	}
 
-	bool StorageImpl::CheckCasContent(u32 workerCount)
+	bool StorageImpl::CheckCasContent()
 	{
 		StringBuffer<> casRoot;
 		casRoot.Append(m_rootDir.data, m_rootDir.count - 1);
@@ -1243,7 +1243,6 @@ namespace uba
 			};
 
 		m_logger.Info(TC("Traverse cas database..."));
-		WorkManagerImpl workManager(workerCount, TC("UbaWrk/ChkCasC"));
 		u32 entryCount = 0;
 		u32 errorCount = 0;
 		u64 newestWrittenError = 0;
@@ -1251,7 +1250,7 @@ namespace uba
 		TraverseAllCasFiles(casRoot.data, [&](const StringBufferBase& fullPath, const DirectoryEntry& e)
 			{
 				++entryCount;
-				workManager.AddWork([&, filePath = fullPath.ToString(), name = TString(e.name), lastWritten = e.lastWritten](const WorkContext&)
+				m_workManager.AddWork([&, filePath = fullPath.ToString(), name = TString(e.name), lastWritten = e.lastWritten](const WorkContext&)
 					{
 						StringBuffer<> timeStr;
 						writeTimeAgo(timeStr, lastWritten);
@@ -1303,7 +1302,7 @@ namespace uba
 			});
 		m_logger.Info(TC("Validating %u entries..."), entryCount);
 
-		workManager.FlushWork();
+		m_workManager.FlushWork();
 
 		StringBuffer<> newestLastWrittenStr;
 		writeTimeAgo(newestLastWrittenStr, newestWrittenError);
@@ -1315,10 +1314,9 @@ namespace uba
 		return true;
 	}
 
-	bool StorageImpl::CheckFileTable(const tchar* searchPath, u32 workerCount)
+	bool StorageImpl::CheckFileTable(const tchar* searchPath)
 	{
 		m_logger.Info(TC("Searching %s to check files against file table..."), searchPath);
-		WorkManagerImpl workManager(workerCount, TC("UbaWrk/ChkFTbl"));
 
 		List<TString> directories;
 		directories.push_back(searchPath);
@@ -1351,7 +1349,7 @@ namespace uba
 					u64 lastWritten = e.lastWritten;
 					u64 size = e.size;
 
-					workManager.AddWork([&, p = path.ToString(), lastWritten, size](const WorkContext&)
+					m_workManager.AddWork([&, p = path.ToString(), lastWritten, size](const WorkContext&)
 						{
 							++foundFiles;
 							StringKey key = ToStringKey(p.data(), p.size());
@@ -1388,7 +1386,7 @@ namespace uba
 						}, 1, TC("CheckFileTable"));
 				});
 		}
-		workManager.FlushWork();
+		m_workManager.FlushWork();
 
 		m_logger.Info(TC("Done. %u errors found. Searched %u files where %u was tracked and %u matched table."), errorCount.load(), foundFiles.load(), trackedFiles.load(), testedFiles.load(), TimeToText(GetTime() - startTime).str);
 
@@ -1432,14 +1430,13 @@ namespace uba
 	{
 		u32 deleteCount = 0;
 
-		WorkManagerImpl workManager(GetLogicalProcessorCount(), TC("UbaWrk/DelCas"));
 		{
 			Atomic<u32> atomicDeleteCount;
 			TraverseDir(m_logger, m_rootDir, [&](const DirectoryEntry& e)
 				{
 					if (!IsDirectory(e.attributes))
 						return;
-					workManager.AddWork([&, name = TString(e.name)](const WorkContext&)
+					m_workManager.AddWork([&, name = TString(e.name)](const WorkContext&)
 						{
 							StringBuffer<> fullPath;
 							fullPath.Append(m_rootDir).Append(name);
@@ -1448,7 +1445,7 @@ namespace uba
 							atomicDeleteCount += deleteCountTemp;
 						}, 1, TC("DeleteAllCas"));
 				});
-			workManager.FlushWork();
+			m_workManager.FlushWork();
 			deleteCount += atomicDeleteCount;
 		}
 
diff --git a/Engine/Source/Programs/UnrealBuildAccelerator/Common/Public/UbaNetworkServer.h b/Engine/Source/Programs/UnrealBuildAccelerator/Common/Public/UbaNetworkServer.h
--- a/Engine/Source/Programs/UnrealBuildAccelerator/Common/Public/UbaNetworkServer.h
+++ b/Engine/Source/Programs/UnrealBuildAccelerator/Common/Public/UbaNetworkServer.h
@@ -85,6 +85,7 @@ namespace uba
 		virtual void AddWork(const WorkFunction& work, u32 count, const tchar* desc, const Color& color = ColorWork, bool highPriority = false) override final;
 		virtual void DoWork(u32 count = 1) override final;
 		virtual u32 GetWorkerCount() override final;
+		virtual bool FlushWork(u32 timeoutMs = 0) override final;
 
 		struct ClientStats
 		{
diff --git a/Engine/Source/Programs/UnrealBuildAccelerator/Common/Public/UbaStorage.h b/Engine/Source/Programs/UnrealBuildAccelerator/Common/Public/UbaStorage.h
--- a/Engine/Source/Programs/UnrealBuildAccelerator/Common/Public/UbaStorage.h
+++ b/Engine/Source/Programs/UnrealBuildAccelerator/Common/Public/UbaStorage.h
@@ -136,8 +136,8 @@ namespace uba
 		virtual ~StorageImpl();
 
 		bool LoadCasTable(bool logStats = true, bool alwaysCheckAllFiles = false, bool* outWasTerminated = nullptr);
-		bool CheckCasContent(u32 workerCount);
-		bool CheckFileTable(const tchar* searchPath, u32 workerCount);
+		bool CheckCasContent();
+		bool CheckFileTable(const tchar* searchPath);
 		const tchar* GetTempPath();
 
 		virtual bool SaveCasTable(bool deleteIsRunningfile, bool deleteDropped = true) override;
diff --git a/Engine/Source/Programs/UnrealBuildAccelerator/Common/Public/UbaWorkManager.h b/Engine/Source/Programs/UnrealBuildAccelerator/Common/Public/UbaWorkManager.h
--- a/Engine/Source/Programs/UnrealBuildAccelerator/Common/Public/UbaWorkManager.h
+++ b/Engine/Source/Programs/UnrealBuildAccelerator/Common/Public/UbaWorkManager.h
@@ -38,6 +38,7 @@ namespace uba
 		virtual void AddWork(const WorkFunction& work, u32 count, const tchar* desc, const Color& color = ColorWork, bool highPriority = false) = 0;
 		virtual u32 GetWorkerCount() = 0;
 		virtual void DoWork(u32 count = 1) = 0;
+		virtual bool FlushWork(u32 timeoutMs = 0) = 0;
 
 		u32 TrackWorkStart(const StringView& desc, const Color& color);
 		void TrackWorkHint(u32 id, const StringView& hint, u64 startTime = 0);
@@ -59,10 +60,10 @@ namespace uba
 	public:
 		WorkManagerImpl(u32 workerCount, const tchar* workerDesc = TC("UbaWrk"));
 		virtual ~WorkManagerImpl();
-		virtual void AddWork(const WorkFunction& work, u32 count, const tchar* desc, const Color& color = ColorWork, bool highPriority = false) override;
-		virtual u32 GetWorkerCount() override;
-		virtual void DoWork(u32 count = 1) override;
-		bool FlushWork(u32 timeoutMs = 0);
+		virtual void AddWork(const WorkFunction& work, u32 count, const tchar* desc, const Color& color = ColorWork, bool highPriority = false) override final;
+		virtual u32 GetWorkerCount() override final;
+		virtual void DoWork(u32 count = 1) override final;
+		virtual bool FlushWork(u32 timeoutMs = 0) override final;
 
 	private:
 		struct Worker;
