[UBA]
* Fixed so file id etc are kept for files stored in memory but are not modified
* Fixed bug where if there is a hierarchy of three processes and top ancestor opens a file for write, close it, and then spawns middle process which also opens same file for write and then spawn its child. When coming back from the inner child the middle process updates written files and since owner is still root ancestor (owner is set on close) it overwrote the actively written file in middle process and asserted and failed. Reproable when doing incremental linking which falls back to full linking

diff --git a/Engine/Source/Programs/UnrealBuildAccelerator/Common/Private/UbaProcess.cpp b/Engine/Source/Programs/UnrealBuildAccelerator/Common/Private/UbaProcess.cpp
--- a/Engine/Source/Programs/UnrealBuildAccelerator/Common/Private/UbaProcess.cpp
+++ b/Engine/Source/Programs/UnrealBuildAccelerator/Common/Private/UbaProcess.cpp
@@ -1004,9 +1004,10 @@ namespace uba
 	bool ProcessImpl::HandleUpdateTables(BinaryReader& reader, BinaryWriter& writer)
 	{
 		UBA_ASSERT(!m_shared.firstOverflowMessage);
+		bool isInit = reader.ReadBool();
 		writer.WriteU32(m_session.GetDirectoryTableSize());
 		writer.WriteU32(m_session.GetFileMappingSize());
-		return HandleGetWrittenFiles(reader, writer);
+		return HandleGetWrittenFiles(isInit, writer);
 	}
 
 	struct ProcessImpl::OverflowedWrittenFilesMessage
@@ -1017,6 +1018,12 @@ namespace uba
 	};
 
 	bool ProcessImpl::HandleGetWrittenFiles(BinaryReader& reader, BinaryWriter& writer)
+	{
+		bool isInit = reader.ReadBool();
+		return HandleGetWrittenFiles(isInit, writer);
+	}
+
+	bool ProcessImpl::HandleGetWrittenFiles(bool isInit, BinaryWriter& writer)
 	{
 		if (auto msg = m_shared.firstOverflowMessage)
 		{
@@ -1040,8 +1047,20 @@ namespace uba
 		for (auto& kv : m_shared.writtenFiles)
 		{
 			WrittenFile& wf = kv.second;
-			if (wf.owner == this)
-				continue;
+
+			UBA_ASSERT(wf.owner);
+
+			if (!isInit)
+			{
+				// If file is owned by ancestor we also skip re-fetching it.. since we expect ancestor
+				// not to write to file while child process exists. if this changes we need to change approach
+				// and track write file stamp or something on the file and track what child processes have
+				bool ownedByThisOrAncestor = false;
+				for (ProcessImpl* it = this; it && !ownedByThisOrAncestor; it=it->m_parentProcess)
+					ownedByThisOrAncestor = wf.owner == it;
+				if (ownedByThisOrAncestor)
+					continue;
+			}
 
 			u64 fileSize;
 			if (wf.mappingHandle.IsValid())
@@ -1438,33 +1457,6 @@ namespace uba
 		return true;
 	}
 
-	bool ProcessImpl::HandleTakeFileOwnership(BinaryReader& reader, BinaryWriter& writer)
-	{
-		StringKey fileNameKey = reader.ReadStringKey();
-		StringBuffer<> fileName;
-		reader.ReadString(fileName);
-
-		bool isOwner = false;
-		FileMappingHandle mappingHandle;
-		u64 mappingWritten = 0;
-
-		SCOPED_FUTEX(m_shared.writtenFilesLock, lock);
-		auto findIt = m_shared.writtenFiles.find(fileNameKey);
-		if (findIt != m_shared.writtenFiles.end())
-		{
-			WrittenFile& file = findIt->second;
-			isOwner = file.owner == this;
-			mappingHandle = file.mappingHandle;
-			mappingWritten = file.mappingWritten;
-			file.owner = this;
-			file.originalMappingHandle = {};
-		}
-		writer.WriteBool(isOwner);
-		writer.WriteU64(mappingHandle.ToU64());
-		writer.WriteU64(mappingWritten);
-		return true;
-	}
-
 	bool ProcessImpl::HandleRunSpecialProgram(BinaryReader& reader, BinaryWriter& writer)
 	{
 		m_session.RunSpecialProgram(*this, reader, writer);
diff --git a/Engine/Source/Programs/UnrealBuildAccelerator/Common/Public/UbaProcess.h b/Engine/Source/Programs/UnrealBuildAccelerator/Common/Public/UbaProcess.h
--- a/Engine/Source/Programs/UnrealBuildAccelerator/Common/Public/UbaProcess.h
+++ b/Engine/Source/Programs/UnrealBuildAccelerator/Common/Public/UbaProcess.h
@@ -79,6 +79,7 @@ namespace uba
 
 		bool HandleMessage(BinaryReader& reader, BinaryWriter& writer);
 		bool HandleSpecialApplication();
+		bool HandleGetWrittenFiles(bool isInit, BinaryWriter& writer);
 
 		#define UBA_PROCESS_MESSAGE(M) bool Handle##M(BinaryReader& reader, BinaryWriter& writer);
 		UBA_PROCESS_MESSAGES
diff --git a/Engine/Source/Programs/UnrealBuildAccelerator/Core/Public/UbaProtocol.h b/Engine/Source/Programs/UnrealBuildAccelerator/Core/Public/UbaProtocol.h
--- a/Engine/Source/Programs/UnrealBuildAccelerator/Core/Public/UbaProtocol.h
+++ b/Engine/Source/Programs/UnrealBuildAccelerator/Core/Public/UbaProtocol.h
@@ -38,7 +38,6 @@ namespace uba
 		UBA_PROCESS_MESSAGE(HostRun) \
 		UBA_PROCESS_MESSAGE(ResolveCallstack) \
 		UBA_PROCESS_MESSAGE(CheckRemapping) \
-		UBA_PROCESS_MESSAGE(TakeFileOwnership) \
 		UBA_PROCESS_MESSAGE(RunSpecialProgram) \
 
 	enum MessageType : u8
diff --git a/Engine/Source/Programs/UnrealBuildAccelerator/Detours/Private/PosixOS/UbaDetoursFunctionsPosix.cpp b/Engine/Source/Programs/UnrealBuildAccelerator/Detours/Private/PosixOS/UbaDetoursFunctionsPosix.cpp
--- a/Engine/Source/Programs/UnrealBuildAccelerator/Detours/Private/PosixOS/UbaDetoursFunctionsPosix.cpp
+++ b/Engine/Source/Programs/UnrealBuildAccelerator/Detours/Private/PosixOS/UbaDetoursFunctionsPosix.cpp
@@ -605,7 +605,7 @@ void Shared_close(int fd, const TrueClose& trueClose)
 		struct stat attr;
 		TRUE_WRAPPER(fstat)(fd, &attr); // is this safe? not sure about fflush and user space FILE* streams
 		fi.size = attr.st_size;
-		fi.isWritten = true;
+		fi.created = true; // TODO: Should be set earlier
 		closeGuard.Execute();
 		//for (auto& kv : g_fileHandles)
 		//	UBA_ASSERT(&fi != kv.second.fileObject->fileInfo);
@@ -1739,7 +1739,7 @@ UBA_EXPORT int UBA_WRAPPER(rename)(const char* oldpath, const char* newpath)
 			info = findIt->second;
 		info.originalName = g_memoryBlock.Strdup(fixedNewPath).data;
 		info.name = info.originalName;
-		info.isWritten = true;
+		info.created = true;
 	}
 
 	g_directoryTable.ParseDirectoryTable(directoryTableSize);
diff --git a/Engine/Source/Programs/UnrealBuildAccelerator/Detours/Private/UbaDetoursFileMappingTable.cpp b/Engine/Source/Programs/UnrealBuildAccelerator/Detours/Private/UbaDetoursFileMappingTable.cpp
--- a/Engine/Source/Programs/UnrealBuildAccelerator/Detours/Private/UbaDetoursFileMappingTable.cpp
+++ b/Engine/Source/Programs/UnrealBuildAccelerator/Detours/Private/UbaDetoursFileMappingTable.cpp
@@ -189,10 +189,7 @@ namespace uba
 			u64 fileSize = reader.Read7BitEncoded();
 			info.fileNameKey = key;
 			info.size = fileSize;
-
-			#if !PLATFORM_WINDOWS
-			info.isWritten = true;
-			#endif
+			info.created = true;
 
 			if (mappingHandle.IsValid())
 				backedName.Clear().Append(':').AppendBase62(mappingHandle.ToU64());
@@ -200,8 +197,11 @@ namespace uba
 			if (!info.name || !backedName.Equals(info.name))
 				info.name = g_mappedFileTable.m_memoryBlock.Strdup(backedName).data;
 
+			DEBUG_LOG(TC("GOT WRITTEN FILE: %s (BackedFile: %s Size: %llu)"), info.originalName, info.name, info.size);
+
 			if (auto mf = info.memoryFile) // Memory file is now wrong.. we need to delete it (or should we create a new one with the new mapping?)
 			{
+				UBA_ASSERTF(!info.refCount, TC("File has memory file and refcount. This is not supported. %s (BackedFile: %s Size: %llu)"), info.originalName, info.name, info.size);
 				#if PLATFORM_WINDOWS
 				if (!mf->isLocalOnly)
 				{
@@ -214,17 +214,16 @@ namespace uba
 				// delete mf; // Let them leak
 				info.memoryFile = nullptr;
 			}
-
-			DEBUG_LOG(TC("GOT WRITTEN FILE: %s (BackedFile: %s Size: %llu)"), info.originalName, info.name, info.size);
 		}
 		return overflow == 0;
 	}
 
-	void Rpc_GetWrittenFilesNoLock()
+	void Rpc_GetWrittenFilesNoLock(bool isInit)
 	{
 		while (true)
 		{
 			RPC_MESSAGE_NO_LOCK(GetWrittenFiles, updateTables)
+			writer.WriteBool(isInit);
 			writer.Flush();
 			BinaryReader reader;
 			if (UpdateWrittenFiles(reader))
@@ -240,6 +239,7 @@ namespace uba
 		bool done;
 		{
 			RPC_MESSAGE_NO_LOCK(UpdateTables, updateTables)
+			writer.WriteBool(false);
 			writer.Flush();
 			BinaryReader reader;
 			directoryTableSize = reader.ReadU32();
@@ -249,7 +249,7 @@ namespace uba
 			DEBUG_LOG_PIPE(L"UpdateTables", L"");
 		}
 		if (!done)
-			Rpc_GetWrittenFilesNoLock();
+			Rpc_GetWrittenFilesNoLock(false);
 
 		pcs.Leave();
 
@@ -260,7 +260,7 @@ namespace uba
 	void Rpc_GetWrittenFiles()
 	{
 		SCOPED_WRITE_LOCK(g_communicationLock, pcs);
-		Rpc_GetWrittenFilesNoLock();
+		Rpc_GetWrittenFilesNoLock(true);
 	}
 
 	u32 Rpc_GetEntryOffset(const StringKey& entryNameKey, StringView entryName, bool checkIfDir)
diff --git a/Engine/Source/Programs/UnrealBuildAccelerator/Detours/Private/UbaDetoursFileMappingTable.h b/Engine/Source/Programs/UnrealBuildAccelerator/Detours/Private/UbaDetoursFileMappingTable.h
--- a/Engine/Source/Programs/UnrealBuildAccelerator/Detours/Private/UbaDetoursFileMappingTable.h
+++ b/Engine/Source/Programs/UnrealBuildAccelerator/Detours/Private/UbaDetoursFileMappingTable.h
@@ -45,6 +45,9 @@ namespace uba
 		// Is set to true once file has been tracked to prevent tracking same file multiple times
 		bool tracked = false;
 
+		// If set to true, this means file was created during process group lifetime (which means it has its own file id etc)
+		bool created = false;
+
 		// File mapping handle coming from session process (Set in NtCreateFile)
 		void* trueFileMapHandle = 0;
 
@@ -62,8 +65,6 @@ namespace uba
 		u32 fileMapDesiredAccess = PAGE_READONLY;
 		u32 fileMapViewDesiredAccess = FILE_MAP_READ; // This is the parameter to MapViewOfFile
 		bool mappingChecked = false;
-		#else
-		bool isWritten = false;
 		#endif
 	};
 
diff --git a/Engine/Source/Programs/UnrealBuildAccelerator/Detours/Private/UbaDetoursShared.cpp b/Engine/Source/Programs/UnrealBuildAccelerator/Detours/Private/UbaDetoursShared.cpp
--- a/Engine/Source/Programs/UnrealBuildAccelerator/Detours/Private/UbaDetoursShared.cpp
+++ b/Engine/Source/Programs/UnrealBuildAccelerator/Detours/Private/UbaDetoursShared.cpp
@@ -463,6 +463,9 @@ namespace uba
 
 		memset(&outAttr.data, 0, sizeof(outAttr.data));
 
+		bool foundMapping = false;
+		u64 fileSize = InvalidValue;
+
 		#if PLATFORM_WINDOWS
 		if (fileName[1] == ':' && fileName[3] == 0 && (ToLower(fileName[0]) == ToLower(g_virtualWorkingDir[0]) || ToLower(fileName[0]) == g_systemRoot[0]))
 		{
@@ -501,16 +504,21 @@ namespace uba
 				//outAttr.fileIndex = 
 
 #if PLATFORM_WINDOWS
+				foundMapping = true;
+				fileSize = fi.size;
 				outAttr.useCache = true;
 				outAttr.exists = true;
 				outAttr.lastError = ErrorSuccess;
-				LARGE_INTEGER li = ToLargeInteger(fi.size);
+				LARGE_INTEGER li = ToLargeInteger(fileSize);
 				outAttr.data.nFileSizeLow = li.LowPart;
 				outAttr.data.nFileSizeHigh = li.HighPart;
 				outAttr.data.dwFileAttributes = FILE_ATTRIBUTE_NORMAL;
-				return fileName.data;
+				if (fi.created)
+					return fileName.data;
+				// We can't exit here because we need fileIndex and volume serial among other things.
+				// TODO: make sure fi.created is correct on windows everywhere. CopyFile/MoveFile should probably set this in certain conditions
 #else
-				if (fi.isWritten)
+				if (fi.created)
 				{
 					outAttr.useCache = true;
 					outAttr.exists = true;
@@ -525,6 +533,8 @@ namespace uba
 
 		if (!g_allowDirectoryCache || fileName.StartsWith(g_systemTemp.data)) // We need to skip SystemTemp.. lots of stuff going on there.
 		{
+			if (foundMapping)
+				return fileName.data;
 			outAttr.useCache = false;
 			return fileName.data;
 		}
@@ -536,6 +546,8 @@ namespace uba
 
 		if (dirTableOffset == ~u32(0))
 		{
+			if (foundMapping)
+				return fileName.data;
 			outAttr.useCache = true;
 			outAttr.exists = false;
 			outAttr.lastError = ErrorFileNotFound;
@@ -547,6 +559,8 @@ namespace uba
 
 		if (!info.attributes)
 		{
+			if (foundMapping)
+				return fileName.data;
 			// File used to exist but was deleted
 			outAttr.useCache = true;
 			outAttr.exists = false;
@@ -554,7 +568,8 @@ namespace uba
 			return fileName.data;
 		}
 
-		u64 fileSize = info.size;
+		if (fileSize == InvalidValue)
+			fileSize = info.size;
 
 		// Could be compressed and then directory table size is wrong
 		if (CouldBeCompressedFile(fileName))
diff --git a/Engine/Source/Programs/UnrealBuildAccelerator/Detours/Private/Windows/UbaDetoursFunctionsKernelBase.inl b/Engine/Source/Programs/UnrealBuildAccelerator/Detours/Private/Windows/UbaDetoursFunctionsKernelBase.inl
--- a/Engine/Source/Programs/UnrealBuildAccelerator/Detours/Private/Windows/UbaDetoursFunctionsKernelBase.inl
+++ b/Engine/Source/Programs/UnrealBuildAccelerator/Detours/Private/Windows/UbaDetoursFunctionsKernelBase.inl
@@ -869,7 +869,7 @@ DWORD Detoured_SetFilePointer(HANDLE hFile, LONG lDistanceToMove, PLONG lpDistan
 		}
 
 		trueHandle = dh.trueHandle;
-		UBA_ASSERT(trueHandle != INVALID_HANDLE_VALUE);
+		UBA_ASSERTF(trueHandle != INVALID_HANDLE_VALUE, TC("SetFilePointer got detoured handle %llu that was not memory file or filemap (%ls)"), uintptr_t(hFile), HandleToName(hFile));
 	}
 
 	DEBUG_LOG_TRUE(L"SetFilePointer", L"%llu (%ls)", uintptr_t(hFile), HandleToName(hFile));
@@ -918,7 +918,7 @@ DWORD Detoured_SetFilePointerEx(HANDLE hFile, LARGE_INTEGER liDistanceToMove, PL
 			return TRUE;
 		}
 		trueHandle = dh.trueHandle;
-		UBA_ASSERT(trueHandle != INVALID_HANDLE_VALUE);
+		UBA_ASSERTF(trueHandle != INVALID_HANDLE_VALUE, TC("SetFilePointer got detoured handle %llu that was not memory file or filemap (%ls)"), uintptr_t(hFile), HandleToName(hFile));
 	}
 	DEBUG_LOG_TRUE(L"SetFilePointerEx", L"%llu %lli %u (%ls)", uintptr_t(hFile), liDistanceToMove.QuadPart, dwMoveMethod, HandleToName(hFile));
 	return True_SetFilePointerEx(trueHandle, liDistanceToMove, lpNewFilePointer, dwMoveMethod); // This ends up in NtSetInformationFile
diff --git a/Engine/Source/Programs/UnrealBuildAccelerator/Detours/Private/Windows/UbaDetoursFunctionsNtDll.inl b/Engine/Source/Programs/UnrealBuildAccelerator/Detours/Private/Windows/UbaDetoursFunctionsNtDll.inl
--- a/Engine/Source/Programs/UnrealBuildAccelerator/Detours/Private/Windows/UbaDetoursFunctionsNtDll.inl
+++ b/Engine/Source/Programs/UnrealBuildAccelerator/Detours/Private/Windows/UbaDetoursFunctionsNtDll.inl
@@ -1304,15 +1304,17 @@ NTSTATUS NTAPI Shared_NtCreateFile(bool IsCreateFunc, PHANDLE hFileHandle, ACCES
 					return res2;
 				}
 
-				LARGE_INTEGER fileSize2;
-				if (!GetFileSizeEx(fileHandle, &fileSize2))
+				BY_HANDLE_FILE_INFORMATION fileInfo2;
+				if (!GetFileInformationByHandle(fileHandle, &fileInfo2))
 				{
-					UBA_ASSERTF(false, TC("%s"), fileName.data);
+					UBA_ASSERTF(false, TC("GetFileInformationByHandle failed when in NtCreateFile and open file for write (%s)"), fileName.data);
 					return STATUS_OBJECT_NAME_EXISTS;
 				}
+				u64 fileSize2 = ToLargeInteger(fileInfo2.nFileSizeHigh, fileInfo2.nFileSizeLow).QuadPart;
+
 				DetouredHandle tempDh(HandleType_File, INVALID_HANDLE_VALUE);
-				mf->EnsureCommitted(tempDh, fileSize2.QuadPart);
-				u64 left = fileSize2.QuadPart;
+				mf->EnsureCommitted(tempDh, fileSize2);
+				u64 left = fileSize2;
 				u8* writePos = mf->baseAddress;
 				while (left)
 				{
@@ -1328,7 +1330,10 @@ NTSTATUS NTAPI Shared_NtCreateFile(bool IsCreateFunc, PHANDLE hFileHandle, ACCES
 				}
 				NtClose(fileHandle);
 		
-				mf->writtenSize = fileSize2.QuadPart;
+				mf->writtenSize = fileSize2;
+				mf->fileTime = ToLargeInteger(fileInfo2.ftCreationTime.dwHighDateTime, fileInfo2.ftCreationTime.dwLowDateTime).QuadPart;
+				mf->volumeSerial = fileInfo2.dwVolumeSerialNumber;
+				mf->fileIndex = ToLargeInteger(fileInfo2.nFileIndexHigh, fileInfo2.nFileIndexLow).QuadPart;
 			}
 			else
 			{
@@ -1336,17 +1341,18 @@ NTSTATUS NTAPI Shared_NtCreateFile(bool IsCreateFunc, PHANDLE hFileHandle, ACCES
 				{
 					UBA_ASSERTF(false, TC("Trying to open %s with openif. This is not supported"), fileName.data);
 				}
-			}
 
-			// TODO: Time should be in sync with host machine!
-			FILETIME ft;
-			SYSTEMTIME st;
-			GetSystemTime(&st);
-			SystemTimeToFileTime(&st, &ft);
-			mf->fileTime = (u64&)ft;
-			mf->volumeSerial = 1;
-			mf->fileIndex = InterlockedDecrement(&g_memoryFileIndexCounter);
+				// TODO: Time should be in sync with host machine!
+				FILETIME ft;
+				SYSTEMTIME st;
+				GetSystemTime(&st);
+				SystemTimeToFileTime(&st, &ft);
+				mf->fileTime = (u64&)ft;
+				mf->volumeSerial = 1;
+				mf->fileIndex = InterlockedDecrement(&g_memoryFileIndexCounter);
+			}
 
+			info.created = true;
 			info.memoryFile = mf;
 		}
 		else
