[UBA]
* Changed so FileMappingBuffer::MapView opens mapping for read instead of write (don't know why it was write)

diff --git a/Engine/Source/Programs/UnrealBuildAccelerator/Common/Private/UbaFileMapping.cpp b/Engine/Source/Programs/UnrealBuildAccelerator/Common/Private/UbaFileMapping.cpp
--- a/Engine/Source/Programs/UnrealBuildAccelerator/Common/Private/UbaFileMapping.cpp
+++ b/Engine/Source/Programs/UnrealBuildAccelerator/Common/Private/UbaFileMapping.cpp
@@ -530,7 +530,7 @@ namespace uba
 			if (!MapViewCommit(data + commitStart, commitSize))
 			{
 				UnmapViewOfFile(m_logger, data, mapSize, hint);
-				m_logger.Error(TC("%s - Failed to allocate memory for %s (%s)"), f.name, hint, LastErrorToText().data);
+				m_logger.Error(TC("%s - Failed to allocate/map %s for %s (%s)"), f.name, BytesToText(commitSize).str, hint, LastErrorToText().data);
 				return res;
 			}
 
@@ -659,6 +659,16 @@ namespace uba
 	}
 
 	MappedView FileMappingBuffer::MapView(FileMappingHandle handle, u64 offset, u64 size, const tchar* hint, bool isIndependent)
+	{
+		return MapView(handle, offset, size, hint, isIndependent, FILE_MAP_READ);
+	}
+
+	MappedView FileMappingBuffer::MapViewWrite(FileMappingHandle handle, u64 offset, u64 size, const tchar* hint, bool isIndependent)
+	{
+		return MapView(handle, offset, size, hint, isIndependent, FILE_MAP_WRITE);
+	}
+
+	MappedView FileMappingBuffer::MapView(FileMappingHandle handle, u64 offset, u64 size, const tchar* hint, bool isIndependent, u32 desiredAccess)
 	{
 		UBA_ASSERT(handle.IsValid());// && (handle == m_files[0].handle || handle == m_files[1].handle));
 
@@ -666,7 +676,7 @@ namespace uba
 
 		if (isIndependent)
 		{
-			u8* data = MapViewOfFile(m_logger, handle, FILE_MAP_WRITE, offset, size);
+			u8* data = MapViewOfFile(m_logger, handle, desiredAccess, offset, size);
 			if (!data)
 				return {};
 			res.handle = handle;
@@ -692,7 +702,7 @@ namespace uba
 			alignedOffsetStart = AlignUp(offset - (m_pageSize - 1), m_pageSize);
 			u64 newOffset = offset + size;
 			u64 alignedOffsetEnd = AlignUp(newOffset, m_pageSize);
-			data = MapViewOfFile(m_logger, handle, FILE_MAP_WRITE, alignedOffsetStart, alignedOffsetEnd - alignedOffsetStart);
+			data = MapViewOfFile(m_logger, handle, desiredAccess, alignedOffsetStart, alignedOffsetEnd - alignedOffsetStart);
 			if (!data)
 			{
 				m_logger.Error(TC("%s - MapView failed to map view of file for %s with size %llu and offset %llu (%s)"), file.name, hint, size, offset, LastErrorToText().data);
diff --git a/Engine/Source/Programs/UnrealBuildAccelerator/Common/Private/UbaSession.cpp b/Engine/Source/Programs/UnrealBuildAccelerator/Common/Private/UbaSession.cpp
--- a/Engine/Source/Programs/UnrealBuildAccelerator/Common/Private/UbaSession.cpp
+++ b/Engine/Source/Programs/UnrealBuildAccelerator/Common/Private/UbaSession.cpp
@@ -1069,7 +1069,7 @@ namespace uba
 		if (mappedViewRead.isCompressed)
 		{
 			auto mvrg = MakeGuard([&](){ m_storage.UnmapView(mappedViewRead, fileName); });
-			u8* readMemory = mappedViewRead.memory;
+			const u8* readMemory = mappedViewRead.memory;
 			contentSize = *(u64*)readMemory;
 			readMemory += 8;
 
diff --git a/Engine/Source/Programs/UnrealBuildAccelerator/Common/Private/UbaStorage.cpp b/Engine/Source/Programs/UnrealBuildAccelerator/Common/Private/UbaStorage.cpp
--- a/Engine/Source/Programs/UnrealBuildAccelerator/Common/Private/UbaStorage.cpp
+++ b/Engine/Source/Programs/UnrealBuildAccelerator/Common/Private/UbaStorage.cpp
@@ -2057,8 +2057,8 @@ namespace uba
 				FileHandle readHandle = InvalidFileHandle;
 				auto rsg = MakeGuard([&](){ if (readHandle != InvalidFileHandle) CloseFile(nullptr, readHandle); });
 
-				u8* compressedData = nullptr;
-				u8* readData = nullptr;
+				const u8* compressedData = nullptr;
+				const u8* readData = nullptr;
 				MappedView mappedView;
 				auto mapViewGuard = MakeGuard([&](){ m_casDataBuffer.UnmapView(mappedView, destination); });
 
@@ -2473,7 +2473,7 @@ namespace uba
 			return true;
 		}
 		FileMappingScope scope(m_logger);
-		u8* fileMem = scope.MapView(TC("ReadFileSize"), casEntry.mappingHandle, casEntry.mappingOffset, casEntry.mappingSize, hint);
+		u8* fileMem = scope.MapView(TC("ReadFileSize"), casEntry.mappingHandle, casEntry.mappingOffset, 8, hint);
 		if (!fileMem)
 			return false;
 		outSize = *(u64*)fileMem;
diff --git a/Engine/Source/Programs/UnrealBuildAccelerator/Common/Private/UbaStorageServer.cpp b/Engine/Source/Programs/UnrealBuildAccelerator/Common/Private/UbaStorageServer.cpp
--- a/Engine/Source/Programs/UnrealBuildAccelerator/Common/Private/UbaStorageServer.cpp
+++ b/Engine/Source/Programs/UnrealBuildAccelerator/Common/Private/UbaStorageServer.cpp
@@ -158,7 +158,7 @@ namespace uba
 		{
 			if (!memoryBegin)
 				server.m_logger.Warning(TC("This should not happen. It means there is a race between a fetch and a disconnect. Report to honk (%s)"), reason);
-			server.m_bufferSlots.Push(memoryBegin);
+			server.m_bufferSlots.Push((u8*)memoryBegin);
 			memoryBegin = nullptr;
 		}
 
@@ -552,8 +552,8 @@ namespace uba
 							m_logger.Error(TC("Failed to close file %s (%s)"), casFile.data, LastErrorToText().data);
 					});
 				u64 fileSize = 0;
-				u8* memoryBegin = nullptr;
-				u8* memoryPos = nullptr;
+				const u8* memoryBegin = nullptr;
+				const u8* memoryPos = nullptr;
 				bool ownsMapping = false;
 
 				MappedView mappedView;
@@ -650,10 +650,11 @@ namespace uba
 				}
 				else
 				{
-					memoryBegin = m_bufferSlots.Pop();
+					u8* buffer = m_bufferSlots.Pop();;
+					memoryBegin = buffer;
 					memoryPos = memoryBegin;
 					u32 toRead = u32(Min(left, BufferSlotSize));
- 					if (!ReadFile(m_logger, casFile.data, readFileHandle, memoryBegin, toRead))
+ 					if (!ReadFile(m_logger, casFile.data, readFileHandle, buffer, toRead))
 					{
 						UBA_ASSERT(false); // Implement
 						return m_logger.Error(TC("Failed to read file %s (%s) (2)"), casFile.data, LastErrorToText().data);;
diff --git a/Engine/Source/Programs/UnrealBuildAccelerator/Common/Private/UbaStorageUtils.cpp b/Engine/Source/Programs/UnrealBuildAccelerator/Common/Private/UbaStorageUtils.cpp
--- a/Engine/Source/Programs/UnrealBuildAccelerator/Common/Private/UbaStorageUtils.cpp
+++ b/Engine/Source/Programs/UnrealBuildAccelerator/Common/Private/UbaStorageUtils.cpp
@@ -1293,10 +1293,10 @@ namespace uba
 		return true;
 	}
 
-	bool DecompressWriter::DecompressMemoryToFile(u8* compressedData, FileAccessor& destination, u64 decompressedSize, bool useNoBuffering)
+	bool DecompressWriter::DecompressMemoryToFile(const u8* compressedData, FileAccessor& destination, u64 decompressedSize, bool useNoBuffering)
 	{
 		StorageStats& stats = m_stats;
-		u8* readPos = compressedData;
+		const u8* readPos = compressedData;
 
 		u8* slot = m_bufferSlots.Pop();
 		auto _ = MakeGuard([&](){ m_bufferSlots.Push(slot); });
diff --git a/Engine/Source/Programs/UnrealBuildAccelerator/Common/Public/UbaFileMapping.h b/Engine/Source/Programs/UnrealBuildAccelerator/Common/Public/UbaFileMapping.h
--- a/Engine/Source/Programs/UnrealBuildAccelerator/Common/Public/UbaFileMapping.h
+++ b/Engine/Source/Programs/UnrealBuildAccelerator/Common/Public/UbaFileMapping.h
@@ -66,6 +66,7 @@ namespace uba
 
 		MappedView AllocAndMapView(FileMappingType type, u64 size, u64 alignment, const tchar* hint, bool allowShrink = false, bool isIndependent = false);
 		MappedView MapView(FileMappingHandle handle, u64 offset, u64 size, const tchar* hint, bool isIndependent = false);
+		MappedView MapViewWrite(FileMappingHandle handle, u64 offset, u64 size, const tchar* hint, bool isIndependent = false);
 		MappedView MapView(const StringView& str, u64 size, const tchar* hint);
 		void UnmapView(MappedView view, const tchar* hint, u64 newSize = InvalidValue);
 		
@@ -91,6 +92,7 @@ namespace uba
 			bool keepMapped = false;
 		};
 
+		MappedView MapView(FileMappingHandle handle, u64 offset, u64 size, const tchar* hint, bool isIndependent, u32 desiredAccess);
 		MappedView AllocAndMapViewNoLock(File& file, u64 size, u64 alignment, const tchar* hint);
 		File& GetFile(FileMappingHandle handle, u8& outStorageIndex);
 
diff --git a/Engine/Source/Programs/UnrealBuildAccelerator/Common/Public/UbaStorageServer.h b/Engine/Source/Programs/UnrealBuildAccelerator/Common/Public/UbaStorageServer.h
--- a/Engine/Source/Programs/UnrealBuildAccelerator/Common/Public/UbaStorageServer.h
+++ b/Engine/Source/Programs/UnrealBuildAccelerator/Common/Public/UbaStorageServer.h
@@ -97,8 +97,8 @@ namespace uba
 
 			FileHandle readFileHandle = InvalidFileHandle;
 			MappedView mappedView;
-			u8* memoryBegin = nullptr;
-			u8* memoryPos = nullptr;
+			const u8* memoryBegin = nullptr;
+			const u8* memoryPos = nullptr;
 			bool ownsMapping = false;
 
 			void Release(StorageServer& server, const tchar* reason);
diff --git a/Engine/Source/Programs/UnrealBuildAccelerator/Common/Public/UbaStorageUtils.h b/Engine/Source/Programs/UnrealBuildAccelerator/Common/Public/UbaStorageUtils.h
--- a/Engine/Source/Programs/UnrealBuildAccelerator/Common/Public/UbaStorageUtils.h
+++ b/Engine/Source/Programs/UnrealBuildAccelerator/Common/Public/UbaStorageUtils.h
@@ -117,7 +117,7 @@ namespace uba
 	struct DecompressWriter
 	{
 		bool DecompressMemoryToMemory(const u8* compressedData, u64 compressedSize, u8* writeData, u64 decompressedSize, const tchar* readHint, const tchar* writeHint);
-		bool DecompressMemoryToFile(u8* compressedData, FileAccessor& destination, u64 decompressedSize, bool useNoBuffering);
+		bool DecompressMemoryToFile(const u8* compressedData, FileAccessor& destination, u64 decompressedSize, bool useNoBuffering);
 		bool DecompressFileToMemory(const tchar* fileName, FileHandle fileHandle, u8* dest, u64 decompressedSize, const tchar* writeHint, u64 fileStartOffset);
 
 		Logger& m_logger;
