[UBA]
* Fixed so when failure reason is failed message the string also contains which message

diff --git a/Engine/Source/Programs/UnrealBuildAccelerator/Common/Private/UbaProcess.cpp b/Engine/Source/Programs/UnrealBuildAccelerator/Common/Private/UbaProcess.cpp
--- a/Engine/Source/Programs/UnrealBuildAccelerator/Common/Private/UbaProcess.cpp
+++ b/Engine/Source/Programs/UnrealBuildAccelerator/Common/Private/UbaProcess.cpp
@@ -435,7 +435,7 @@ namespace uba
 
 	bool ProcessImpl::HasFailedMessage()
 	{
-		return !m_messageSuccess;
+		return m_firstMessageError != MessageType_None;
 	}
 
 	bool ProcessImpl::WaitForExit(u32 millisecondsTimeout)
@@ -577,7 +577,7 @@ namespace uba
 			}
 
 
-			if (exitCode == 0 && !m_messageSuccess)
+			if (exitCode == 0 && m_firstMessageError != 0)
 				exitCode = UBA_EXIT_CODE(1);
 
 			bool isChild = m_parentProcess != nullptr;
@@ -798,6 +798,12 @@ namespace uba
 		return CancelWithError(m_session.m_logger.Error(TC("Unknown message type %u"), messageType));
 	}
 
+	void ProcessImpl::CallSession(MessageType type, bool success)
+	{
+		if (!success && m_firstMessageError == MessageType_None)
+			m_firstMessageError = type;
+	}
+
 	void ProcessImpl::LogLine(bool printInSession, TString&& line, LogEntryType logType)
 	{
 		if (IsCancelled())// || m_startInfo.stdoutHandle != INVALID_HANDLE_VALUE)
@@ -814,7 +820,7 @@ namespace uba
 	{
 		InitMessage msg { *this };
 		InitResponse response;
-		m_messageSuccess = m_session.GetInitResponse(response, msg) && m_messageSuccess;
+		CallSession(MessageType_Init, m_session.GetInitResponse(response, msg));
 		writer.WriteU32(m_id);
 		writer.WriteBool(m_echoOn);
 		writer.WriteBool(m_parentProcess != nullptr);
@@ -850,7 +856,7 @@ namespace uba
 		msg.access = (FileAccess)reader.ReadByte();
 
 		CreateFileResponse response;
-		m_messageSuccess = m_session.CreateFile(response, msg) && m_messageSuccess;
+		CallSession(MessageType_CreateFile, m_session.CreateFile(response, msg));
 		writer.WriteString(response.fileName);
 		writer.WriteU64(response.size);
 		writer.WriteU32(response.closeId);
@@ -868,7 +874,7 @@ namespace uba
 		msg.loaderPaths = reader.GetPositionData();
 		
 		GetFullFileNameResponse response;
-		m_messageSuccess = m_session.GetFullFileName(response, msg) && m_messageSuccess;
+		CallSession(MessageType_GetFullFileName, m_session.GetFullFileName(response, msg));
 		writer.WriteString(response.fileName);
 		writer.WriteString(response.virtualFileName);
 		writer.WriteU32(response.mappedFileTableSize);
@@ -880,7 +886,7 @@ namespace uba
 		GetLongPathNameMessage msg { *this };
 		reader.ReadString(msg.fileName);
 		GetLongPathNameResponse response;
-		m_messageSuccess = m_session.GetLongPathName(response, msg) && m_messageSuccess;
+		CallSession(MessageType_GetLongPathName, m_session.GetLongPathName(response, msg));
 		writer.WriteU32(response.errorCode);
 		writer.WriteString(response.fileName);
 		return true;
@@ -900,7 +906,7 @@ namespace uba
 		if (msg.newNameKey != StringKeyZero)
 			reader.ReadString(msg.newName);
 		CloseFileResponse response;
-		m_messageSuccess = m_session.CloseFile(response, msg) && m_messageSuccess;
+		CallSession(MessageType_CloseFile, m_session.CloseFile(response, msg));
 		writer.WriteU32(response.directoryTableSize);
 		return true;
 	}
@@ -912,7 +918,7 @@ namespace uba
 		msg.fileNameKey = reader.ReadStringKey();
 		msg.closeId = reader.ReadU32();
 		DeleteFileResponse response;
-		m_messageSuccess = m_session.DeleteFile(response, msg) && m_messageSuccess;
+		CallSession(MessageType_DeleteFile, m_session.DeleteFile(response, msg));
 		writer.WriteBool(response.result);
 		writer.WriteU32(response.errorCode);
 		writer.WriteU32(response.directoryTableSize);
@@ -927,7 +933,7 @@ namespace uba
 		msg.toKey = reader.ReadStringKey();
 		reader.ReadString(msg.toName);
 		CopyFileResponse response;
-		m_messageSuccess = m_session.CopyFile(response, msg) && m_messageSuccess;
+		CallSession(MessageType_CopyFile, m_session.CopyFile(response, msg));
 		writer.WriteString(response.fromName);
 		writer.WriteString(response.toName);
 		writer.WriteU32(response.closeId);
@@ -945,7 +951,7 @@ namespace uba
 		reader.ReadString(msg.toName);
 		msg.flags = reader.ReadU32();
 		MoveFileResponse response;
-		m_messageSuccess = m_session.MoveFile(response, msg) && m_messageSuccess;
+		CallSession(MessageType_MoveFile, m_session.MoveFile(response, msg));
 		writer.WriteBool(response.result);
 		writer.WriteU32(response.errorCode);
 		writer.WriteU32(response.directoryTableSize);
@@ -959,7 +965,7 @@ namespace uba
 		reader.ReadString(msg.fileName);
 		msg.fileMode = reader.ReadU32();
 		ChmodResponse response;
-		m_messageSuccess = m_session.Chmod(response, msg) && m_messageSuccess;
+		CallSession(MessageType_Chmod, m_session.Chmod(response, msg));
 		writer.WriteU32(response.errorCode);
 		return true;
 	}
@@ -969,7 +975,7 @@ namespace uba
 		msg.nameKey = reader.ReadStringKey();
 		reader.ReadString(msg.name);
 		CreateDirectoryResponse response;
-		m_messageSuccess = m_session.CreateDirectory(response, msg) && m_messageSuccess;
+		CallSession(MessageType_CreateDirectory, m_session.CreateDirectory(response, msg));
 		writer.WriteBool(response.result);
 		writer.WriteU32(response.errorCode);
 		writer.WriteU32(response.directoryTableSize);
@@ -982,7 +988,7 @@ namespace uba
 		msg.nameKey = reader.ReadStringKey();
 		reader.ReadString(msg.name);
 		RemoveDirectoryResponse response;
-		m_messageSuccess = m_session.RemoveDirectory(response, msg) && m_messageSuccess;
+		CallSession(MessageType_RemoveDirectory, m_session.RemoveDirectory(response, msg));
 		writer.WriteBool(response.result);
 		writer.WriteU32(response.errorCode);
 		writer.WriteU32(response.directoryTableSize);
@@ -995,7 +1001,7 @@ namespace uba
 		reader.ReadString(msg.directoryName);
 		msg.directoryNameKey = reader.ReadStringKey();
 		ListDirectoryResponse response;
-		m_messageSuccess = m_session.GetListDirectoryInfo(response, msg.directoryName, msg.directoryNameKey) && m_messageSuccess;
+		CallSession(MessageType_ListDirectory, m_session.GetListDirectoryInfo(response, msg.directoryName, msg.directoryNameKey));
 		writer.WriteU32(response.tableSize);
 		writer.WriteU32(response.tableOffset);
 		return true;
@@ -1048,8 +1054,6 @@ namespace uba
 		{
 			WrittenFile& wf = kv.second;
 
-			UBA_ASSERT(wf.owner);
-
 			if (!isInit)
 			{
 				// If file is owned by ancestor we also skip re-fetching it.. since we expect ancestor
@@ -1385,7 +1389,7 @@ namespace uba
 
 		bool newProcess = false;
 		m_exitCode = prevExitCode;
-		m_messageSuccess = m_session.GetNextProcess(*this, newProcess, nextProcess, prevExitCode, statsReader) && m_messageSuccess;
+		CallSession(MessageType_GetNextProcess, m_session.GetNextProcess(*this, newProcess, nextProcess, prevExitCode, statsReader));
 		writer.WriteBool(newProcess);
 		m_exitCode = ~0u;
 		if (!newProcess)
diff --git a/Engine/Source/Programs/UnrealBuildAccelerator/Common/Private/UbaSession.cpp b/Engine/Source/Programs/UnrealBuildAccelerator/Common/Private/UbaSession.cpp
--- a/Engine/Source/Programs/UnrealBuildAccelerator/Common/Private/UbaSession.cpp
+++ b/Engine/Source/Programs/UnrealBuildAccelerator/Common/Private/UbaSession.cpp
@@ -3058,7 +3058,7 @@ namespace uba
 				lastWriteTime = writtenFile.lastWriteTime;
 			}
 
-			if ((msg.process.m_extractExports || m_extractObjFilesSymbols) && rules.ShouldExtractSymbols(activeFile.name))
+			if ((msg.process.m_extractExports || m_extractObjFilesSymbols) && rules.ShouldExtractSymbols(activeFile.name) && fileSize != 0)
 				if (!ExtractSymbolsFromObjectFile(msg, name.data, fileSize))
 					return false;
 		}
diff --git a/Engine/Source/Programs/UnrealBuildAccelerator/Common/Private/UbaSessionClient.cpp b/Engine/Source/Programs/UnrealBuildAccelerator/Common/Private/UbaSessionClient.cpp
--- a/Engine/Source/Programs/UnrealBuildAccelerator/Common/Private/UbaSessionClient.cpp
+++ b/Engine/Source/Programs/UnrealBuildAccelerator/Common/Private/UbaSessionClient.cpp
@@ -17,6 +17,18 @@
 
 namespace uba
 {
+	const tchar* ToString(MessageType type)
+	{
+		switch (type)
+		{
+			#define UBA_PROCESS_MESSAGE(x) case MessageType_##x: return TC(#x);
+			UBA_PROCESS_MESSAGES
+			#undef UBA_PROCESS_MESSAGE
+		default:
+			return TC("UnknownProcessMessage"); // Should never happen
+		}
+	}
+
 	void SessionClientCreateInfo::Apply(const Config& config)
 	{
 		SessionCreateInfo::Apply(config);
@@ -2074,7 +2086,7 @@ namespace uba
 							if (process.HasFailedMessage()) // If there are failure caused by failed messages we send back for retry
 							{
 								if (session.m_loop)
-									session.SendReturnProcess(rec->handle.GetId(), TC("Failed message"));
+									session.SendReturnProcess(rec->handle.GetId(), StringBuffer<128>().Appendf(TC("Failed message (%s)"), ToString(process.m_firstMessageError)).data);
 								return;
 							}
 						}
diff --git a/Engine/Source/Programs/UnrealBuildAccelerator/Common/Public/UbaProcess.h b/Engine/Source/Programs/UnrealBuildAccelerator/Common/Public/UbaProcess.h
--- a/Engine/Source/Programs/UnrealBuildAccelerator/Common/Public/UbaProcess.h
+++ b/Engine/Source/Programs/UnrealBuildAccelerator/Common/Public/UbaProcess.h
@@ -80,6 +80,7 @@ namespace uba
 		bool HandleMessage(BinaryReader& reader, BinaryWriter& writer);
 		bool HandleSpecialApplication();
 		bool HandleGetWrittenFiles(bool isInit, BinaryWriter& writer);
+		void CallSession(MessageType type, bool success);
 
 		#define UBA_PROCESS_MESSAGE(M) bool Handle##M(BinaryReader& reader, BinaryWriter& writer);
 		UBA_PROCESS_MESSAGES
@@ -148,8 +149,8 @@ namespace uba
 		u32 m_exitCode = ~0u;
 		u32 m_messageCount = 0;
 		Atomic<bool> m_hasExited;
+		MessageType m_firstMessageError = MessageType_None;
 		bool m_runningRemote = false;
-		bool m_messageSuccess = true;
 		bool m_echoOn = true;
 		bool m_gotExitMessage = false;
 		bool m_parentReportedExit = false;
diff --git a/Engine/Source/Programs/UnrealBuildAccelerator/Core/Public/UbaProtocol.h b/Engine/Source/Programs/UnrealBuildAccelerator/Core/Public/UbaProtocol.h
--- a/Engine/Source/Programs/UnrealBuildAccelerator/Core/Public/UbaProtocol.h
+++ b/Engine/Source/Programs/UnrealBuildAccelerator/Core/Public/UbaProtocol.h
@@ -42,7 +42,7 @@ namespace uba
 
 	enum MessageType : u8
 	{
-		Unused = 0,
+		MessageType_None = 0,
 		#define UBA_PROCESS_MESSAGE(type) MessageType_##type,
 		UBA_PROCESS_MESSAGES
 		#undef UBA_PROCESS_MESSAGE
