[UBA]
* Fixed bug in LLVM IR path in UbaObjTool where it added all non-imported symbols as exported in "onlyexported" mode (used by lib writer path)
* Enabled obj-file crawler for UbaObjTool

diff --git a/Engine/Source/Programs/UnrealBuildAccelerator/Common/Private/UbaObjectFileLLVMIR.cpp b/Engine/Source/Programs/UnrealBuildAccelerator/Common/Private/UbaObjectFileLLVMIR.cpp
--- a/Engine/Source/Programs/UnrealBuildAccelerator/Common/Private/UbaObjectFileLLVMIR.cpp
+++ b/Engine/Source/Programs/UnrealBuildAccelerator/Common/Private/UbaObjectFileLLVMIR.cpp
@@ -888,10 +888,14 @@ namespace uba
 				}
 
 				// ExternalLinkage || DLLImportLinkage || DLLExportLinkage || LinkOnceODRAutoHideLinkage || WeakODRLinkage
-				isExported = linkage == 0 || linkage == 5 || linkage == 6 || linkage == 15 || linkage == 17;
-
-				if (dllStorage == DllStorage_Import)
-					isExported = false;
+				if (parseMode == ObjectFileParseMode_All)
+				{
+					isExported = linkage == 0 || linkage == 5 || linkage == 6 || linkage == 15 || linkage == 17;
+					if (dllStorage == DllStorage_Import)
+						isExported = false;
+				}
+				else
+					isExported = dllStorage == DllStorage_Export;
 			}
 
 			if (record.code == MODULE_CODE_GLOBALVAR)
diff --git a/Engine/Source/Programs/UnrealBuildAccelerator/Core/Private/UbaApplicationRules.cpp b/Engine/Source/Programs/UnrealBuildAccelerator/Core/Private/UbaApplicationRules.cpp
--- a/Engine/Source/Programs/UnrealBuildAccelerator/Core/Private/UbaApplicationRules.cpp
+++ b/Engine/Source/Programs/UnrealBuildAccelerator/Core/Private/UbaApplicationRules.cpp
@@ -684,6 +684,11 @@ namespace uba
 		{
 			return true;
 		}
+
+		virtual DependencyCrawlerType GetDependencyCrawlerType() const override
+		{
+			return DependencyCrawlerType_MsvcLinker;
+		}
 	};
 
 	class ApplicationRulesUbaTestApp : public ApplicationRules
diff --git a/Engine/Source/Programs/UnrealBuildAccelerator/ObjTool/Private/UbaImportLibWriter.cpp b/Engine/Source/Programs/UnrealBuildAccelerator/ObjTool/Private/UbaImportLibWriter.cpp
--- a/Engine/Source/Programs/UnrealBuildAccelerator/ObjTool/Private/UbaImportLibWriter.cpp
+++ b/Engine/Source/Programs/UnrealBuildAccelerator/ObjTool/Private/UbaImportLibWriter.cpp
@@ -556,7 +556,7 @@ namespace uba
 
 	public:
 
-		ImportLib() : m_memory(512*1024*1024) // 512mb should be more than enough reserved memory for the biggest import libs
+		ImportLib() : m_memory(128*1024*1024) // 512mb should be more than enough reserved memory for the biggest import libs
 		{
 		}
 
