[UBA]
* Added -httpcrypto to uba cache service that can be used to decrypt crypto keys registered using http addcrypto

diff --git a/Engine/Source/Programs/UnrealBuildAccelerator/CacheService/Private/UbaCacheService.cpp b/Engine/Source/Programs/UnrealBuildAccelerator/CacheService/Private/UbaCacheService.cpp
--- a/Engine/Source/Programs/UnrealBuildAccelerator/CacheService/Private/UbaCacheService.cpp
+++ b/Engine/Source/Programs/UnrealBuildAccelerator/CacheService/Private/UbaCacheService.cpp
@@ -3,6 +3,7 @@
 #include "UbaCacheServer.h"
 #include "UbaApplication.h"
 #include "UbaConfig.h"
+#include "UbaCrypto.h"
 #include "UbaDirectoryIterator.h"
 #include "UbaFile.h"
 #include "UbaHttpServer.h"
@@ -57,6 +58,7 @@ namespace uba
 		logger.Info(TC("  -expiration=<seconds>   Time until unused cache entries get deleted. Defaults to %s (%u seconds)"), TimeToText(MsToTime(DefaultExpiration*1000)).str, DefaultExpiration);
 		logger.Info(TC("  -config=<file>          Config file that contains options for various systems"));
 		logger.Info(TC("  -http=<port>            If set, a http server will be started and listen on <port>"));
+		logger.Info(TC("  -httpcrypto=<crypto>    If set, all provided keys will be decrypted using this key."));
 		logger.Info(TC("  -fullmaintenance        Force a full maintenance"));
 		logger.Info(TC("  -nomaintenance          Skip all maintenance"));
 		logger.Info(TC("  -crash                  Force a crash (for testing)"));
@@ -115,6 +117,7 @@ namespace uba
 		float storageCapacityGb = float(DefaultCapacityGb);
 		StringBuffer<256> workDir;
 		StringBuffer<128> listenIp;
+		Vector<CryptoKey> httpCryptos;
 		u16 port = DefaultCachePort;
 		u16 httpPort = 0;
 		bool quiet = false;
@@ -186,6 +189,19 @@ namespace uba
 				if (!value.Parse(httpPort))
 					httpPort = 80;
 			}
+			else if (name.Equals(TCV("-httpcrypto")))
+			{
+				if (value.count != 32)
+					return PrintHelp(TC("Invalid number of characters in crypto string. Should be 32"));
+				u8 crypto[16];
+				((u64*)crypto)[0] = StringToValue(value.data, 16);
+				((u64*)crypto)[1] = StringToValue(value.data + 16, 16);
+				LoggerWithWriter tempLogger(g_consoleLogWriter, TC(""));
+				CryptoKey key = Crypto::CreateKey(tempLogger, crypto);
+				if (key == InvalidCryptoKey)
+					return PrintHelp(TC("Failed to create http crypto key"));
+				httpCryptos.emplace_back(key);
+			}
 			else if (name.Equals(TCV("-fullmaintenance")))
 			{
 				fullMaintenance = true;
@@ -506,7 +522,20 @@ namespace uba
 					if (!CryptoFromString(crypto128Data, 16, arguments))
 						return "Failed to read crypto argument (Needs to be 32 characters long)";
 					u64 expirationTime = GetTime() + MsToTime(expirationSeconds*1000);
-					networkServer.RegisterCryptoKey(crypto128Data, expirationTime);
+
+					if (httpCryptos.empty())
+						networkServer.RegisterCryptoKey(crypto128Data, expirationTime);
+					else
+					{
+						for (auto& crypto : httpCryptos)
+						{
+							u8 decryptedKey[16];
+							memcpy(decryptedKey, crypto128Data, 16);
+							if (!Crypto::Decrypt(logger, crypto, decryptedKey, 16))
+								return "Failed to decrypt key from http addcrypto";
+							networkServer.RegisterCryptoKey(decryptedKey, expirationTime);
+						}
+					}
 					return (const char*)nullptr;
 				});
 			httpServer.StartListen(httpPort);
