[UBA]
* Changed scheduler lock to use futex instead of readerwriterlock on posix platforms to see if it resolves a rare hang happening on osx

diff --git a/Engine/Source/Programs/UnrealBuildAccelerator/Common/Private/UbaScheduler.cpp b/Engine/Source/Programs/UnrealBuildAccelerator/Common/Private/UbaScheduler.cpp
--- a/Engine/Source/Programs/UnrealBuildAccelerator/Common/Private/UbaScheduler.cpp
+++ b/Engine/Source/Programs/UnrealBuildAccelerator/Common/Private/UbaScheduler.cpp
@@ -146,7 +146,7 @@ namespace uba
 	void Scheduler::SkipAllQueued()
 	{
 		Vector<ProcessStartInfo2*> skipped;
-		SCOPED_WRITE_LOCK(m_processEntriesLock, lock);
+		SCOPED_FUTEX(m_processEntriesLock, lock);
 		for (auto& entry : m_processEntries)
 			if (entry.status < ProcessStatus_Running)
 			{
@@ -160,7 +160,7 @@ namespace uba
 
 	void Scheduler::Cleanup()
 	{
-		SCOPED_WRITE_LOCK(m_processEntriesLock, lock);
+		SCOPED_FUTEX(m_processEntriesLock, lock);
 		for (auto& entry : m_processEntries)
 		{
 			UBA_ASSERTF(entry.status > ProcessStatus_Running, TC("Found processes in queue/running state when stopping scheduler."));
@@ -209,7 +209,7 @@ namespace uba
 
 		bool useCache = info.cacheBucketId && !m_cacheClients.empty() && !m_writeToCache && rules->IsCacheable();
 
-		SCOPED_WRITE_LOCK(m_processEntriesLock, lock);
+		SCOPED_FUTEX(m_processEntriesLock, lock);
 		u32 index = u32(m_processEntries.size());
 		auto& entry = m_processEntries.emplace_back();
 		entry.info = info2;
@@ -254,7 +254,7 @@ namespace uba
 
 	bool Scheduler::IsEmpty()
 	{
-		SCOPED_READ_LOCK(m_processEntriesLock, lock);
+		SCOPED_FUTEX_READ(m_processEntriesLock, lock);
 		return m_processEntries.size() <= m_finishedProcesses;
 	}
 
@@ -269,7 +269,7 @@ namespace uba
 			return 0;
 
 		u32 count = 0;
-		SCOPED_READ_LOCK(m_processEntriesLock, lock);
+		SCOPED_FUTEX_READ(m_processEntriesLock, lock);
 		for (auto& entry : m_processEntries)
 		{
 			if (!entry.canExecuteRemotely)
@@ -290,7 +290,7 @@ namespace uba
 			return 0;
 
 		float weight = 0;
-		SCOPED_READ_LOCK(m_processEntriesLock, lock);
+		SCOPED_FUTEX_READ(m_processEntriesLock, lock);
 		for (auto& entry : m_processEntries)
 		{
 			if (!entry.canExecuteRemotely)
@@ -337,7 +337,7 @@ namespace uba
 		if (processIndex == ~0u)
 			return;
 
-		SCOPED_WRITE_LOCK(m_processEntriesLock, lock);
+		SCOPED_FUTEX(m_processEntriesLock, lock);
 		if (m_processEntries[processIndex].status != ProcessStatus_Running)
 			return;
 		m_processEntries[processIndex].status = ProcessStatus_QueuedForRun;
@@ -354,7 +354,7 @@ namespace uba
 		if (processIndex == ~0u)
 			return;
 
-		SCOPED_WRITE_LOCK(m_processEntriesLock, lock);
+		SCOPED_FUTEX(m_processEntriesLock, lock);
 		if (m_processEntries[processIndex].status != ProcessStatus_Running)
 			return;
 		m_processEntries[processIndex].status = ProcessStatus_QueuedForRun;
@@ -377,7 +377,7 @@ namespace uba
 		if (m_session.IsRemoteExecutionDisabled())
 			return;
 		u32 count = 0;
-		SCOPED_READ_LOCK(m_processEntriesLock, lock);
+		SCOPED_FUTEX_READ(m_processEntriesLock, lock);
 		for (auto& entry : m_processEntries)
 			count += u32(entry.canExecuteRemotely && entry.status <= ProcessStatus_QueuedForRun);
 		lock.Leave();
@@ -491,7 +491,7 @@ namespace uba
 		while (true)
 		{
 			ProcessStatus prevStatus;
-			SCOPED_WRITE_LOCK(m_processEntriesLock, lock);
+			SCOPED_FUTEX(m_processEntriesLock, lock);
 			u32 indexToRun = PopProcess(isLocal, prevStatus);
 			if (indexToRun == ~0u)
 				return false;
@@ -600,7 +600,7 @@ namespace uba
 		while (true)
 		{
 			ProcessStatus prevStatus;
-			SCOPED_WRITE_LOCK(m_processEntriesLock, lock);
+			SCOPED_FUTEX(m_processEntriesLock, lock);
 			u32 indexToRun = PopProcess(isLocal, prevStatus);
 			if (indexToRun == ~0u)
 				return false;
@@ -664,7 +664,7 @@ namespace uba
 
 		bool isDone = exitedResponse == ProcessExitedResponse_None;
 
-		SCOPED_WRITE_LOCK(m_processEntriesLock, lock);
+		SCOPED_FUTEX(m_processEntriesLock, lock);
 		auto& entry = m_processEntries[info.processIndex];
 
 		u32* dependencies = entry.dependencies;
@@ -1026,7 +1026,7 @@ namespace uba
 		}
 
 		// Return weight while running these tasks
-		SCOPED_WRITE_LOCK(m_processEntriesLock, lock);
+		SCOPED_FUTEX(m_processEntriesLock, lock);
 		m_activeLocalProcessWeight -= weight;
 		lock.Leave();
 
diff --git a/Engine/Source/Programs/UnrealBuildAccelerator/Common/Public/UbaScheduler.h b/Engine/Source/Programs/UnrealBuildAccelerator/Common/Public/UbaScheduler.h
--- a/Engine/Source/Programs/UnrealBuildAccelerator/Common/Public/UbaScheduler.h
+++ b/Engine/Source/Programs/UnrealBuildAccelerator/Common/Public/UbaScheduler.h
@@ -125,7 +125,7 @@ namespace uba
 			bool canExecuteRemotely;
 		};
 
-		ReaderWriterLock m_processEntriesLock;
+		Futex m_processEntriesLock;
 		Vector<ProcessEntry> m_processEntries;
 		u32 m_processEntriesStart = 0;
 
