[UBA]
* Moved some config reading code to outside scheduler lock

diff --git a/Engine/Source/Programs/UnrealBuildAccelerator/Common/Private/UbaScheduler.cpp b/Engine/Source/Programs/UnrealBuildAccelerator/Common/Private/UbaScheduler.cpp
--- a/Engine/Source/Programs/UnrealBuildAccelerator/Common/Private/UbaScheduler.cpp
+++ b/Engine/Source/Programs/UnrealBuildAccelerator/Common/Private/UbaScheduler.cpp
@@ -208,16 +208,8 @@ namespace uba
 		info2->rules = rules;
 
 		bool useCache = info.cacheBucketId && !m_cacheClients.empty() && !m_writeToCache && rules->IsCacheable();
-
-		SCOPED_FUTEX(m_processEntriesLock, lock);
-		u32 index = u32(m_processEntries.size());
-		auto& entry = m_processEntries.emplace_back();
-		entry.info = info2;
-		entry.dependencies = dep;
-		entry.dependencyCount = info.dependencyCount;
-		entry.status = useCache ? ProcessStatus_QueuedForCache : ProcessStatus_QueuedForRun;
-		entry.canDetour = info.canDetour;
-		entry.canExecuteRemotely = info.canExecuteRemotely && info.canDetour;
+		bool canDetour = info.canDetour;
+		bool canExecuteRemotely = info.canExecuteRemotely && info.canDetour;
 
 		if (m_processConfigs)
 		{
@@ -229,11 +221,20 @@ namespace uba
 			lower.Replace('.', '_');
 			if (const ConfigTable* processConfig = m_processConfigs->GetTable(lower.data))
 			{
-				processConfig->GetValueAsBool(entry.canExecuteRemotely, TC("CanExecuteRemotely"));
-				processConfig->GetValueAsBool(entry.canDetour, TC("CanDetour"));
+				processConfig->GetValueAsBool(canExecuteRemotely, TC("CanExecuteRemotely"));
+				processConfig->GetValueAsBool(canDetour, TC("CanDetour"));
 			}
 		}
 
+		SCOPED_FUTEX(m_processEntriesLock, lock);
+		u32 index = u32(m_processEntries.size());
+		auto& entry = m_processEntries.emplace_back();
+		entry.info = info2;
+		entry.dependencies = dep;
+		entry.dependencyCount = info.dependencyCount;
+		entry.status = useCache ? ProcessStatus_QueuedForCache : ProcessStatus_QueuedForRun;
+		entry.canDetour = canDetour;
+		entry.canExecuteRemotely = canExecuteRemotely;
 		lock.Leave();
 
 		++m_totalProcesses;
