[UBA]
* Fixed deadlock when outside code calls cancel on entire session from inside a process exit callback

diff --git a/Engine/Source/Programs/UnrealBuildAccelerator/Common/Private/UbaProcess.cpp b/Engine/Source/Programs/UnrealBuildAccelerator/Common/Private/UbaProcess.cpp
--- a/Engine/Source/Programs/UnrealBuildAccelerator/Common/Private/UbaProcess.cpp
+++ b/Engine/Source/Programs/UnrealBuildAccelerator/Common/Private/UbaProcess.cpp
@@ -1812,7 +1812,7 @@ namespace uba
 		}
 
 		bool affinitySet = false;
-		if (!m_messageThread.Wait(0))
+		if (m_messageThread.IsInside())
 		{
 			GroupAffinity aff;
 			if (m_messageThread.GetGroupAffinity(aff))
@@ -2312,4 +2312,8 @@ namespace uba
 		return false;
 	}
 
+	bool ProcessImpl::IsCalledFromThis() const
+	{
+		return m_messageThread.IsInside();
+	}
 }
diff --git a/Engine/Source/Programs/UnrealBuildAccelerator/Common/Private/UbaSession.cpp b/Engine/Source/Programs/UnrealBuildAccelerator/Common/Private/UbaSession.cpp
--- a/Engine/Source/Programs/UnrealBuildAccelerator/Common/Private/UbaSession.cpp
+++ b/Engine/Source/Programs/UnrealBuildAccelerator/Common/Private/UbaSession.cpp
@@ -1795,16 +1795,17 @@ namespace uba
 			Vector<ProcessHandle> processes;
 			{
 				SCOPED_FUTEX(m_processesLock, lock);
-				isEmpty = m_processes.empty();
 				processes.reserve(m_processes.size());
 				for (auto& pair : m_processes)
-					processes.push_back(pair.second);
+					if (pair.second.m_process && !pair.second.m_process->IsCalledFromThis())
+						processes.push_back(pair.second);
+				isEmpty = processes.empty();
 			}
 
 			if (isFirst)
 			{
 				isFirst = false;
-				if (!processes.empty())
+				if (!isEmpty)
 					m_logger.Info(TC("Cancelling %llu processes and wait for them to exit"), processes.size());
 				++m_logger.isMuted;
 			}
diff --git a/Engine/Source/Programs/UnrealBuildAccelerator/Common/Private/UbaThread.cpp b/Engine/Source/Programs/UnrealBuildAccelerator/Common/Private/UbaThread.cpp
--- a/Engine/Source/Programs/UnrealBuildAccelerator/Common/Private/UbaThread.cpp
+++ b/Engine/Source/Programs/UnrealBuildAccelerator/Common/Private/UbaThread.cpp
@@ -170,6 +170,8 @@ namespace uba
 		}
 		else
 		{
+			UBA_ASSERTF(GetCurrentThreadId() != GetThreadId(m_handle), TC("Thread is trying to wait on itself. This will end up in infinite hang"));
+
 			if (WaitForSingleObject(m_handle, milliseconds) == WAIT_TIMEOUT)
 				return false;
 		}
@@ -214,6 +216,15 @@ namespace uba
 #endif
 	}
 
+	bool Thread::IsInside() const
+	{
+#if PLATFORM_WINDOWS
+		return GetCurrentThreadId() == GetThreadId(m_handle);
+#else
+		return pthread_self() == *(pthread_t*)&m_handle;
+#endif
+	}
+
 	bool TraverseAllThreads(const TraverseThreadFunc& func, const TraverseThreadErrorFunc& errorFunc)
 	{
 #if PLATFORM_WINDOWS
diff --git a/Engine/Source/Programs/UnrealBuildAccelerator/Common/Public/UbaProcess.h b/Engine/Source/Programs/UnrealBuildAccelerator/Common/Public/UbaProcess.h
--- a/Engine/Source/Programs/UnrealBuildAccelerator/Common/Public/UbaProcess.h
+++ b/Engine/Source/Programs/UnrealBuildAccelerator/Common/Public/UbaProcess.h
@@ -23,6 +23,7 @@ namespace uba
 		virtual const Vector<u8>& GetTrackedInputs() const = 0;
 		virtual const Vector<u8>& GetTrackedOutputs() const = 0;
 		virtual void Cancel(bool terminate) { UBA_ASSERT(false); }
+		virtual bool IsCalledFromThis() const { return false; }
 		virtual const tchar* GetExecutingHost() const { UBA_ASSERT(false); return nullptr; }
 		virtual bool IsRemote() const { UBA_ASSERT(false); return false; }
 		virtual ProcessExecutionType GetExecutionType() const { UBA_ASSERT(false); return ProcessExecutionType_Native; }
@@ -53,6 +54,7 @@ namespace uba
 		virtual const Vector<u8>& GetTrackedInputs() const override { return m_trackedInputs; }
 		virtual const Vector<u8>& GetTrackedOutputs() const override { return m_trackedOutputs; }
 		virtual void Cancel(bool terminate) override;
+		virtual bool IsCalledFromThis() const override;
 		virtual const tchar* GetExecutingHost() const override { return TC(""); }
 		virtual bool IsRemote() const override { return false; }
 		virtual ProcessExecutionType GetExecutionType() const override { return m_detourEnabled ? ProcessExecutionType_Local : ProcessExecutionType_Native; }
diff --git a/Engine/Source/Programs/UnrealBuildAccelerator/Common/Public/UbaThread.h b/Engine/Source/Programs/UnrealBuildAccelerator/Common/Public/UbaThread.h
--- a/Engine/Source/Programs/UnrealBuildAccelerator/Common/Public/UbaThread.h
+++ b/Engine/Source/Programs/UnrealBuildAccelerator/Common/Public/UbaThread.h
@@ -23,6 +23,7 @@ namespace uba
 		void Start(Function<u32()>&& func, const tchar* description = nullptr);
 		bool Wait(u32 milliseconds = ~0u, Event* wakeupEvent = nullptr);
 		bool GetGroupAffinity(GroupAffinity& out);
+		bool IsInside() const;
 
 	private:
 		Function<u32()>	m_func;
