[UBA]
* Changed errors to warnings for decompression failures since the outer step decides if this is actually a real fail (some parts of the code can retry using different path if failing)

diff --git a/Engine/Source/Programs/UnrealBuildAccelerator/Common/Private/UbaSessionClient.cpp b/Engine/Source/Programs/UnrealBuildAccelerator/Common/Private/UbaSessionClient.cpp
--- a/Engine/Source/Programs/UnrealBuildAccelerator/Common/Private/UbaSessionClient.cpp
+++ b/Engine/Source/Programs/UnrealBuildAccelerator/Common/Private/UbaSessionClient.cpp
@@ -741,7 +741,7 @@ namespace uba
 				return false;
 		}
 		if (casKey == CasKeyZero)
-			return m_logger.Error(TC("Failed to store cas on server for local file %s (size %llu destination %s)"), file.backedName.c_str(), file.mappingWritten, file.name.c_str());
+			return m_logger.Warning(TC("Failed to store cas on server for local file %s (size %llu destination %s)"), file.backedName.c_str(), file.mappingWritten, file.name.c_str());
 	
 		CloseFileMapping(m_logger, file.mappingHandle, file.backedName.c_str());
 		file.mappingHandle = {};
@@ -756,7 +756,7 @@ namespace uba
 			writer.WriteU32(file.attributes);
 			writer.WriteCasKey(casKey);
 			if (!msg.Send(reader, Stats().sendFileMsg))
-				return m_logger.Error(TC("Failed to send file %s to server"), file.backedName.c_str());
+				return m_logger.Warning(TC("Failed to send file %s to server"), file.backedName.c_str());
 		}
 
 		#if UBA_DEBUG_LOGGER
@@ -764,7 +764,7 @@ namespace uba
 		#endif
 
 		if (!reader.ReadBool())
-			return m_logger.Error(TC("Server failed to copy cas %s to %s (local source %s)"), CasKeyString(casKey).str, file.name.c_str(), file.backedName.c_str());
+			return m_logger.Warning(TC("Server failed to copy cas %s to %s (local source %s)"), CasKeyString(casKey).str, file.name.c_str(), file.backedName.c_str());
 		return true;
 	}
 
diff --git a/Engine/Source/Programs/UnrealBuildAccelerator/Common/Private/UbaStorageUtils.cpp b/Engine/Source/Programs/UnrealBuildAccelerator/Common/Private/UbaStorageUtils.cpp
--- a/Engine/Source/Programs/UnrealBuildAccelerator/Common/Private/UbaStorageUtils.cpp
+++ b/Engine/Source/Programs/UnrealBuildAccelerator/Common/Private/UbaStorageUtils.cpp
@@ -1212,7 +1212,7 @@ namespace uba
 						{
 							bool f = false;
 							if (rec->error.compare_exchange_strong(f, true))
-								rec->logger->Error(TC("Decompressed block size %u is invalid. Decompressed file is %u (%s)"), decompressedBlockSize, rec->decompressedSize, rec->hint);
+								rec->logger->Warning(TC("Decompressed block size %u is invalid. Decompressed file is %u (%s)"), decompressedBlockSize, rec->decompressedSize, rec->hint);
 							if (!--rec->refCount)
 								delete rec;
 							rec->done.Set();
@@ -1232,7 +1232,7 @@ namespace uba
 						{
 							bool f = false;
 							if (rec->error.compare_exchange_strong(f, true))
-								rec->logger->Error(TC("Expecting to be able to decompress %u bytes to %u bytes but got %llu (%s)"), compressedBlockSize, decompressedBlockSize, decompLen, rec->hint);
+								rec->logger->Warning(TC("Expecting to be able to decompress %u bytes to %u bytes but got %llu (%s)"), compressedBlockSize, decompressedBlockSize, decompLen, rec->hint);
 							if (!--rec->refCount)
 								delete rec;
 							rec->done.Set();
@@ -1276,7 +1276,7 @@ namespace uba
 					break;
 				u32 decompressedBlockSize = ((const u32*)readPos)[1];
 				if (decompressedBlockSize == 0 || decompressedBlockSize > left)
-					return m_logger.Error(TC("Decompressed block size %u is invalid. Decompressed file is %u (%s -> %s)"), decompressedBlockSize, decompressedSize, readHint, writeHint);
+					return m_logger.Warning(TC("Decompressed block size %u is invalid. Decompressed file is %u (%s -> %s)"), decompressedBlockSize, decompressedSize, readHint, writeHint);
 				readPos += sizeof(u32) * 2;
 
 				void* decoderMem = nullptr;
@@ -1284,7 +1284,7 @@ namespace uba
 				TimerScope ts(stats.decompressToMem);
 				OO_SINTa decompLen = OodleLZ_Decompress(readPos, (OO_SINTa)compressedBlockSize, writePos, (OO_SINTa)decompressedBlockSize, OodleLZ_FuzzSafe_Yes, OodleLZ_CheckCRC_No, OodleLZ_Verbosity_None, NULL, 0, NULL, NULL, decoderMem, decoderMemSize);
 				if (decompLen != decompressedBlockSize)
-					return m_logger.Error(TC("Expecting to be able to decompress %u to %u bytes at pos %llu but got %llu. File compressed size: %llu (%s -> %s)"), compressedBlockSize, decompressedBlockSize, decompressedSize - left, decompLen, compressedSize, readHint, writeHint);
+					return m_logger.Warning(TC("Expecting to be able to decompress %u to %u bytes at pos %llu but got %llu. File compressed size: %llu Decompressed size: %llu (%s -> %s)"), compressedBlockSize, decompressedBlockSize, decompressedSize - left, decompLen, compressedSize, decompressedSize, readHint, writeHint);
 				writePos += decompressedBlockSize;
 				readPos += compressedBlockSize;
 				left -= decompressedBlockSize;
