[UBA]
* Fixed race condition in refcounting of file mapping entries. This is probably the reason for the failed Duplicate handle error that we've seen (once) on the farm

diff --git a/Engine/Source/Programs/UnrealBuildAccelerator/Common/Private/UbaSession.cpp b/Engine/Source/Programs/UnrealBuildAccelerator/Common/Private/UbaSession.cpp
--- a/Engine/Source/Programs/UnrealBuildAccelerator/Common/Private/UbaSession.cpp
+++ b/Engine/Source/Programs/UnrealBuildAccelerator/Common/Private/UbaSession.cpp
@@ -844,8 +844,6 @@ namespace uba
 
 		if (entry.handled)
 		{
-			entryLock.Leave();
-
 			if (!entry.success)
 				return false;
 			out.size = entry.contentSize;
