[UBA]
* Swapped order of file handle and mapping handle when writing to u64 in order to make values smaller so they compress better when turned into base62 strings
* Changed assert into debug entry for when memoryfile is replaced from child process since it happens and is not wrong (parent process closes file and reopen later)
* Added IsExecutable function that reads attributes (on posix) and returns true if flags says it is an executable

diff --git a/Engine/Source/Programs/UnrealBuildAccelerator/Core/Public/UbaPlatform.h b/Engine/Source/Programs/UnrealBuildAccelerator/Core/Public/UbaPlatform.h
--- a/Engine/Source/Programs/UnrealBuildAccelerator/Core/Public/UbaPlatform.h
+++ b/Engine/Source/Programs/UnrealBuildAccelerator/Core/Public/UbaPlatform.h
@@ -137,6 +137,7 @@ namespace uba
 	inline LARGE_INTEGER ToLargeInteger(u64 v) { LARGE_INTEGER li; li.QuadPart = (LONGLONG)v; return li; }
 	inline LARGE_INTEGER ToLargeInteger(u32 high, u32 low) { LARGE_INTEGER li; li.HighPart = (LONG)high; li.LowPart = low; return li; }
 	inline bool IsDirectory(u32 attributes) { return (attributes & FILE_ATTRIBUTE_DIRECTORY) != 0; }
+	inline bool IsExecutable(u32 attributes) { return false; }
 	inline bool IsUncPath(const tchar* path) { return path[0] == '\\' && path[1] == '\\'; }
 	#define TStrlen(s) u32(wcslen(s))
 	#define TStrchr(a, b) wcschr(a, b)
@@ -153,8 +154,8 @@ namespace uba
 	{
 		bool operator==(const FileMappingHandle& o) const { return mh == o.mh; }
 		bool IsValid() const { return mh != 0; }
-		u64 ToU64() const { return u64(fh) | (u64(mh) << 32); }
-		static FileMappingHandle FromU64(u64 v) { return { (HANDLE)(v & ~0u), (HANDLE)(v >> 32) }; }
+		u64 ToU64() const { return u64(mh) | (u64(fh) << 32); }
+		static FileMappingHandle FromU64(u64 v) { return { (HANDLE)(v >> 32), (HANDLE)(v & ~0u) }; }
 		HANDLE fh = 0;
 		HANDLE mh = 0;
 	};
@@ -180,6 +181,7 @@ namespace uba
 	inline u64 FromTimeSpec(const timespec& ts) { return u64(ts.tv_sec) * 10'000'000ull + u64(ts.tv_nsec/100); }
 	inline timespec ToTimeSpec(u64 time) { timespec ts; ts.tv_sec = time / 10'000'000ull; ts.tv_nsec = (time - (u64(ts.tv_sec) * 10'000'000ull)) * 100; return ts; }
 	inline bool IsDirectory(u32 attributes) { return S_ISDIR(attributes); }
+	inline bool IsExecutable(u32 attributes) { return (attributes & (S_IXUSR | S_IXGRP | S_IXOTH)) != 0; }
 	inline bool IsUncPath(const tchar* path) { return false; }
 	#if PLATFORM_LINUX
 	inline constexpr bool CaseInsensitiveFs = false;
diff --git a/Engine/Source/Programs/UnrealBuildAccelerator/Detours/Private/UbaDetoursFileMappingTable.cpp b/Engine/Source/Programs/UnrealBuildAccelerator/Detours/Private/UbaDetoursFileMappingTable.cpp
--- a/Engine/Source/Programs/UnrealBuildAccelerator/Detours/Private/UbaDetoursFileMappingTable.cpp
+++ b/Engine/Source/Programs/UnrealBuildAccelerator/Detours/Private/UbaDetoursFileMappingTable.cpp
@@ -201,7 +201,8 @@ namespace uba
 
 			if (auto mf = info.memoryFile) // Memory file is now wrong.. we need to delete it (or should we create a new one with the new mapping?)
 			{
-				UBA_ASSERTF(!info.refCount, TC("File has memory file and refcount. This is not supported. %s (BackedFile: %s Size: %llu)"), info.originalName, info.name, info.size);
+				if (info.refCount)
+					DEBUG_LOG(TC("File has memory file and refcount %u. Will drop memory file and use received file. %s (BackedFile: %s Size: %llu)"), info.refCount, info.originalName, info.name, info.size);
 				#if PLATFORM_WINDOWS
 				if (!mf->isLocalOnly)
 				{
