[UBA]
* Added UbaCli.exe -symbol=<pathtodll>+0x<offset> to be able to resolve addresses to symbols

diff --git a/Engine/Source/Programs/UnrealBuildAccelerator/Cli/Private/UbaCli.cpp b/Engine/Source/Programs/UnrealBuildAccelerator/Cli/Private/UbaCli.cpp
--- a/Engine/Source/Programs/UnrealBuildAccelerator/Cli/Private/UbaCli.cpp
+++ b/Engine/Source/Programs/UnrealBuildAccelerator/Cli/Private/UbaCli.cpp
@@ -185,6 +185,8 @@ namespace uba
 		TString testDecompress;
 		TString addCas;
 		TString configFile;
+		TString symbolFile;
+		u64 symbolOffset = 0;
 
 		struct VfsEntry
 		{
@@ -349,6 +351,19 @@ namespace uba
 					}
 					testDecompress = value.data;
 				}
+				else if (name.Equals(TCV("-symbol")))
+				{
+					const tchar* plus = value.First('+');
+					if (!plus)
+						return PrintHelp(TC("No + found. Format is -symbol=<file>+0x<offset>"));
+					symbolFile = TString((const tchar*)value.data, plus);
+					tchar* endPtr;
+					#if PLATFORM_WINDOWS
+					symbolOffset = wcstol(plus+1, &endPtr, 0);
+					#else
+					symbolOffset = strtol(plus+1, &endPtr, 0);
+					#endif
+				}
 				else if (name.Equals(TCV("-deletecas")))
 				{
 					deleteCas = true;
@@ -759,6 +774,40 @@ namespace uba
 			exit = true;
 		}
 
+		if (!symbolFile.empty())
+		{
+			FileInformation info;
+			GetFileInformation(info, logger, symbolFile.c_str());
+			u64 fakeStartAddress = 0x800000;
+
+			StackBinaryWriter<128> writer;
+			#if PLATFORM_WINDOWS
+			writer.WriteBool(false);
+			#endif
+			writer.Write7BitEncoded(1); // Callstack count
+			writer.Write7BitEncoded(0); // Module index
+			writer.Write7BitEncoded(symbolOffset); // Module offset
+
+			writer.Write7BitEncoded(1); // Module count
+			writer.Write7BitEncoded(fakeStartAddress); // Start
+			writer.Write7BitEncoded(info.size); // Size
+			writer.WriteString(symbolFile); // Name
+
+			BinaryReader reader(writer.GetData(), 0, writer.GetPosition());
+			StringBuffer<> callstack;
+
+			StringView search[] = { 
+			#if PLATFORM_WINDOWS
+				TCV("srv*C:\\Symbols*https://msdl.microsoft.com/download/symbols"),
+			#endif
+				{}
+			};
+
+			ParseCallstackInfo(callstack, reader, nullptr, search, false);
+			logger.Log(LogEntryType_Info, callstack.data, callstack.count);
+			return true;
+		}
+
 		if (exit)
 			return true;
 
diff --git a/Engine/Source/Programs/UnrealBuildAccelerator/Common/Private/UbaLogger.cpp b/Engine/Source/Programs/UnrealBuildAccelerator/Common/Private/UbaLogger.cpp
--- a/Engine/Source/Programs/UnrealBuildAccelerator/Common/Private/UbaLogger.cpp
+++ b/Engine/Source/Programs/UnrealBuildAccelerator/Common/Private/UbaLogger.cpp
@@ -503,7 +503,7 @@ namespace uba
 	List<ContentionTracker>& GetContentionTrackerList();
 	#endif
 
-	void ParseCallstackInfo(StringBufferBase& out, BinaryReader& reader, const tchar* executable, const StringView* searchPaths)
+	void ParseCallstackInfo(StringBufferBase& out, BinaryReader& reader, const tchar* executable, const StringView* searchPaths, bool writeHeader)
 	{
 		#if PLATFORM_WINDOWS
 		bool isRunningWine = reader.ReadBool();
@@ -531,7 +531,8 @@ namespace uba
 			mod.handled = false;
 		}
 
-		out.Append(TCV("\n CALLSTACK")).Append(isRunningWine ? TCV(" (Wine)") : TCV("")).Append(':');
+		if (writeHeader)
+			out.Append(TCV("\n CALLSTACK")).Append(isRunningWine ? TCV(" (Wine)") : TCV("")).Append(':');
 
 		if (entries.empty())
 		{
@@ -577,7 +578,8 @@ namespace uba
 
 		for (auto& entry : entries)
 		{
-			out.Append(TCV("\n   "));
+			if (writeHeader)
+				out.Append(TCV("\n   "));
 			if (entry.moduleIndex == ~0u)
 			{
 				out.Append(TCV("<Unknown>"));
diff --git a/Engine/Source/Programs/UnrealBuildAccelerator/Common/Public/UbaLogger.h b/Engine/Source/Programs/UnrealBuildAccelerator/Common/Public/UbaLogger.h
--- a/Engine/Source/Programs/UnrealBuildAccelerator/Common/Public/UbaLogger.h
+++ b/Engine/Source/Programs/UnrealBuildAccelerator/Common/Public/UbaLogger.h
@@ -97,7 +97,7 @@ namespace uba
 	Logger* StopDebugLogger(Logger* logger);
 	#endif
 
-	void ParseCallstackInfo(StringBufferBase& out, BinaryReader& reader, const tchar* executable, const StringView* searchPaths);
+	void ParseCallstackInfo(StringBufferBase& out, BinaryReader& reader, const tchar* executable, const StringView* searchPaths, bool writeHeader = true);
 	void PrintContentionSummary(class Logger& logger);
 
 	struct CallstackInfo { Vector<u8> data; Vector<u32> threadIds; TString desc; };
