[UBA]
* Changed logic waiting to stop connection to allow stalls up to 60 seconds but will print a warning after 5s

diff --git a/Engine/Source/Programs/UnrealBuildAccelerator/Common/Private/UbaNetworkServer.cpp b/Engine/Source/Programs/UnrealBuildAccelerator/Common/Private/UbaNetworkServer.cpp
--- a/Engine/Source/Programs/UnrealBuildAccelerator/Common/Private/UbaNetworkServer.cpp
+++ b/Engine/Source/Programs/UnrealBuildAccelerator/Common/Private/UbaNetworkServer.cpp
@@ -160,18 +160,34 @@ namespace uba
 		{
 			Disconnect(TC("Stop"));
 
-			u64 startTimer = GetTime();
+			u64 startTime = GetTime();
+			bool waitReported = false;
 			while (m_activeWorkerCount)
 			{
-				if (TimeToMs(GetTime() - startTimer) > 3000)
+				u64 waitTime = GetTime() - startTime;
+				u64 waitTimeMs = TimeToMs(waitTime);
+				if (waitTimeMs < 5*1000)
 				{
-					m_server.m_logger.Error(TC("Connection %u has waited 3 seconds to stop... something is stuck (Active worker count: %u)"), m_id, m_activeWorkerCount.load());
+					Sleep(5);
+					continue;
+				}
+
+				if (!waitReported)
+				{
+					waitReported = true;
+					m_server.m_logger.Warning(TC("Connection %u has waited %s to stop... something is stuck (Active worker count: %u)"), m_id, TimeToText(waitTime).str, m_activeWorkerCount.load());
 					PrintAllCallstacks(m_server.m_logger);
-					return false;
 				}
-				Sleep(1);
+
+				if (waitTimeMs > 60*1000) // Wait 60 seconds
+					return false;
+
+				Sleep(200);
 			}
 
+			if (waitReported)
+				m_server.m_logger.Info(TC("Connection %u wait for active workers ended after %s"), m_id, TimeToText(GetTime() - startTime).str);
+
 			if (!m_disconnectCallbackCalled.IsSet(30*1000)) // This should never time out!
 			{
 				m_server.m_logger.Warning(TC("Disconnect callback event timed out. This should never happen!!"));
