[UBA]
* Fixed so crypto is enabled if -http is provided to cache service (since http is only used to add crypto keys right now)

diff --git a/Engine/Source/Programs/UnrealBuildAccelerator/CacheService/Private/UbaCacheService.cpp b/Engine/Source/Programs/UnrealBuildAccelerator/CacheService/Private/UbaCacheService.cpp
--- a/Engine/Source/Programs/UnrealBuildAccelerator/CacheService/Private/UbaCacheService.cpp
+++ b/Engine/Source/Programs/UnrealBuildAccelerator/CacheService/Private/UbaCacheService.cpp
@@ -127,6 +127,7 @@ namespace uba
 		bool allowSave = true;
 		bool signalHandlerEnabled = true;
 		bool shouldCrash = false;
+		bool requiresCrypto = false;
 		u32 expirationTimeSeconds = DefaultExpiration;
 		u32 maxWorkerCount = DefaultWorkerCount;
 		u32 reportIntervalSeconds = DefaultReportIntervalSeconds;
@@ -188,6 +189,7 @@ namespace uba
 			{
 				if (!value.Parse(httpPort))
 					httpPort = 80;
+				requiresCrypto = true; // We force crypto connections when enabling http since http is only used to add crypto keys right now
 			}
 			else if (name.Equals(TCV("-httpcrypto")))
 			{
@@ -545,7 +547,7 @@ namespace uba
 			auto stopListen = MakeGuard([&]() { networkBackend.StopListen(); });
 			auto stopServer = MakeGuard([&]() { networkServer.DisconnectClients(); });
 
-			if (!networkServer.StartListen(networkBackend, port, listenIp.data))
+			if (!networkServer.StartListen(networkBackend, port, listenIp.data, requiresCrypto))
 				return false;
 
 			u64 lastUpdateTime = GetTime();
