[UBA]
* Added support for http command "addsession" which can take parameters crypto,readonly and expiration. crypto is the 32 character (128bit) string, readonly=true means that the connection can not write to cache. and expiration is in seconds how long the crypto key is available before it goes away
* Fixed so UbaCli -crypto also affects cache connection if -cache is used

diff --git a/Engine/Source/Programs/UnrealBuildAccelerator/CacheService/Private/UbaCacheService.cpp b/Engine/Source/Programs/UnrealBuildAccelerator/CacheService/Private/UbaCacheService.cpp
--- a/Engine/Source/Programs/UnrealBuildAccelerator/CacheService/Private/UbaCacheService.cpp
+++ b/Engine/Source/Programs/UnrealBuildAccelerator/CacheService/Private/UbaCacheService.cpp
@@ -72,6 +72,7 @@ namespace uba
 #endif
 		logger.Info(TC(" Example of how to register crypto key to cache server (when -http=80 is provided)"));
 		logger.Info(TC("   curl http://localhost/addcrypto?3f58aa57466db9999213456789123445"));
+		logger.Info(TC("   curl http://localhost/addsession?crypto=3f58aa57466db9999213456789123445,readonly=true,expiration=60"));
 		logger.Info(TC(""));
 		return false;
 	}
@@ -508,37 +509,85 @@ namespace uba
 
 		if (httpPort)
 		{
-			httpServer.AddCommandHandler([&](const tchar* command, tchar* arguments)
+			httpServer.AddCommandHandler([&](StringView command, StringBufferBase& arguments)
 				{
-					if (!Equals(command, TC("addcrypto")))
-						return "Unknown command ('addcrypto' only available)";
+					u8 crypto128Data[16];
 					u64 expirationSeconds = 60;
-					if (tchar* comma  = TStrchr(arguments, ','))
+					u64 userData = 0;
+					auto registerCrypto = [&]()
+						{
+							u64 expirationTime = GetTime() + MsToTime(expirationSeconds*1000);
+							if (httpCryptos.empty())
+								networkServer.RegisterCryptoKey(crypto128Data, expirationTime, userData);
+							else
+							{
+								for (auto& crypto : httpCryptos)
+								{
+									u8 decryptedKey[16];
+									memcpy(decryptedKey, crypto128Data, 16);
+									if (!Crypto::Decrypt(logger, crypto, decryptedKey, 16))
+										return "Failed to decrypt key from http addcrypto";
+									networkServer.RegisterCryptoKey(decryptedKey, expirationTime, userData);
+								}
+							}
+							return (const char*)nullptr;
+						};
+
+
+					if (command.Equals(TCV("addcrypto")))
 					{
-						*comma = 0;
-						if (!Parse(expirationSeconds, comma+1, TStrlen(comma+1)))
-							return "Failed to parse expiration seconds";
+						if (tchar* comma  = TStrchr(arguments.data, ','))
+						{
+							*comma = 0;
+							if (!Parse(expirationSeconds, comma+1, TStrlen(comma+1)))
+								return "Failed to parse expiration seconds";
+						}
+						if (!CryptoFromString(crypto128Data, 16, arguments.data))
+							return "Failed to read crypto argument (Needs to be 32 characters long)";
+						return registerCrypto();
 					}
-
-					u8 crypto128Data[16];
-					if (!CryptoFromString(crypto128Data, 16, arguments))
-						return "Failed to read crypto argument (Needs to be 32 characters long)";
-					u64 expirationTime = GetTime() + MsToTime(expirationSeconds*1000);
-
-					if (httpCryptos.empty())
-						networkServer.RegisterCryptoKey(crypto128Data, expirationTime);
-					else
+					if (command.Equals(TCV("addsession")))
 					{
-						for (auto& crypto : httpCryptos)
+						tchar* argument = arguments.data;
+						while (true)
 						{
-							u8 decryptedKey[16];
-							memcpy(decryptedKey, crypto128Data, 16);
-							if (!Crypto::Decrypt(logger, crypto, decryptedKey, 16))
-								return "Failed to decrypt key from http addcrypto";
-							networkServer.RegisterCryptoKey(decryptedKey, expirationTime);
+							tchar* comma = TStrchr(argument, ',');
+							if (comma)
+								*comma = 0;
+
+							tchar* equalSign = TStrchr(argument, '=');
+							if (!equalSign)
+								return "Bad argument format";
+							*equalSign = 0;
+							tchar* value = equalSign + 1;
+
+							if (Equals(argument, TC("crypto")))
+							{
+								if (!CryptoFromString(crypto128Data, 16, value))
+									return "Failed to read crypto argument (Needs to be 32 characters long)";
+							}
+							else if (Equals(argument, TC("readonly")))
+							{
+								userData = (!Equals(value, TC("false")) && !Equals(value, TC("0"))) ? 1 : 0;
+							}
+							else if (Equals(argument, TC("expiration")))
+							{
+								if (!Parse(expirationSeconds, value, TStrlen(value+1)))
+									return "Failed to parse expiration seconds";
+							}
+							else
+							{
+								return "Unknown argument to addsession. (supports crypto,readonly and expiration)";
+							}
+
+							if (!comma)
+								break;
+							argument = comma + 1;
 						}
+						return registerCrypto();
 					}
-					return (const char*)nullptr;
+	
+					return "Unknown command (existing commands: addcrypto, addsession)";
 				});
 			httpServer.StartListen(httpPort);
 		}
diff --git a/Engine/Source/Programs/UnrealBuildAccelerator/CacheService/Private/UbaHttpServer.cpp b/Engine/Source/Programs/UnrealBuildAccelerator/CacheService/Private/UbaHttpServer.cpp
--- a/Engine/Source/Programs/UnrealBuildAccelerator/CacheService/Private/UbaHttpServer.cpp
+++ b/Engine/Source/Programs/UnrealBuildAccelerator/CacheService/Private/UbaHttpServer.cpp
@@ -63,10 +63,10 @@ namespace uba
 		if (argumentsEnd - request >= 256)
 			return sendFail("arguments too long");
 
-		StringBuffer<256> arguments;
+		StringBuffer<512> arguments;
 		arguments.Append(request, u32(argumentsEnd - request));
 
-		if (const char* result = c.server.m_handler(command.data, arguments.data))
+		if (const char* result = c.server.m_handler(command, arguments))
 			return sendFail(result);
 
 		snprintf(response, 1024, "HTTP/1.1 200 OK\r\n""Content-Type: text/plain\r\n""\r\n");
diff --git a/Engine/Source/Programs/UnrealBuildAccelerator/CacheService/Private/UbaHttpServer.h b/Engine/Source/Programs/UnrealBuildAccelerator/CacheService/Private/UbaHttpServer.h
--- a/Engine/Source/Programs/UnrealBuildAccelerator/CacheService/Private/UbaHttpServer.h
+++ b/Engine/Source/Programs/UnrealBuildAccelerator/CacheService/Private/UbaHttpServer.h
@@ -12,7 +12,7 @@ namespace uba
 	public:
 		HttpServer(LogWriter& logWriter, NetworkBackend& backend, const tchar* name = TC("UbaHttpServer"));
 
-		using CommandHandler = Function<const char*(const tchar* command, tchar* arguments)>;
+		using CommandHandler = Function<const char*(StringView command, StringBufferBase& arguments)>;
 		void AddCommandHandler(const CommandHandler& handler);
 
 		bool StartListen(u16 port);
diff --git a/Engine/Source/Programs/UnrealBuildAccelerator/Cli/Private/UbaCli.cpp b/Engine/Source/Programs/UnrealBuildAccelerator/Cli/Private/UbaCli.cpp
--- a/Engine/Source/Programs/UnrealBuildAccelerator/Cli/Private/UbaCli.cpp
+++ b/Engine/Source/Programs/UnrealBuildAccelerator/Cli/Private/UbaCli.cpp
@@ -489,6 +489,11 @@ namespace uba
 					arguments += TC("\"");
 				return true;
 			}
+
+			if (commandType == CommandType_None)
+			{
+				return parseOption(name, value);
+			}
 			if (commandType != CommandType_NotSet)
 			{
 				application = arg;
@@ -953,9 +958,9 @@ namespace uba
 		if (!ctorSuccess)
 			return false;
 
+		u8 crypto128Data[16];
 		if (!crypto.empty())
 		{
-			u8 crypto128Data[16];
 			if (!CryptoFromString(crypto128Data, 16, crypto.c_str()))
 				return logger.Error(TC("Failed to parse crypto key %s"), crypto.c_str());
 			networkServer.RegisterCryptoKey(crypto128Data);
@@ -1002,7 +1007,10 @@ namespace uba
 
 		auto CreateCacheClient = [&]()
 			{
-				auto nc = new NetworkClient(ctorSuccess, {logWriter});
+				NetworkClientCreateInfo nci(logWriter);
+				if (!crypto.empty())
+					nci.cryptoKey128 = crypto128Data;
+				auto nc = new NetworkClient(ctorSuccess, nci);
 				cacheClient = new CacheClient({logWriter, storageServer, *nc, sessionServer});
 			};
 
diff --git a/Engine/Source/Programs/UnrealBuildAccelerator/Common/Private/UbaCacheServer.cpp b/Engine/Source/Programs/UnrealBuildAccelerator/Common/Private/UbaCacheServer.cpp
--- a/Engine/Source/Programs/UnrealBuildAccelerator/Common/Private/UbaCacheServer.cpp
+++ b/Engine/Source/Programs/UnrealBuildAccelerator/Common/Private/UbaCacheServer.cpp
@@ -12,6 +12,10 @@
 #include "UbaTrace.h"
 #include <algorithm>
 
+#if defined(__clang__) && PLATFORM_WINDOWS && (defined(_M_X64) || defined(_M_IX86))
+#include <x86intrin.h>
+#endif
+
 //#include <oodle2.h>
 
 // TODO
@@ -1476,6 +1480,11 @@ namespace uba
 		return Min(workerCountToUse, u32(m_buckets.size()));
 	}
 
+	bool IsReadOnly(const ConnectionInfo& connectionInfo)
+	{
+		return connectionInfo.GetCryptoUserData() != 0;
+	}
+
 	UBA_NOINLINE bool CacheServer::HandleMessage(const ConnectionInfo& connectionInfo, u8 messageType, BinaryReader& reader, BinaryWriter& writer)
 	{
 		if (messageType != CacheMessageType_Connect && m_isRunningMaintenance)
@@ -1511,6 +1520,8 @@ namespace uba
 			logStr.Appendf(TC("Client %u connected"), connectionInfo.GetId());
 			if (clientVersion != CacheNetworkVersion)
 				logStr.Appendf(TC(" (v%u)"), clientVersion);
+			if (IsReadOnly(connectionInfo))
+				logStr.Appendf(TC(" (readonly)"));
 			if (!hint.empty())
 				logStr.Appendf(TC(" %s"), hint.c_str());
 			m_logger.Log(LogEntryType_Info, logStr);
@@ -1519,16 +1530,22 @@ namespace uba
 		}
 		case CacheMessageType_StorePathTable:
 		{
+			if (IsReadOnly(connectionInfo))
+				return false;
 			GetConnectionBucket(connectionInfo, reader).pathTable.ReadMem(reader, false);
 			return true;
 		}
 		case CacheMessageType_StoreCasTable:
 		{
+			if (IsReadOnly(connectionInfo))
+				return false;
 			GetConnectionBucket(connectionInfo, reader).casKeyTable.ReadMem(reader, false);
 			return true;
 		}
 		case CacheMessageType_StoreEntry:
 		{
+			if (IsReadOnly(connectionInfo))
+				return false;
 			u32 clientVersion;
 			auto& bucket = GetConnectionBucket(connectionInfo, reader, &clientVersion);
 			return HandleStoreEntry(bucket, reader, writer, clientVersion, connectionInfo.GetId());
@@ -1643,7 +1660,7 @@ namespace uba
 			if (!isInput)
 				continue;
 
-			u32 casKeyOffset;
+			u32 casKeyOffset = 0;
 			CasKey casKey;
 			if (!CreateCasKeyOffset(casKeyOffset, casKey, offset, inputs, TC("input")))
 				return false;
@@ -1753,7 +1770,7 @@ namespace uba
 		{
 			u32 outputOffset = u32(reader.Read7BitEncoded());
 
-			u32 casKeyOffset;
+			u32 casKeyOffset = 0;
 			CasKey casKey;
 			if (!CreateCasKeyOffset(casKeyOffset, casKey, outputOffset, outputs, TC("output")))
 				return false;
diff --git a/Engine/Source/Programs/UnrealBuildAccelerator/Common/Private/UbaNetworkServer.cpp b/Engine/Source/Programs/UnrealBuildAccelerator/Common/Private/UbaNetworkServer.cpp
--- a/Engine/Source/Programs/UnrealBuildAccelerator/Common/Private/UbaNetworkServer.cpp
+++ b/Engine/Source/Programs/UnrealBuildAccelerator/Common/Private/UbaNetworkServer.cpp
@@ -260,6 +260,7 @@ namespace uba
 						continue;
 					keyGuard.Cancel();
 					conn.m_cryptoKey = key;
+					conn.m_cryptoUserData = entry.userData;
 					break;
 				}
 				if (conn.m_cryptoKey == InvalidCryptoKey)
@@ -504,6 +505,8 @@ namespace uba
 		Timer m_encryptTimer;
 		Timer m_decryptTimer;
 
+		u64 m_cryptoUserData = 0;
+
 		Connection(const Connection& o) = delete;
 		void operator=(const Connection& o) = delete;
 	};
@@ -532,6 +535,12 @@ namespace uba
 		#endif
 	}
 
+	u64 ConnectionInfo::GetCryptoUserData() const
+	{
+		auto& conn = *(NetworkServer::Connection*)internalData;
+		return conn.m_cryptoUserData;
+	}
+
 	bool ConnectionInfo::ShouldDisconnect() const
 	{
 		auto& conn = *(NetworkServer::Connection*)internalData;
@@ -806,13 +815,13 @@ namespace uba
 		m_workersEnabled = true;
 	}
 
-	bool NetworkServer::RegisterCryptoKey(const u8* cryptoKey128, u64 expirationTime)
+	bool NetworkServer::RegisterCryptoKey(const u8* cryptoKey128, u64 expirationTime, u64 cryptoUserData)
 	{
 		CryptoKey key = Crypto::CreateKey(m_logger, cryptoKey128);
 		if (key == InvalidCryptoKey)
 			return false;
 		SCOPED_FUTEX(m_cryptoKeysLock, lock);
-		m_cryptoKeys.push_back(CryptoEntry{key, expirationTime});
+		m_cryptoKeys.push_back(CryptoEntry{key, expirationTime, cryptoUserData});
 		return true;
 	}
 
diff --git a/Engine/Source/Programs/UnrealBuildAccelerator/Common/Public/UbaNetworkServer.h b/Engine/Source/Programs/UnrealBuildAccelerator/Common/Public/UbaNetworkServer.h
--- a/Engine/Source/Programs/UnrealBuildAccelerator/Common/Public/UbaNetworkServer.h
+++ b/Engine/Source/Programs/UnrealBuildAccelerator/Common/Public/UbaNetworkServer.h
@@ -22,6 +22,7 @@ namespace uba
 		const Guid& GetUid() const;
 		u32 GetId() const;
 		bool GetName(StringBufferBase& out) const;
+		u64 GetCryptoUserData() const;
 		bool ShouldDisconnect() const;
 		void* internalData = nullptr;
 	};
@@ -57,7 +58,7 @@ namespace uba
 		void DisallowNewClients();	// Disallow new clients to connect but old clients can still create more connections
 		void DisconnectClients();	// Disconnect all active connections
 
-		bool RegisterCryptoKey(const u8* cryptoKey128, u64 expirationTime = ~u64(0));
+		bool RegisterCryptoKey(const u8* cryptoKey128, u64 expirationTime = ~u64(0), u64 cryptoUserData = 0);
 		void SetClientsConfig(const Config& config);
 
 		bool AddClient(NetworkBackend& backend, const tchar* ip, u16 port = DefaultPort, const u8* cryptoKey128 = nullptr); // Adds a client that server will create one or more connections to (note this will return before we know if it was a success or not)
@@ -132,6 +133,7 @@ namespace uba
 		{
 			CryptoKey key;
 			u64 expirationTime;
+			u64 userData;
 		};
 		Futex m_cryptoKeysLock;
 		List<CryptoEntry> m_cryptoKeys;
