[UBA]
* Added UbaCli -uniqueid which will make sure session id is unique and end up in its own directory. Will also automatically delete session dir when exiting (in debug it is not deleted)
* Added DirectoryCache::CreateDirectory outAlreadyExists that is true if directory already exists
* Moved Fputs from UbaLogger.cpp to UbaPlatform.h and renamed it to TFputs
* UbaSession uniqueid is now truly unique. if it fails to create directory yymmdd_hhmmss it adds _1, 2, 3, 4 and so on and retry
* Minor test fix for crypto

diff --git a/Engine/Source/Programs/UnrealBuildAccelerator/Cli/Private/UbaCli.cpp b/Engine/Source/Programs/UnrealBuildAccelerator/Cli/Private/UbaCli.cpp
--- a/Engine/Source/Programs/UnrealBuildAccelerator/Cli/Private/UbaCli.cpp
+++ b/Engine/Source/Programs/UnrealBuildAccelerator/Cli/Private/UbaCli.cpp
@@ -76,6 +76,7 @@ namespace uba
 		logger.Info(TC("   -port=[<host>:]<port>   The ip/name and port (default: %u) of the machine we want to help"), DefaultPort);
 		logger.Info(TC("   -log                    Log all processes detouring information to file (only works with debug builds)"));
 		logger.Info(TC("   -quiet                  Does not output any logging in console except errors"));
+		logger.Info(TC("   -uniqueid               Makes sure to always create session dir with a unique name"));
 		logger.Info(TC("   -loop=<count>           Loop the commandline <count> number of times. Will exit when/if it fails"));
 		logger.Info(TC("   -workdir=<dir>          Working directory"));
 		logger.Info(TC("   -config=<file>          Config file that contains options for various systems"));
@@ -164,6 +165,7 @@ namespace uba
 		bool storeCompressed = true;
 		bool disableCustomAllocator = false;
 		bool quiet = false;
+		bool useUniqueId = false;
 		bool checkCas = false;
 		bool checkCas2 = false;
 		bool checkCloud = false;
@@ -301,6 +303,10 @@ namespace uba
 				{
 					quiet = true;
 				}
+				else if (name.Equals(TCV("-uniqueid")))
+				{
+					useUniqueId = true;
+				}
 				else if (name.Equals(TCV("-nocustomalloc")))
 				{
 					disableCustomAllocator = true;
@@ -979,7 +985,7 @@ namespace uba
 		auto destroyStorage = MakeGuard([&]() { delete &storageServer; });
 
 		SessionServerCreateInfo info(storageServer, networkServer, logWriter);
-		info.useUniqueId = useScheduler;
+		info.useUniqueId = useUniqueId || useScheduler;
 		info.traceEnabled = cacheCommand.empty() || writeCacheSummary;
 		info.detailedTrace = detailedTrace;
 		info.traceChildProcesses = traceChildProcesses;
@@ -996,11 +1002,18 @@ namespace uba
 		#endif
 		//info.remoteTraceEnabled = true;
 
-		info.deleteSessionsOlderThanSeconds = 1;
+		info.deleteSessionsOlderThanSeconds = info.useUniqueId ? 0 : 1;
 		info.Apply(config);
 
 		SessionServer& sessionServer = *new SessionServer(info);
-		auto destroySession = MakeGuard([&]() { delete &sessionServer; });
+		auto destroySession = MakeGuard([&]()
+			{
+				#if !UBA_DEBUG
+				if (info.useUniqueId)
+					DeleteAllFiles(logger, sessionServer.GetSessionDir());
+				#endif
+				delete &sessionServer;
+			});
 
 		CacheClient* cacheClient = nullptr;
 		auto ccg = MakeGuard([&]() { if (!cacheClient) return; auto& nc = cacheClient->GetClient(); nc.Disconnect(); delete cacheClient; delete &nc; });
@@ -1108,11 +1121,14 @@ namespace uba
 				}
 			}
 
-
 			pinfo.logFile = logFile.data;
-			pinfo.logLineUserData = &logger;
 			if (enableStdOut)
-				pinfo.logLineFunc = [](void* userData, const tchar* line, u32 length, LogEntryType type) { ((Logger*)userData)->Log(type, line, length); };
+				pinfo.logLineFunc = [](void* userData, const tchar* line, u32 length, LogEntryType type)
+				{
+					auto s = type == LogEntryType_Error ? stderr : stdout;
+					TFputs(line, s);
+					TFputs(TC("\n"), s);
+				};
 			if (populateCache)
 				pinfo.trackInputs = true;
 			logger.Info(TC("Running %s %s"), app.c_str(), arg.c_str());
diff --git a/Engine/Source/Programs/UnrealBuildAccelerator/Common/Private/UbaFile.cpp b/Engine/Source/Programs/UnrealBuildAccelerator/Common/Private/UbaFile.cpp
--- a/Engine/Source/Programs/UnrealBuildAccelerator/Common/Private/UbaFile.cpp
+++ b/Engine/Source/Programs/UnrealBuildAccelerator/Common/Private/UbaFile.cpp
@@ -969,8 +969,11 @@ namespace uba
 #endif
 	}
 
-	bool DirectoryCache::CreateDirectory(Logger& logger, const tchar* dir)
+	bool DirectoryCache::CreateDirectory(Logger& logger, const tchar* dir, bool* outAlreadyExists)
 	{
+		bool throwaway;
+		bool& alreadyExists = outAlreadyExists ? *outAlreadyExists : throwaway;
+
 		u64 dirLen = TStrlen(dir);
 		if (dir[dirLen - 1] == PathSeparator)
 			--dirLen;
@@ -982,13 +985,22 @@ namespace uba
 		lock.Leave();
 		SCOPED_FUTEX(cd.lock, dirLock);
 		if (cd.handled)
+		{
+			alreadyExists = true;
 			return true;
+		}
 		cd.handled = true;
 		if (uba::CreateDirectoryW(dir))
+		{
+			alreadyExists = false;
 			return true;
+		}
 		u32 lastError = GetLastError();
 		if (lastError == ERROR_ALREADY_EXISTS)
+		{
+			alreadyExists = true;
 			return true;
+		}
 		if (lastError != ERROR_PATH_NOT_FOUND)
 			return logger.Error(TC("Failed to create directory %s (%s)"), dir, LastErrorToText(lastError).data);
 
@@ -1001,9 +1013,18 @@ namespace uba
 			return false;
 		if (!CreateDirectory(logger, temp))
 			return false;
-		if (!uba::CreateDirectoryW(dir))
-			return logger.Error(TC("Failed to create directory %s (%s)"), dir, LastErrorToText().data);
-		return true;
+		if (uba::CreateDirectoryW(dir))
+		{
+			alreadyExists = false;
+			return true;
+		}
+		lastError = GetLastError();
+		if (lastError == ERROR_ALREADY_EXISTS)
+		{
+			alreadyExists = true;
+			return true;
+		}
+		return logger.Error(TC("Failed to create directory %s (%s)"), dir, LastErrorToText(lastError).data);
 	}
 
 	void DirectoryCache::Clear()
diff --git a/Engine/Source/Programs/UnrealBuildAccelerator/Common/Private/UbaLogger.cpp b/Engine/Source/Programs/UnrealBuildAccelerator/Common/Private/UbaLogger.cpp
--- a/Engine/Source/Programs/UnrealBuildAccelerator/Common/Private/UbaLogger.cpp
+++ b/Engine/Source/Programs/UnrealBuildAccelerator/Common/Private/UbaLogger.cpp
@@ -11,9 +11,7 @@
 #if PLATFORM_WINDOWS
 #include <io.h>
 #include <dbghelp.h>
-#define Fputs fputws
 #else
-#define Fputs fputs
 #define _vsnwprintf_s(buffer,capacity,count,format,args) vsnprintf(buffer, capacity, format, args) // TODO: This will overflow
 #define _vscwprintf(format, args) vsnprintf(0, 0, format, args)
 #endif
@@ -64,8 +62,8 @@ namespace uba
 			return;
 		}
 
-		Fputs(sb.data, stdout);
-		Fputs(TC("\n"), stdout);
+		TFputs(sb.data, stdout);
+		TFputs(TC("\n"), stdout);
 		fflush(stdout);
 
 #if PLATFORM_WINDOWS
@@ -311,8 +309,8 @@ namespace uba
 		{
 			if (prefixLen)
 			{
-				Fputs(prefix, stdout);
-				Fputs(TC(" - "), stdout);
+				TFputs(prefix, stdout);
+				TFputs(TC(" - "), stdout);
 			}
 			_putws(str);
 		}
@@ -344,11 +342,11 @@ namespace uba
 #else
 		if (prefixLen)
 		{
-			Fputs(prefix, stdout);
-			Fputs(TC(" - "), stdout);
+			TFputs(prefix, stdout);
+			TFputs(TC(" - "), stdout);
 		}
-		Fputs(str, stdout);
-		Fputs(TC("\n"), stdout);
+		TFputs(str, stdout);
+		TFputs(TC("\n"), stdout);
 #endif
 	}
 
diff --git a/Engine/Source/Programs/UnrealBuildAccelerator/Common/Private/UbaSession.cpp b/Engine/Source/Programs/UnrealBuildAccelerator/Common/Private/UbaSession.cpp
--- a/Engine/Source/Programs/UnrealBuildAccelerator/Common/Private/UbaSession.cpp
+++ b/Engine/Source/Programs/UnrealBuildAccelerator/Common/Private/UbaSession.cpp
@@ -1523,22 +1523,38 @@ namespace uba
 	,	m_writeFilesNoBufferingMin(u64(info.writeFilesNoBufferingMinMb)*1024*1024)
 	,	m_dependencyCrawler(m_logger, workManager)
 	{
+		UBA_ASSERTF(info.rootDir && *info.rootDir, TC("No root dir set when creating session"));
+		m_rootDir.count = GetFullPathNameW(info.rootDir, m_rootDir.capacity, m_rootDir.data, NULL);
+		m_rootDir.Replace('/', PathSeparator).EnsureEndsWithSlash();
+
+		m_sessionDir.Append(m_rootDir).Append(TCV("sessions")).Append(PathSeparator);
 		if (info.useUniqueId)
 		{
-			time_t rawtime;
-			time(&rawtime);
-			tm ti;
-			localtime_s(&ti, &rawtime);
-			m_id.Appendf(TC("%02i%02i%02i_%02i%02i%02i"), ti.tm_year - 100,ti.tm_mon+1,ti.tm_mday, ti.tm_hour, ti.tm_min, ti.tm_sec);
+			u32 retryIndex = 0;
+			while (true)
+			{
+				time_t rawtime;
+				time(&rawtime);
+				tm ti;
+				localtime_s(&ti, &rawtime);
+				m_id.Appendf(TC("%02i%02i%02i_%02i%02i%02i"), ti.tm_year - 100,ti.tm_mon+1,ti.tm_mday, ti.tm_hour, ti.tm_min, ti.tm_sec);
+				if (retryIndex > 0)
+					m_id.Append('_').AppendValue(retryIndex);
+				m_sessionDir.Append(m_id);
+				bool alreadyExists;
+				if (m_storage.CreateDirectory(m_sessionDir.data, &alreadyExists) && !alreadyExists)
+					break;
+				m_sessionDir.Resize(m_sessionDir.count - m_id.count);
+				m_id.Clear();
+				++retryIndex;
+			}
 		}
 		else
 		{
 			m_id.Append(TCV("Debug"));
+			m_sessionDir.Append(m_id);
 		}
-
-		UBA_ASSERTF(info.rootDir && *info.rootDir, TC("No root dir set when creating session"));
-		m_rootDir.count = GetFullPathNameW(info.rootDir, m_rootDir.capacity, m_rootDir.data, NULL);
-		m_rootDir.Replace('/', PathSeparator).EnsureEndsWithSlash();
+		m_sessionDir.Append(PathSeparator);
 
 		m_runningRemote = runningRemote;
 		m_disableCustomAllocator = info.disableCustomAllocator;
@@ -1590,7 +1606,6 @@ namespace uba
 				});
 		}
 
-		m_sessionDir.Append(m_rootDir).Append(TCV("sessions")).Append(PathSeparator).Append(m_id).Append(PathSeparator);
 		m_sessionBinDir.Append(m_sessionDir).Append(TCV("bin"));
 		m_sessionOutputDir.Append(m_sessionDir).Append(TCV("output"));
 		m_sessionLogDir.Append(m_sessionDir).Append(TCV("log"));
@@ -2172,6 +2187,11 @@ namespace uba
 		return m_rootDir.data;
 	}
 
+	const tchar* Session::GetSessionDir()
+	{
+		return m_sessionDir.data;
+	}
+
 	u32 Session::CreateProcessId()
 	{
 		return ++m_processIdCounter;
diff --git a/Engine/Source/Programs/UnrealBuildAccelerator/Common/Private/UbaStorage.cpp b/Engine/Source/Programs/UnrealBuildAccelerator/Common/Private/UbaStorage.cpp
--- a/Engine/Source/Programs/UnrealBuildAccelerator/Common/Private/UbaStorage.cpp
+++ b/Engine/Source/Programs/UnrealBuildAccelerator/Common/Private/UbaStorage.cpp
@@ -2441,9 +2441,9 @@ namespace uba
 		return w.DecompressFileToMemory(fileName, fileHandle, dest, decompressedSize, writeHint, fileStartOffset);
 	}
 
-	bool StorageImpl::CreateDirectory(const tchar* dir)
+	bool StorageImpl::CreateDirectory(const tchar* dir, bool* outAlreadyExists)
 	{
-		return m_dirCache.CreateDirectory(m_logger, dir);
+		return m_dirCache.CreateDirectory(m_logger, dir, outAlreadyExists);
 	}
 
 	bool StorageImpl::DeleteCasForFile(const tchar* file)
diff --git a/Engine/Source/Programs/UnrealBuildAccelerator/Common/Public/UbaFile.h b/Engine/Source/Programs/UnrealBuildAccelerator/Common/Public/UbaFile.h
--- a/Engine/Source/Programs/UnrealBuildAccelerator/Common/Public/UbaFile.h
+++ b/Engine/Source/Programs/UnrealBuildAccelerator/Common/Public/UbaFile.h
@@ -97,7 +97,7 @@ namespace uba
 	class DirectoryCache
 	{
 	public:
-		bool CreateDirectory(Logger& logger, const tchar* dir);
+		bool CreateDirectory(Logger& logger, const tchar* dir, bool* outAlreadyExists = nullptr);
 		void Clear();
 
 	private:
diff --git a/Engine/Source/Programs/UnrealBuildAccelerator/Common/Public/UbaSession.h b/Engine/Source/Programs/UnrealBuildAccelerator/Common/Public/UbaSession.h
--- a/Engine/Source/Programs/UnrealBuildAccelerator/Common/Public/UbaSession.h
+++ b/Engine/Source/Programs/UnrealBuildAccelerator/Common/Public/UbaSession.h
@@ -92,6 +92,7 @@ namespace uba
 		const ApplicationRules* GetRules(const ProcessStartInfo& si); // Get application rules used for process
 		const tchar* GetTempPath(); // Path for temp files used for current session
 		const tchar* GetRootDir(); // Root dir for session files
+		const tchar* GetSessionDir(); // Session dir
 		bool ShouldStoreIntermediateFilesCompressed() const { return m_storeIntermediateFilesCompressed; }
 		bool HasDetailedtrace() const { return m_detailedTrace; }
 
diff --git a/Engine/Source/Programs/UnrealBuildAccelerator/Common/Public/UbaStorage.h b/Engine/Source/Programs/UnrealBuildAccelerator/Common/Public/UbaStorage.h
--- a/Engine/Source/Programs/UnrealBuildAccelerator/Common/Public/UbaStorage.h
+++ b/Engine/Source/Programs/UnrealBuildAccelerator/Common/Public/UbaStorage.h
@@ -34,7 +34,7 @@ namespace uba
 
 		virtual bool DecompressFileToMemory(const tchar* fileName, FileHandle fileHandle, u8* dest, u64 decompressedSize, const tchar* writeHint, u64 fileStartOffset) = 0;
 		virtual bool DecompressMemoryToMemory(const u8* compressedData, u64 compressedSize, u8* writeData, u64 decompressedSize, const tchar* readHint, const tchar* writeHint) = 0;
-		virtual bool CreateDirectory(const tchar* dir) = 0;
+		virtual bool CreateDirectory(const tchar* dir, bool* outAlreadyExists = nullptr) = 0;
 		virtual bool DeleteCasForFile(const tchar* file) = 0;
 		virtual bool GetFileSize(u64& outSize, const CasKey& casKey, const tchar* hint) = 0;
 		virtual bool GetFileData(void* outData, const CasKey& casKey, const tchar* hint) = 0;
@@ -151,7 +151,7 @@ namespace uba
 		virtual void PrintSummary(Logger& logger) override;
 
 		virtual bool DecompressFileToMemory(const tchar* fileName, FileHandle fileHandle, u8* dest, u64 decompressedSize, const tchar* writeHint, u64 fileStartOffset) override;
-		virtual bool CreateDirectory(const tchar* dir) override;
+		virtual bool CreateDirectory(const tchar* dir, bool* outAlreadyExists = nullptr) override;
 		virtual bool DeleteCasForFile(const tchar* file) override;
 		virtual bool GetFileSize(u64& outSize, const CasKey& casKey, const tchar* hint) override;
 		virtual bool GetFileData(void* outData, const CasKey& casKey, const tchar* hint) override;
diff --git a/Engine/Source/Programs/UnrealBuildAccelerator/Core/Public/UbaPlatform.h b/Engine/Source/Programs/UnrealBuildAccelerator/Core/Public/UbaPlatform.h
--- a/Engine/Source/Programs/UnrealBuildAccelerator/Core/Public/UbaPlatform.h
+++ b/Engine/Source/Programs/UnrealBuildAccelerator/Core/Public/UbaPlatform.h
@@ -149,6 +149,7 @@ namespace uba
 	#define TStrcpy_s wcscpy_s
 	#define TStrcat_s wcscat_s
 	#define TStrdup _wcsdup
+	#define TFputs fputws
 	#define UBA_NOINLINE __declspec(noinline)
 	struct FileMappingHandle
 	{
@@ -173,6 +174,7 @@ namespace uba
 	#define TStrcpy_s(a, b, c) strcpy_s(a, b, c)
 	#define TStrcat_s(a, b, c) strcat_s(a, b, c)
 	#define TStrdup strdup
+	#define TFputs fputs
 	#define localtime_s(a,b) localtime_r(b,a)
 	void strcpy_s(tchar* dest, u64 destCapacity, const tchar* source);
 	void strcat_s(tchar* dest, u64 destCapacity, const tchar* source);
diff --git a/Engine/Source/Programs/UnrealBuildAccelerator/Test/Private/UbaTestCrypto.cpp b/Engine/Source/Programs/UnrealBuildAccelerator/Test/Private/UbaTestCrypto.cpp
--- a/Engine/Source/Programs/UnrealBuildAccelerator/Test/Private/UbaTestCrypto.cpp
+++ b/Engine/Source/Programs/UnrealBuildAccelerator/Test/Private/UbaTestCrypto.cpp
@@ -28,7 +28,7 @@ namespace uba
 				if (!Crypto::Encrypt(logger, encryptKey, encryptedData, dataSize))
 					return false;
 
-				if (dataSize > 16 && memcmp(encryptedData, originalData, dataSize) == 0)
+				if (dataSize >= 16 && memcmp(encryptedData, originalData, dataSize) == 0)
 					return false;
 
 				if (!Crypto::Decrypt(logger, decryptKey, encryptedData, dataSize))
