[UBA]
* Handle upper-case in HexToByte conversion function

diff --git a/Engine/Source/Programs/UnrealBuildAccelerator/Core/Public/UbaMemory.h b/Engine/Source/Programs/UnrealBuildAccelerator/Core/Public/UbaMemory.h
--- a/Engine/Source/Programs/UnrealBuildAccelerator/Core/Public/UbaMemory.h
+++ b/Engine/Source/Programs/UnrealBuildAccelerator/Core/Public/UbaMemory.h
@@ -227,7 +227,13 @@ namespace uba
 		u64 m_nextFree = 0;
 	};
 
-	inline u8 HexToByte(tchar c) { return (c >= '0' && c <= '9') ? u8(c - '0') : u8(c - 'a' + 10); }
+	inline u8 HexToByte(tchar c)
+	{
+		if (c >= '0' && c <= '9') return u8(c - '0');
+		else if (c >= 'a' && c <= 'f') return u8(c - 'a' + 10);
+		else if (c >= 'A' && c <= 'F') return u8(c - 'A' + 10);
+		else return 0; // Invalid hex character
+	}
 	constexpr tchar g_hexChars[] = TC("0123456789abcdef");
 
 	// TODO: These are backwards but changing would break cas storage
