[UBA]
* Renamed NetworkServer::IsConnected(clientId) to IsConnectedSlow since it now traverses all connections to find the ones that belongs to the client and also check if they are requested to be disconnected. This will fix so StorageServer::RetrieveCasFile exits faster if connections are requested to be disconnected and prevent a hang when pressing ctrl-c

diff --git a/Engine/Source/Programs/UnrealBuildAccelerator/Common/Private/UbaNetworkServer.cpp b/Engine/Source/Programs/UnrealBuildAccelerator/Common/Private/UbaNetworkServer.cpp
--- a/Engine/Source/Programs/UnrealBuildAccelerator/Common/Private/UbaNetworkServer.cpp
+++ b/Engine/Source/Programs/UnrealBuildAccelerator/Common/Private/UbaNetworkServer.cpp
@@ -12,6 +12,18 @@ namespace uba
 	{
 	}
 
+	struct NetworkServer::Client
+	{
+	public:
+		Client(const Guid& uid_, u32 id_) : uid(uid_), id(id_) {}
+		Guid uid;
+		u32 id;
+		u32 queuedConnectionCount = 0;
+		Atomic<u32> refCount;
+		Atomic<u32> connectionCount;
+		Atomic<u64> sendBytes;
+		Atomic<u64> recvBytes;
+	};
 
 	struct NetworkServer::WorkerContext
 	{
@@ -1061,14 +1073,26 @@ namespace uba
 		out.connectionCount += c.connectionCount;
 	}
 
-	bool NetworkServer::IsConnected(u32 clientId)
+	bool NetworkServer::IsConnectedSlow(u32 clientId)
 	{
 		SCOPED_READ_LOCK(m_clientsLock, lock);
 		auto findIt = m_clients.find(clientId);
 		if (findIt == m_clients.end())
 			return false;
 		Client& c = findIt->second;
-		return c.connectionCount > 0;
+		if (c.connectionCount == 0)
+			return false;
+		lock.Leave();
+
+		SCOPED_READ_LOCK(m_connectionsLock, lock2);
+		for (auto& con : m_connections)
+		{
+			if (con.m_client != &c)
+				continue;
+			if (!con.m_disconnectCalled)
+				return true;
+		}
+		return false;
 	}
 
 	void NetworkServer::ResetTotalStats()
diff --git a/Engine/Source/Programs/UnrealBuildAccelerator/Common/Private/UbaStorageServer.cpp b/Engine/Source/Programs/UnrealBuildAccelerator/Common/Private/UbaStorageServer.cpp
--- a/Engine/Source/Programs/UnrealBuildAccelerator/Common/Private/UbaStorageServer.cpp
+++ b/Engine/Source/Programs/UnrealBuildAccelerator/Common/Private/UbaStorageServer.cpp
@@ -123,7 +123,7 @@ namespace uba
 			u64 waitedMs = TimeToMs(waited);
 			if (waitedMs > 500)
 			{
-				if (!m_server.IsConnected(clientId))
+				if (!m_server.IsConnectedSlow(clientId))
 					return m_logger.Info(TC("Client with id %u disconnected while waiting for cas %s (%s)"), clientId, CasKeyString(casKey).str, hint).ToFalse();
 				if (waitedMs > 4 * 60 * 1000) // 4 minutes timeout
 					return m_logger.Info(TC("Timed out waiting %s for cas %s to be transferred from remote to storage (%s)"), TimeToText(waited).str, CasKeyString(casKey).str, hint).ToFalse();
diff --git a/Engine/Source/Programs/UnrealBuildAccelerator/Common/Public/UbaNetworkServer.h b/Engine/Source/Programs/UnrealBuildAccelerator/Common/Public/UbaNetworkServer.h
--- a/Engine/Source/Programs/UnrealBuildAccelerator/Common/Public/UbaNetworkServer.h
+++ b/Engine/Source/Programs/UnrealBuildAccelerator/Common/Public/UbaNetworkServer.h
@@ -102,7 +102,7 @@ namespace uba
 		u32 GetClientCount();
 		u32 GetConnectionCount();
 		void GetClientStats(ClientStats& out, u32 clientId);
-		bool IsConnected(u32 clientId);
+		bool IsConnectedSlow(u32 clientId);
 		void ResetTotalStats();
 
 		bool DoAdditionalWork();
@@ -187,18 +187,7 @@ namespace uba
 
 		Vector<char> m_clientsConfig;
 
-		struct Client
-		{
-		public:
-			Client(const Guid& uid_, u32 id_) : uid(uid_), id(id_) {}
-			Guid uid;
-			u32 id;
-			u32 queuedConnectionCount = 0;
-			Atomic<u32> refCount;
-			Atomic<u32> connectionCount;
-			Atomic<u64> sendBytes;
-			Atomic<u64> recvBytes;
-		};
+		struct Client;
 		ReaderWriterLock m_clientsLock;
 		UnorderedMap<u32, Client> m_clients;
 		u32 m_clientCounter = 0;
