[UBA]
* Fixed so encryption is compatible between windows and posix
* Added encryption on padding and <16 bytes data. Padding is xored with data from part that is properly encrypted. And for data smaller than 16 bytes we just xor with size and since these messages does not contain anything of value naive encryption is fine.
* Added -testencrypt=<file> and -testdecrypt=<file> to UbaCli to be able to test cross-platform encryption
* Changed so TestCrypto now test sizes under 16 bytes
* Added "scheduler" as option to UbaCli to make sure process is running through scheduler

diff --git a/Engine/Source/Programs/UnrealBuildAccelerator/Cli/Private/UbaCli.cpp b/Engine/Source/Programs/UnrealBuildAccelerator/Cli/Private/UbaCli.cpp
--- a/Engine/Source/Programs/UnrealBuildAccelerator/Cli/Private/UbaCli.cpp
+++ b/Engine/Source/Programs/UnrealBuildAccelerator/Cli/Private/UbaCli.cpp
@@ -180,11 +180,14 @@ namespace uba
 		bool writeCacheSummary = false;
 		bool logToFile = false;
 		bool useHackVfs = false;
+		bool useScheduler = false;
 		TString checkFileTable;
 		TString cacheFilterString;
 		TString cacheCommand;
 		TString testCompress;
 		TString testDecompress;
+		TString testEncrypt;
+		TString testDecrypt;
 		TString addCas;
 		TString configFile;
 		TString symbolFile;
@@ -357,6 +360,22 @@ namespace uba
 					}
 					testDecompress = value.data;
 				}
+				else if (name.Equals(TCV("-testencrypt")))
+				{
+					if (value.IsEmpty())
+						return PrintHelp(TC("-testEncrypt needs a value"));
+					testEncrypt = value.data;
+				}
+				else if (name.Equals(TCV("-testdecrypt")))
+				{
+					if (value.IsEmpty())
+					{
+						if (testEncrypt.empty())
+							return PrintHelp(TC("-testDecrypt needs a value"));
+						value.Clear().Append(g_rootDir).EnsureEndsWithSlash().Append(TCV("TestEncrypt.tmp"));
+					}
+					testDecrypt = value.data;
+				}
 				else if (name.Equals(TCV("-symbol")))
 				{
 					const tchar* plus = value.First('+');
@@ -520,6 +539,11 @@ namespace uba
 			{
 				commandType = CommandType_Agent;
 			}
+			else if (name.Equals(TCV("scheduler")))
+			{
+				commandType = CommandType_Local;
+				useScheduler = true;
+			}
 			else
 			{
 				return parseOption(name, value);
@@ -785,6 +809,72 @@ namespace uba
 			exit = true;
 		}
 
+
+		u8 crypto128Data[16];
+		bool useCrypto = !crypto.empty() || !testEncrypt.empty() || !testDecrypt.empty();
+		if (useCrypto)
+		{
+			if (crypto.empty())
+				crypto = TC("3f58aa57466db9999213456789123445");
+			if (!CryptoFromString(crypto128Data, 16, crypto.c_str()))
+				return logger.Error(TC("Failed to parse crypto key %s"), crypto.c_str());
+		}
+
+
+		if (!testEncrypt.empty())
+		{
+			FileAccessor fa(logger, testEncrypt.c_str());
+			if (!fa.OpenMemoryRead())
+				return logger.Error(TC("Failed to open file %s"), testEncrypt.c_str());
+			u64 fileSize = fa.GetSize();
+
+			StringBuffer<> dest;
+			dest.Append(g_rootDir).EnsureEndsWithSlash().Append(TCV("TestEncrypt.tmp"));
+			FileAccessor destFa(logger, dest.data);
+			if (!destFa.CreateMemoryWrite(false, DefaultAttributes(), fileSize))
+				return false;
+			memcpy(destFa.GetData(), fa.GetData(), fileSize);
+
+			CryptoKey key = Crypto::CreateKey(logger, crypto128Data);
+			if (key == InvalidCryptoKey)
+				return false;
+
+			if (!Crypto::Encrypt(logger, key, destFa.GetData(), u32(fileSize)))
+				return false;
+
+			if (!destFa.Close())
+				return false;
+			logger.Info(TC("Encryption of %s successful (Written to %s)"), testEncrypt.c_str(), dest.data);
+			exit = true;
+		}
+
+		if (!testDecrypt.empty())
+		{
+			FileAccessor fa(logger, testDecrypt.c_str());
+			if (!fa.OpenMemoryRead())
+				return logger.Error(TC("Failed to open file %s"), testDecrypt.c_str());
+			u64 fileSize = fa.GetSize();
+
+			StringBuffer<> dest;
+			dest.Append(g_rootDir).EnsureEndsWithSlash().Append(TCV("TestDecrypt.tmp"));
+			FileAccessor destFa(logger, dest.data);
+			if (!destFa.CreateMemoryWrite(false, DefaultAttributes(), fileSize))
+				return false;
+			memcpy(destFa.GetData(), fa.GetData(), fileSize);
+
+			CryptoKey key = Crypto::CreateKey(logger, crypto128Data);
+			if (key == InvalidCryptoKey)
+				return false;
+
+			if (!Crypto::Decrypt(logger, key, destFa.GetData(), u32(fileSize)))
+				return false;
+
+			if (!destFa.Close())
+				return false;
+			logger.Info(TC("Decryption of %s successful (Written to %s)"), testDecrypt.c_str(), dest.data);
+			exit = true;
+		}
+
 		if (!symbolFile.empty())
 		{
 			FileInformation info;
@@ -964,18 +1054,17 @@ namespace uba
 		if (!ctorSuccess)
 			return false;
 
-		u8 crypto128Data[16];
-		if (!crypto.empty())
+		if (useCrypto)
 		{
-			if (!CryptoFromString(crypto128Data, 16, crypto.c_str()))
-				return logger.Error(TC("Failed to parse crypto key %s"), crypto.c_str());
 			networkServer.RegisterCryptoKey(crypto128Data);
 			logger.Info(TC("Using crypto key %s for connections"), crypto.c_str());
 		}
 
 
 		bool isRemote = commandType == CommandType_Remote || commandType == CommandType_Agent;
-		bool useScheduler = StringView(application).EndsWith(TCV(".yaml"));
+
+		if (StringView(application).EndsWith(TCV(".yaml")))
+			useScheduler = true;
 
 		StorageServerCreateInfo storageInfo(networkServer, g_rootDir.data, logWriter);
 		storageInfo.casCapacityBytes = storageCapacity;
@@ -1193,7 +1282,7 @@ namespace uba
 
 		CoordinatorWrapper coordinator;
 
-		auto RunScheduler = [&](const tchar* yamlFile)
+		auto RunScheduler = [&](const TString& app, const TString& arg)
 		{
 			/*
 			bool ctorSuccess;
@@ -1232,12 +1321,27 @@ namespace uba
 			info.writeToCache = populateCache;
 			Scheduler scheduler(info);
 
-			if (!scheduler.EnqueueFromFile(yamlFile, [&](EnqueueProcessInfo& epi)
-				{
-					if (rootsHandle)
-						const_cast<ProcessStartInfo&>(epi.info).rootsHandle = rootsHandle;
-				}))
-				return false;
+			sessionServer.SetOuterScheduler(&scheduler);
+
+			if (StringView(app).EndsWith(TCV(".yaml")))
+			{
+				if (!scheduler.EnqueueFromFile(app.c_str(), [&](EnqueueProcessInfo& epi)
+					{
+						if (rootsHandle)
+							const_cast<ProcessStartInfo&>(epi.info).rootsHandle = rootsHandle;
+					}))
+					return false;
+			}
+			else
+			{
+				ProcessStartInfo pinfo;
+				pinfo.description = app.c_str();
+				pinfo.application = app.c_str();
+				pinfo.arguments = arg.c_str();
+				pinfo.workingDir = workDir.data;
+				pinfo.rootsHandle = rootsHandle;
+				scheduler.EnqueueProcess({pinfo});
+			}
 
 			u32 queued, activeLocal, activeRemote, outFinished;
 			scheduler.GetStats(queued, activeLocal, activeRemote, outFinished);
@@ -1341,7 +1445,7 @@ namespace uba
 
 			if (useScheduler)
 			{
-				success = RunScheduler(application.c_str());
+				success = RunScheduler(application, arguments);
 			}
 			else
 			{
diff --git a/Engine/Source/Programs/UnrealBuildAccelerator/Common/Private/UbaCrypto.cpp b/Engine/Source/Programs/UnrealBuildAccelerator/Common/Private/UbaCrypto.cpp
--- a/Engine/Source/Programs/UnrealBuildAccelerator/Common/Private/UbaCrypto.cpp
+++ b/Engine/Source/Programs/UnrealBuildAccelerator/Common/Private/UbaCrypto.cpp
@@ -31,6 +31,12 @@ namespace uba
 {
 	inline constexpr u32 kAesBytes128 = 16;
 
+	inline constexpr u64 kFixedIv[2] =
+	{
+		0x0001020304050607ull,
+		0x08090A0B0C0D0E0Full,
+	};
+
 	CryptoKey Crypto::CreateKey(Logger& logger, const u8* key128)
 	{
 #if UBA_CRYPTO_TYPE == 1
@@ -108,10 +114,21 @@ namespace uba
 #if UBA_CRYPTO_TYPE == 1
 	bool BCryptEncryptDecrypt(Logger& logger, bool encrypt, CryptoKey key, u8* data, u32 size)
 	{
+		// We don't fully encrypt data smaller than 16 bytes. There are a few messages that are harmless that does not need to be encrypted since they contain nothing of value
+		// There is also no way to find these messages in a network stream since you need to be able to decrypt to know where messages start and end
+		if (size < kAesBytes128)
+		{
+			for (u32 i=0; i!=size; ++i)
+				data[i] ^= size;
+			return true;
+		}
+
 		u8 objectBuffer[1024];
 		u32 objectBufferLen = sizeof(objectBuffer);
 
 		u32 alignedSize = (size / kAesBytes128) * kAesBytes128;
+		u32 overflow = size - alignedSize;
+
 		BCRYPT_KEY_HANDLE newKey;
 		NTSTATUS res = BCryptDuplicateKey((BCRYPT_KEY_HANDLE)key, &newKey, objectBuffer, objectBufferLen, 0);
 		if (!BCRYPT_SUCCESS(res))
@@ -121,12 +138,22 @@ namespace uba
 		}
 		auto g = MakeGuard([&]() { BCryptDestroyKey(newKey); });
 
-		u8* initVector = nullptr;
-		u32 initVectorSize = 0;
+		u64 initVector[2] { kFixedIv[0], kFixedIv[1] };
+
 		ULONG cipherTextLength = 0;
-		res = encrypt ?
-			BCryptEncrypt(newKey, data, alignedSize, NULL, initVector, initVectorSize, data, alignedSize, &cipherTextLength, 0) :
-			BCryptDecrypt(newKey, data, alignedSize, NULL, initVector, initVectorSize, data, alignedSize, &cipherTextLength, 0);
+
+		if (encrypt)
+		{
+			for (u32 i=0; i!=overflow; ++i)
+				data[alignedSize + i] ^= data[i]; 
+			res = BCryptEncrypt(newKey, data, alignedSize, NULL, (u8*)initVector, kAesBytes128, data, alignedSize, &cipherTextLength, 0);
+		}
+		else
+		{
+			res = BCryptDecrypt(newKey, data, alignedSize, NULL, (u8*)initVector, kAesBytes128, data, alignedSize, &cipherTextLength, 0);
+			for (u32 i=0; i!=overflow; ++i)
+				data[alignedSize + i] ^= data[i]; 
+		}
 
 		if (!BCRYPT_SUCCESS(res))
 		{
@@ -143,17 +170,31 @@ namespace uba
 #elif UBA_CRYPTO_TYPE == 2
 	bool OpenSslEncryptDecrypt(CryptoKey key, u8* data, u32 size, int enc)
 	{
+		if (size < kAesBytes128)
+		{
+			for (u32 i=0; i!=size; ++i)
+				data[i] ^= size;
+			return true;
+		}
+
 		auto& keyData = *(KeyData*)uintptr_t(key);
-		u8 iv[AES_BLOCK_SIZE];
-		memcpy(iv, (u8*)uintptr_t(key), AES_BLOCK_SIZE);
+		u64 iv[2] { kFixedIv[0], kFixedIv[1] };
 
-		u32 alignedSize = size & 0xfffffff0;
-		if (alignedSize)
-			AES_cbc_encrypt(data, data, alignedSize, &(enc ? keyData.encryptKey : keyData.decryptKey), iv, enc);
+		u32 alignedSize = (size / kAesBytes128) * kAesBytes128;
+		u32 overflow = size - alignedSize;
 
-		u8* data2 = data + alignedSize;
-		for (u32 i=0, left=u32(size-alignedSize);i!=left; ++i)
-			data2[i] = data2[i] ^ iv[i];
+		if (enc)
+		{
+			for (u32 i=0; i!=overflow; ++i)
+				data[alignedSize + i] ^= data[i]; 
+			AES_cbc_encrypt(data, data, alignedSize, &keyData.encryptKey, (u8*)iv, enc);
+		}
+		else
+		{
+			AES_cbc_encrypt(data, data, alignedSize, &keyData.decryptKey, (u8*)iv, enc);
+			for (u32 i=0; i!=overflow; ++i)
+				data[alignedSize + i] ^= data[i]; 
+		}
 		return true;
 	}
 #endif
@@ -188,5 +229,4 @@ namespace uba
 		((u64*)out)[1] = StringToValue(str + 16, 16);
 		return true;
 	}
-
 }
\ No newline at end of file
diff --git a/Engine/Source/Programs/UnrealBuildAccelerator/Test/Private/UbaTestCrypto.cpp b/Engine/Source/Programs/UnrealBuildAccelerator/Test/Private/UbaTestCrypto.cpp
--- a/Engine/Source/Programs/UnrealBuildAccelerator/Test/Private/UbaTestCrypto.cpp
+++ b/Engine/Source/Programs/UnrealBuildAccelerator/Test/Private/UbaTestCrypto.cpp
@@ -28,7 +28,7 @@ namespace uba
 				if (!Crypto::Encrypt(logger, encryptKey, encryptedData, dataSize))
 					return false;
 
-				if (dataSize >= 16 && memcmp(encryptedData, originalData, dataSize) == 0)
+				if (memcmp(encryptedData, originalData, dataSize) == 0)
 					return false;
 
 				if (!Crypto::Decrypt(logger, decryptKey, encryptedData, dataSize))
