[UBA]
* Changed so crypto init vector is randomized during handshake and unique per connection. This is at least better than hardcoded init vector. Unfortunately there is no good way atm to have a changing init vector since encoding/decoding is happening outside of socket send/recv and is out of order with network

diff --git a/Engine/Source/Programs/UnrealBuildAccelerator/Cli/Private/UbaCli.cpp b/Engine/Source/Programs/UnrealBuildAccelerator/Cli/Private/UbaCli.cpp
--- a/Engine/Source/Programs/UnrealBuildAccelerator/Cli/Private/UbaCli.cpp
+++ b/Engine/Source/Programs/UnrealBuildAccelerator/Cli/Private/UbaCli.cpp
@@ -839,7 +839,7 @@ namespace uba
 			if (key == InvalidCryptoKey)
 				return false;
 
-			if (!Crypto::Encrypt(logger, key, destFa.GetData(), u32(fileSize)))
+			if (!Crypto::Encrypt(logger, key, destFa.GetData(), u32(fileSize), *(Guid*)crypto128Data))
 				return false;
 
 			if (!destFa.Close())
@@ -866,7 +866,7 @@ namespace uba
 			if (key == InvalidCryptoKey)
 				return false;
 
-			if (!Crypto::Decrypt(logger, key, destFa.GetData(), u32(fileSize)))
+			if (!Crypto::Decrypt(logger, key, destFa.GetData(), u32(fileSize), *(Guid*)crypto128Data))
 				return false;
 
 			if (!destFa.Close())
diff --git a/Engine/Source/Programs/UnrealBuildAccelerator/Common/Private/UbaCrypto.cpp b/Engine/Source/Programs/UnrealBuildAccelerator/Common/Private/UbaCrypto.cpp
--- a/Engine/Source/Programs/UnrealBuildAccelerator/Common/Private/UbaCrypto.cpp
+++ b/Engine/Source/Programs/UnrealBuildAccelerator/Common/Private/UbaCrypto.cpp
@@ -31,12 +31,6 @@ namespace uba
 {
 	inline constexpr u32 kAesBytes128 = 16;
 
-	inline constexpr u64 kFixedIv[2] =
-	{
-		0x0001020304050607ull,
-		0x08090A0B0C0D0E0Full,
-	};
-
 	CryptoKey Crypto::CreateKey(Logger& logger, const u8* key128)
 	{
 #if UBA_CRYPTO_TYPE == 1
@@ -112,17 +106,8 @@ namespace uba
 	}
 
 #if UBA_CRYPTO_TYPE == 1
-	bool BCryptEncryptDecrypt(Logger& logger, bool encrypt, CryptoKey key, u8* data, u32 size)
+	bool BCryptEncryptDecrypt(Logger& logger, bool encrypt, CryptoKey key, u8* data, u32 size, const Guid& initVector)
 	{
-		// We don't fully encrypt data smaller than 16 bytes. There are a few messages that are harmless that does not need to be encrypted since they contain nothing of value
-		// There is also no way to find these messages in a network stream since you need to be able to decrypt to know where messages start and end
-		if (size < kAesBytes128)
-		{
-			for (u32 i=0; i!=size; ++i)
-				data[i] ^= size;
-			return true;
-		}
-
 		u8 objectBuffer[1024];
 		u32 objectBufferLen = sizeof(objectBuffer);
 
@@ -138,7 +123,7 @@ namespace uba
 		}
 		auto g = MakeGuard([&]() { BCryptDestroyKey(newKey); });
 
-		u64 initVector[2] { kFixedIv[0], kFixedIv[1] };
+		Guid iv = initVector;
 
 		ULONG cipherTextLength = 0;
 
@@ -146,11 +131,11 @@ namespace uba
 		{
 			for (u32 i=0; i!=overflow; ++i)
 				data[alignedSize + i] ^= data[i]; 
-			res = BCryptEncrypt(newKey, data, alignedSize, NULL, (u8*)initVector, kAesBytes128, data, alignedSize, &cipherTextLength, 0);
+			res = BCryptEncrypt(newKey, data, alignedSize, NULL, (u8*)&iv, kAesBytes128, data, alignedSize, &cipherTextLength, 0);
 		}
 		else
 		{
-			res = BCryptDecrypt(newKey, data, alignedSize, NULL, (u8*)initVector, kAesBytes128, data, alignedSize, &cipherTextLength, 0);
+			res = BCryptDecrypt(newKey, data, alignedSize, NULL, (u8*)&iv, kAesBytes128, data, alignedSize, &cipherTextLength, 0);
 			for (u32 i=0; i!=overflow; ++i)
 				data[alignedSize + i] ^= data[i]; 
 		}
@@ -168,17 +153,10 @@ namespace uba
 		return true;
 	}
 #elif UBA_CRYPTO_TYPE == 2
-	bool OpenSslEncryptDecrypt(CryptoKey key, u8* data, u32 size, int enc)
+	bool OpenSslEncryptDecrypt(CryptoKey key, u8* data, u32 size, bool enc, const Guid& initVector)
 	{
-		if (size < kAesBytes128)
-		{
-			for (u32 i=0; i!=size; ++i)
-				data[i] ^= size;
-			return true;
-		}
-
 		auto& keyData = *(KeyData*)uintptr_t(key);
-		u64 iv[2] { kFixedIv[0], kFixedIv[1] };
+		Guid iv = initVector;
 
 		u32 alignedSize = (size / kAesBytes128) * kAesBytes128;
 		u32 overflow = size - alignedSize;
@@ -187,11 +165,11 @@ namespace uba
 		{
 			for (u32 i=0; i!=overflow; ++i)
 				data[alignedSize + i] ^= data[i]; 
-			AES_cbc_encrypt(data, data, alignedSize, &keyData.encryptKey, (u8*)iv, enc);
+			AES_cbc_encrypt(data, data, alignedSize, &keyData.encryptKey, (u8*)&iv, AES_ENCRYPT);
 		}
 		else
 		{
-			AES_cbc_encrypt(data, data, alignedSize, &keyData.decryptKey, (u8*)iv, enc);
+			AES_cbc_encrypt(data, data, alignedSize, &keyData.decryptKey, (u8*)&iv, AES_DECRYPT);
 			for (u32 i=0; i!=overflow; ++i)
 				data[alignedSize + i] ^= data[i]; 
 		}
@@ -199,26 +177,34 @@ namespace uba
 	}
 #endif
 
-	bool Crypto::Encrypt(Logger& logger, CryptoKey key, u8* data, u32 size)
+	bool EncryptDecrypt(Logger& logger, CryptoKey key, u8* data, u32 size, const Guid& initVector, bool enc)
 	{
+		// We don't fully encrypt data smaller than 16 bytes. There are a few messages that are harmless that does not need to be encrypted since they contain nothing of value
+		if (size < kAesBytes128)
+		{
+			// This is just here to make unit tests happy.. no special security. Using init vector here will probably reduce security since it would be possible to reverse engineer part of init vector
+			for (u32 i=0; i!=size; ++i)
+				data[i] ^= 5381;//initVector.data1;
+			return true;
+		}
+
 #if UBA_CRYPTO_TYPE == 1
-		return BCryptEncryptDecrypt(logger, true, key, data, size);
+		return BCryptEncryptDecrypt(logger, enc, key, data, size, initVector);
 #elif UBA_CRYPTO_TYPE == 2
-		return OpenSslEncryptDecrypt(key, data, size, AES_ENCRYPT);
+		return OpenSslEncryptDecrypt(key, data, size, enc, initVector);
 #else
 		return false;
 #endif
 	}
 
-	bool Crypto::Decrypt(Logger& logger, CryptoKey key, u8* data, u32 size)
+	bool Crypto::Encrypt(Logger& logger, CryptoKey key, u8* data, u32 size, const Guid& initVector)
 	{
-#if UBA_CRYPTO_TYPE == 1
-		return BCryptEncryptDecrypt(logger, false, key, data, size);
-#elif UBA_CRYPTO_TYPE == 2
-		return OpenSslEncryptDecrypt(key, data, size, AES_DECRYPT);
-#else
-		return false;
-#endif
+		return EncryptDecrypt(logger, key, data, size, initVector, true);
+	}
+
+	bool Crypto::Decrypt(Logger& logger, CryptoKey key, u8* data, u32 size, const Guid& initVector)
+	{
+		return EncryptDecrypt(logger, key, data, size, initVector, false);
 	}
 	
 	bool CryptoFromString(u8* out, u32 outSize, const tchar* str)
diff --git a/Engine/Source/Programs/UnrealBuildAccelerator/Common/Private/UbaNetworkBackendTcp.cpp b/Engine/Source/Programs/UnrealBuildAccelerator/Common/Private/UbaNetworkBackendTcp.cpp
--- a/Engine/Source/Programs/UnrealBuildAccelerator/Common/Private/UbaNetworkBackendTcp.cpp
+++ b/Engine/Source/Programs/UnrealBuildAccelerator/Common/Private/UbaNetworkBackendTcp.cpp
@@ -49,7 +49,7 @@
 
 namespace uba
 {
-	constexpr u32 MaxHeaderSize = 24;
+	constexpr u32 MaxHeaderSize = 33;
 
 	////////////////////////////////////////////////////////////////////////////////////////////////////
 
diff --git a/Engine/Source/Programs/UnrealBuildAccelerator/Common/Private/UbaNetworkClient.cpp b/Engine/Source/Programs/UnrealBuildAccelerator/Common/Private/UbaNetworkClient.cpp
--- a/Engine/Source/Programs/UnrealBuildAccelerator/Common/Private/UbaNetworkClient.cpp
+++ b/Engine/Source/Programs/UnrealBuildAccelerator/Common/Private/UbaNetworkClient.cpp
@@ -20,6 +20,23 @@ namespace uba
 		table.GetValueAsU32(desiredConnectionCount, TC("DesiredConnectionCount"));
 	}
 
+	struct NetworkClient::Connection
+	{
+		Connection(NetworkClient& o) : owner(o), disconnectedEvent(true) {}
+		NetworkClient& owner;
+		void* backendConnection = nullptr;
+		Atomic<u32> connected;
+		Event disconnectedEvent;
+		NetworkBackend* backend = nullptr;
+		Guid cryptoIV;
+
+		#if UBA_TRACK_NETWORK_TIMES
+		Atomic<u64> lastSendTime;
+		Atomic<u64> lastHeaderRecvTime;
+		Atomic<u64> lastBodyRecvTime;
+		#endif
+	};
+
 	NetworkClient::NetworkClient(bool& outCtorSuccess, const NetworkClientCreateInfo& info, const tchar* name)
 	: WorkManagerImpl(info.workerCount == ~0u ? GetLogicalProcessorCount() : info.workerCount, TC("UbaWrk/NetClnt"))
 	,	m_logWriter(info.logWriter)
@@ -106,27 +123,40 @@ namespace uba
 				rc.exitScopeEvent.Set();
 			});
 
-		backend.SetRecvCallbacks(backendConnection, &rc, 1 + sizeof(Guid), [](void* context, const Guid& connectionUid, u8* headerData, void*& outBodyContext, u8*& outBodyData, u32& outBodySize)
+		constexpr u32 recvSize = 1 + sizeof(Guid)*2;
+
+		backend.SetRecvCallbacks(backendConnection, &rc, recvSize, [](void* context, const Guid& connectionUid, u8* headerData, void*& outBodyContext, u8*& outBodyData, u32& outBodySize)
 			{
 				auto& rc = *(RecvContext*)context;
+				NetworkClient& client = rc.client;
+
 				rc.error = *headerData;
 				Guid serverUid = *(Guid*)(headerData+1);
+				Guid cryptoIV = *(Guid*)(headerData+17);
 
 				if (serverUid == Guid())
 					rc.error = 5;
 
 				if (!rc.error)
 				{
-					SCOPED_FUTEX(rc.client.m_serverUidLock, lock);
-					if (rc.client.m_serverUid == Guid())
-						rc.client.m_serverUid = serverUid;
-					else if (rc.client.m_serverUid != serverUid) // Seems like two different servers tried to connect to this client.. keep the first one and ignore the others
+					SCOPED_FUTEX(client.m_serverUidLock, lock);
+					if (client.m_serverUid == Guid())
+						client.m_serverUid = serverUid;
+					else if (client.m_serverUid != serverUid) // Seems like two different servers tried to connect to this client.. keep the first one and ignore the others
 						rc.error = 6;
 				}
 
 				if (!rc.error)
-					if (!rc.client.ConnectedCallback(rc.backend, rc.backendConnection))
+				{
+					CryptoKey key = client.m_cryptoKey;
+					if (key != InvalidCryptoKey)
+						if (!Crypto::Decrypt(client.m_logger, key, (u8*)&cryptoIV, sizeof(cryptoIV), {}))
+							return false;
+
+
+					if (!client.ConnectedCallback(rc.backend, rc.backendConnection, cryptoIV))
 						rc.error = 4;
+				}
 
 				if (rc.error != 0)
 					return false;
@@ -139,14 +169,14 @@ namespace uba
 
 
 		StackBinaryWriter<1024> handshakeData;
-		if (m_cryptoKey)
+		if (m_cryptoKey != InvalidCryptoKey)
 		{
 			handshakeData.WriteByte(1);
 			// If we have a crypto key we start by sending a predefined 128 bytes blob that is encrypted.
 			// If server decrypt it to the same blob, we're good on that part
 			u8* encryptedBuffer = handshakeData.AllocWrite(sizeof(EncryptionHandshakeString));//[1024];
 			memcpy(encryptedBuffer, EncryptionHandshakeString, sizeof(EncryptionHandshakeString));
-			if (!Crypto::Encrypt(m_logger, m_cryptoKey, encryptedBuffer, sizeof(EncryptionHandshakeString)))
+			if (!Crypto::Encrypt(m_logger, m_cryptoKey, encryptedBuffer, sizeof(EncryptionHandshakeString), {}))
 				return false;
 		}
 
@@ -214,7 +244,7 @@ namespace uba
 		c.disconnectedEvent.Set();
 	}
 
-	bool NetworkClient::ConnectedCallback(NetworkBackend& backend, void* backendConnection)
+	bool NetworkClient::ConnectedCallback(NetworkBackend& backend, void* backendConnection, const Guid& cryptoIV)
 	{
 		SCOPED_WRITE_LOCK(m_connectionsLock, lock);
 		if (m_isDisconnecting)
@@ -224,6 +254,7 @@ namespace uba
 		connection->backendConnection = backendConnection;
 		connection->connected = 1;
 		connection->backend = &backend;
+		connection->cryptoIV = cryptoIV;
 
 		SCOPED_FUTEX(m_connectionsItLock, l); // Take this lock to make sure callbacks are set before connection is used
 		m_connectionsIt = --m_connections.end();
@@ -676,7 +707,7 @@ namespace uba
 		if (m_cryptoKey && bodySize)
 		{
 			TimerScope ts(m_encryptTimer);
-			if (!Crypto::Encrypt(m_logger, m_cryptoKey, data + SendHeaderSize, bodySize))
+			if (!Crypto::Encrypt(m_logger, m_cryptoKey, data + SendHeaderSize, bodySize, connection.cryptoIV))
 			{
 				message.m_error = 8;
 				OnDisconnected(connection, 8);
@@ -717,7 +748,7 @@ namespace uba
 			else if (m_cryptoKey && !message.m_error && message.m_responseSize)
 			{
 				TimerScope ts(m_decryptTimer);
-				if (!Crypto::Decrypt(m_logger, m_cryptoKey, (u8*)message.m_response, message.m_responseSize))
+				if (!Crypto::Decrypt(m_logger, m_cryptoKey, (u8*)message.m_response, message.m_responseSize, connection.cryptoIV))
 				{
 					message.m_error = 5;
 					OnDisconnected(connection, 5);
@@ -798,8 +829,9 @@ namespace uba
 		if (m_client->m_cryptoKey)
 		{
 			UBA_ASSERT(!response.GetPosition());
+			UBA_ASSERT(m_connection);
 			TimerScope ts(m_client->m_decryptTimer);
-			if (!Crypto::Decrypt(m_client->m_logger, m_client->m_cryptoKey, (u8*)m_response, m_responseSize))
+			if (!Crypto::Decrypt(m_client->m_logger, m_client->m_cryptoKey, (u8*)m_response, m_responseSize, m_connection->cryptoIV))
 			{
 				m_error = 10;
 				return false;
diff --git a/Engine/Source/Programs/UnrealBuildAccelerator/Common/Private/UbaNetworkServer.cpp b/Engine/Source/Programs/UnrealBuildAccelerator/Common/Private/UbaNetworkServer.cpp
--- a/Engine/Source/Programs/UnrealBuildAccelerator/Common/Private/UbaNetworkServer.cpp
+++ b/Engine/Source/Programs/UnrealBuildAccelerator/Common/Private/UbaNetworkServer.cpp
@@ -209,13 +209,14 @@ namespace uba
 			return true;
 		}
 
-		bool SendInitialResponse(u8 value)
+		bool SendInitialResponse(u8 value, const Guid& cryptoIV = {})
 		{
-			u8 data[32];
+			u8 data[33];
 			*data = value;
 			*(Guid*)(data+1) = m_server.m_uid;
+			*(Guid*)(data+17) = cryptoIV;
 			NetworkBackend::SendContext context(NetworkBackend::SendFlags_Async);
-			return m_backend.Send(m_server.m_logger, m_backendConnection, data, 1 + sizeof(Guid), context, TC("UidResponse"));
+			return m_backend.Send(m_server.m_logger, m_backendConnection, data, sizeof(data), context, TC("UidResponse"));
 		}
 
 		static bool ReceiveHandshakeHeader(void* context, const Guid& connectionUid, u8* headerData, void*& outBodyContext, u8*& outBodyData, u32& outBodySize)
@@ -241,7 +242,7 @@ namespace uba
 			{
 				u8 temp[sizeof(EncryptionHandshakeString)];
 				memcpy(temp, handshakeData, sizeof(temp));
-				if (!Crypto::Decrypt(logger, key, temp, sizeof(EncryptionHandshakeString)))
+				if (!Crypto::Decrypt(logger, key, temp, sizeof(EncryptionHandshakeString), {}))
 					return false;
 				return memcmp(temp, EncryptionHandshakeString, sizeof(EncryptionHandshakeString)) == 0;
 			};
@@ -342,7 +343,16 @@ namespace uba
 			if (server.m_useKeepAlive)
 				conn.m_backend.SetRecvTimeout(conn.m_backendConnection, KeepAliveIdleSeconds*1000, context, RecvTimeout);
 
-			if (!conn.SendInitialResponse(0))
+			Guid cryptoIV;
+			if (conn.m_cryptoKey != InvalidCryptoKey)
+			{
+				CreateGuid(cryptoIV);
+				conn.m_cryptoIV = cryptoIV;
+				if (!Crypto::Encrypt(server.m_logger, conn.m_cryptoKey, (u8*)&cryptoIV, sizeof(cryptoIV), {}))
+					return false;
+			}
+
+			if (!conn.SendInitialResponse(0, cryptoIV))
 				return false;
 
 			SCOPED_FUTEX(conn.m_shutdownLock, shutdownLock);
@@ -504,6 +514,7 @@ namespace uba
 		Client* m_client = nullptr;
 		sockaddr m_remoteSockAddr;
 		CryptoKey m_cryptoKey;
+		Guid m_cryptoIV;
 		Event m_disconnectCallbackCalled;
 		Atomic<int> m_activeWorkerCount;
 		Atomic<int> m_disconnectCalled;
@@ -579,7 +590,7 @@ namespace uba
 		{
 			//TrackHintScope ths(tws, TCV("Decrypt"));
 			TimerScope ts(connection.m_decryptTimer);
-			if (!Crypto::Decrypt(server.m_logger, cryptoKey, context.buffer.data(), context.dataSize))
+			if (!Crypto::Decrypt(server.m_logger, cryptoKey, context.buffer.data(), context.dataSize, connection.m_cryptoIV))
 			{
 				connection.SetShouldDisconnect();
 				connection.Release();
@@ -635,7 +646,7 @@ namespace uba
 				//TrackHintScope ths(tws, TCV("Encrypt"));
 				TimerScope ts(connection.m_encryptTimer);
 				u8* bodyData = writer.GetData() + HeaderSize;
-				if (!Crypto::Encrypt(server.m_logger, cryptoKey, bodyData, bodySize))
+				if (!Crypto::Encrypt(server.m_logger, cryptoKey, bodyData, bodySize, connection.m_cryptoIV))
 				{
 					connection.SetShouldDisconnect();
 					size = MessageErrorSize;
@@ -1204,7 +1215,7 @@ namespace uba
 			{
 				TimerScope ts(connection.m_encryptTimer);
 				u8* bodyData = writer.GetData() + HeaderSize;
-				if (!Crypto::Encrypt(m_logger, connection.m_cryptoKey, bodyData, bodySize))
+				if (!Crypto::Encrypt(m_logger, connection.m_cryptoKey, bodyData, bodySize, connection.m_cryptoIV))
 				{
 					connection.SetShouldDisconnect();
 					return false;
diff --git a/Engine/Source/Programs/UnrealBuildAccelerator/Common/Public/UbaCrypto.h b/Engine/Source/Programs/UnrealBuildAccelerator/Common/Public/UbaCrypto.h
--- a/Engine/Source/Programs/UnrealBuildAccelerator/Common/Public/UbaCrypto.h
+++ b/Engine/Source/Programs/UnrealBuildAccelerator/Common/Public/UbaCrypto.h
@@ -18,8 +18,8 @@ namespace uba
 		static CryptoKey DuplicateKey(Logger& logger, CryptoKey original);
 		static void DestroyKey(CryptoKey key);
 
-		static bool Encrypt(Logger& logger, CryptoKey key, u8* data, u32 size);
-		static bool Decrypt(Logger& logger, CryptoKey key, u8* data, u32 size);
+		static bool Encrypt(Logger& logger, CryptoKey key, u8* data, u32 size, const Guid& initVector);
+		static bool Decrypt(Logger& logger, CryptoKey key, u8* data, u32 size, const Guid& initVector);
 	};
 
 	bool CryptoFromString(u8* out, u32 outSize, const tchar* str);
diff --git a/Engine/Source/Programs/UnrealBuildAccelerator/Common/Public/UbaNetworkClient.h b/Engine/Source/Programs/UnrealBuildAccelerator/Common/Public/UbaNetworkClient.h
--- a/Engine/Source/Programs/UnrealBuildAccelerator/Common/Public/UbaNetworkClient.h
+++ b/Engine/Source/Programs/UnrealBuildAccelerator/Common/Public/UbaNetworkClient.h
@@ -75,24 +75,10 @@ namespace uba
 		NetworkBackend* GetFirstConnectionBackend();
 
 	private:
-		struct Connection
-		{
-			Connection(NetworkClient& o) : owner(o), disconnectedEvent(true) {}
-			NetworkClient& owner;
-			void* backendConnection = nullptr;
-			Atomic<u32> connected;
-			Event disconnectedEvent;
-			NetworkBackend* backend = nullptr;
-
-			#if UBA_TRACK_NETWORK_TIMES
-			Atomic<u64> lastSendTime;
-			Atomic<u64> lastHeaderRecvTime;
-			Atomic<u64> lastBodyRecvTime;
-			#endif
-		};
+		struct Connection;
 
 		bool AddConnection(NetworkBackend& backend, void* backendConnection, bool* timedOut);
-		bool ConnectedCallback(NetworkBackend& backend, void* backendConnection);
+		bool ConnectedCallback(NetworkBackend& backend, void* backendConnection, const Guid& cryptoIV);
 		static void DisconnectCallback(void* context, const Guid& connectionUid, void* connection);
 		static bool ReceiveResponseHeader(void* context, const Guid& connectionUid, u8* headerData, void*& outBodyContext, u8*& outBodyData, u32& outBodySize);
 		static bool ReceiveResponseBody(void* context, bool recvError, u8* headerData, void* bodyContext, u8* bodyData, u32 bodySize);
diff --git a/Engine/Source/Programs/UnrealBuildAccelerator/Test/Private/UbaTestAll.h b/Engine/Source/Programs/UnrealBuildAccelerator/Test/Private/UbaTestAll.h
--- a/Engine/Source/Programs/UnrealBuildAccelerator/Test/Private/UbaTestAll.h
+++ b/Engine/Source/Programs/UnrealBuildAccelerator/Test/Private/UbaTestAll.h
@@ -28,10 +28,12 @@ namespace uba
 		UBA_TEST(TestMemoryBlock) \
 		UBA_TEST(TestParseArguments) \
 		UBA_TEST(TestBinaryWriter) \
+		UBA_TEST(TestCrypto) \
 		UBA_TEST(TestSockets) \
 		UBA_TEST(TestClientServer) \
 		UBA_TEST(TestClientServer2) \
 		UBA_TEST(TestClientServerMem) \
+		UBA_TEST(TestClientServerCrypto) \
 		UBA_TEST(TestStorage) \
 		UBA_TEST(TestRemoteStorageStore) \
 		UBA_TEST(TestRemoteStorageFetch) \
@@ -43,7 +45,6 @@ namespace uba
 		UBA_TEST(TestHashTable) \
 		UBA_TEST(TestLoadConfig) \
 		UBA_TEST(TestSaveConfig) \
-		UBA_TEST(TestCrypto) \
 		UBA_TEST(TestBinDependencies) \
 		UBA_TEST(TestRootPaths) \
 		UBA_TEST(TestRegisterChanges) \
diff --git a/Engine/Source/Programs/UnrealBuildAccelerator/Test/Private/UbaTestCrypto.cpp b/Engine/Source/Programs/UnrealBuildAccelerator/Test/Private/UbaTestCrypto.cpp
--- a/Engine/Source/Programs/UnrealBuildAccelerator/Test/Private/UbaTestCrypto.cpp
+++ b/Engine/Source/Programs/UnrealBuildAccelerator/Test/Private/UbaTestCrypto.cpp
@@ -25,13 +25,13 @@ namespace uba
 
 			for (u32 i=0; i!=3; ++i)
 			{
-				if (!Crypto::Encrypt(logger, encryptKey, encryptedData, dataSize))
+				if (!Crypto::Encrypt(logger, encryptKey, encryptedData, dataSize, {}))
 					return false;
 
 				if (memcmp(encryptedData, originalData, dataSize) == 0)
 					return false;
 
-				if (!Crypto::Decrypt(logger, decryptKey, encryptedData, dataSize))
+				if (!Crypto::Decrypt(logger, decryptKey, encryptedData, dataSize, {}))
 					return false;
 
 				if (memcmp(encryptedData, originalData, dataSize) != 0)
diff --git a/Engine/Source/Programs/UnrealBuildAccelerator/Test/Private/UbaTestNetwork.cpp b/Engine/Source/Programs/UnrealBuildAccelerator/Test/Private/UbaTestNetwork.cpp
--- a/Engine/Source/Programs/UnrealBuildAccelerator/Test/Private/UbaTestNetwork.cpp
+++ b/Engine/Source/Programs/UnrealBuildAccelerator/Test/Private/UbaTestNetwork.cpp
@@ -7,6 +7,7 @@
 #include "UbaNetworkMessage.h"
 #include "UbaNetworkBackendTcp.h"
 #include "UbaNetworkBackendMemory.h"
+#include "UbaTest.h"
 
 namespace uba
 {
@@ -184,4 +185,72 @@ namespace uba
 		logger.Info(TC("Got value %u"), value);
 		return true;
 	}
+
+	bool TestClientServerCrypto(LoggerWithWriter& logger, const StringBufferBase& testRootDir)
+	{
+		auto& logWriter = logger.m_writer;
+		NetworkBackendTcp networkBackend(logWriter);
+
+		constexpr u8 crypto[] = { 0x00, 0x01, 0x02, 0x03, 0x04, 0x05, 0x06, 0x07, 0x08, 0x09, 0x0a, 0x0b, 0x0c, 0x0d, 0x0e, 0x0f };
+
+		NetworkClientCreateInfo ci(logWriter);
+		ci.cryptoKey128 = crypto;
+
+		bool ctorSuccess = true;
+		NetworkServer server(ctorSuccess, logWriter);
+		NetworkClient client(ctorSuccess, ci);
+
+		server.RegisterCryptoKey(crypto);
+
+		server.RegisterService(1, [&](const ConnectionInfo& connectionInfo, const WorkContext& workContext, MessageInfo& messageInfo, BinaryReader& reader, BinaryWriter& writer)
+			{
+				logger.Info(TC("Got ping!"));
+				CHECK_TRUE(reader.ReadU32() == 32);
+				UBA_ASSERT(messageInfo.type == SessionMessageType_Ping);
+				writer.WriteByte(42);
+				writer.WriteGuid({});
+				return true;
+			});
+
+		auto ds = MakeGuard([&]() { networkBackend.StopListen(); server.DisconnectClients(); });
+		if (!server.StartListen(networkBackend, 1234, TC("0.0.0.0"), true))
+			return logger.Error(TC("Failed to listen"));
+		Sleep(100);
+		auto dg = MakeGuard([&]() { client.Disconnect(); });
+		if (!client.Connect(networkBackend, TC("127.0.0.1"), 1234))
+			return logger.Error(TC("Failed to connect"));
+
+
+		{
+			StackBinaryWriter<128> writer;
+			NetworkMessage msg(client, 1, SessionMessageType_Ping, writer);
+			writer.WriteU32(32);
+			writer.WriteGuid({});
+			StackBinaryReader<32> reader;
+			if (!msg.Send(reader))
+				return logger.Error(TC("Failed to get message"));
+			u8 value = reader.ReadByte();
+			CHECK_TRUE(value == 42);
+		}
+
+		{
+			StackBinaryWriter<128> writer;
+			NetworkMessage msg(client, 1, SessionMessageType_Ping, writer);
+			writer.WriteU32(32);
+			StackBinaryReader<32> reader;
+			Event ev(true);
+			if (!msg.SendAsync(reader, [](bool error, void* userData) { ((Event*)userData)->Set(); }, &ev))
+				return logger.Error(TC("Failed to get message"));
+			ev.IsSet();
+			if (!msg.ProcessAsyncResults(reader))
+				return false;
+			u8 value = reader.ReadByte();
+			CHECK_TRUE(value == 42);
+		}
+
+		if (!server.SendKeepAlive())
+			return logger.Error(TC("Failed to send keep alive"));
+		Sleep(100);
+		return true;
+	}
 }
