[UBA]
* Fixed so all *_base ucrt alloc related functions check if provided memory is in mimalloc. Fixes crash in ucrtbase when reallocating static initialization memory (used for static deinitialization)
* Renamed DisableCustomAllocator to AllowCustomAllocator to prevent double negation

diff --git a/Engine/Source/Programs/UnrealBuildAccelerator/Cli/Private/UbaCli.cpp b/Engine/Source/Programs/UnrealBuildAccelerator/Cli/Private/UbaCli.cpp
--- a/Engine/Source/Programs/UnrealBuildAccelerator/Cli/Private/UbaCli.cpp
+++ b/Engine/Source/Programs/UnrealBuildAccelerator/Cli/Private/UbaCli.cpp
@@ -1079,7 +1079,7 @@ namespace uba
 		info.detailedTrace = detailedTrace;
 		info.traceChildProcesses = traceChildProcesses;
 		info.launchVisualizer = launchVisualizer;
-		info.disableCustomAllocator = disableCustomAllocator;
+		info.allowCustomAllocator = !disableCustomAllocator;
 		//info.shouldWriteToDisk = shouldWriteToDisk;
 		info.rootDir = g_rootDir.data;
 		//info.traceName.Append(TCV("TESTTRACE"));
diff --git a/Engine/Source/Programs/UnrealBuildAccelerator/Common/Private/UbaSession.cpp b/Engine/Source/Programs/UnrealBuildAccelerator/Common/Private/UbaSession.cpp
--- a/Engine/Source/Programs/UnrealBuildAccelerator/Common/Private/UbaSession.cpp
+++ b/Engine/Source/Programs/UnrealBuildAccelerator/Common/Private/UbaSession.cpp
@@ -217,7 +217,7 @@ namespace uba
 		table.GetValueAsString(extraInfo, TC("ExtraInfo"));
 		table.GetValueAsBool(logToFile, TC("LogToFile"));
 		table.GetValueAsBool(useUniqueId, TC("UseUniqueId"));
-		table.GetValueAsBool(disableCustomAllocator, TC("DisableCustomAllocator"));
+		table.GetValueAsBool(allowCustomAllocator, TC("AllowCustomAllocator"));
 		table.GetValueAsBool(launchVisualizer, TC("LaunchVisualizer"));
 		table.GetValueAsBool(allowMemoryMaps, TC("AllowMemoryMaps"));
 		table.GetValueAsBool(allowKeepFilesInMemory, TC("AllowKeepFilesInMemory"));
@@ -1557,7 +1557,7 @@ namespace uba
 		m_sessionDir.Append(PathSeparator);
 
 		m_runningRemote = runningRemote;
-		m_disableCustomAllocator = info.disableCustomAllocator;
+		m_allowCustomAllocator = info.allowCustomAllocator;
 		m_allowMemoryMaps = info.allowMemoryMaps;
 		m_allowKeepFilesInMemory = info.allowKeepFilesInMemory;
 		m_allowOutputFiles = info.allowOutputFiles;
@@ -1871,7 +1871,6 @@ namespace uba
 	ProcessHandle Session::InternalRunProcess(const ProcessStartInfo& startInfo, bool async, ProcessImpl* parent, bool enableDetour)
 	{
 		auto& si = const_cast<ProcessStartInfo&>(startInfo);
-		si.useCustomAllocator &= !m_disableCustomAllocator;
 		const tchar* originalLogFile = si.logFile;
 		
 		u32 processId = CreateProcessId();
@@ -1900,7 +1899,7 @@ namespace uba
 		auto process = new ProcessImpl(*this, processId, parent, enableDetour);
 		ProcessHandle h(process);
 		if (!process->Start(startInfo, m_runningRemote, env, async))
-			return {};
+			h = {};
 
 		si.logFile = originalLogFile;
 		return h;
@@ -2657,6 +2656,9 @@ namespace uba
 				RunDependencyCrawler(process);
 		}
 
+		if (!m_allowCustomAllocator)
+			startInfo.useCustomAllocator = false;
+
 		return true;
 	}
 
diff --git a/Engine/Source/Programs/UnrealBuildAccelerator/Common/Public/UbaSession.h b/Engine/Source/Programs/UnrealBuildAccelerator/Common/Public/UbaSession.h
--- a/Engine/Source/Programs/UnrealBuildAccelerator/Common/Public/UbaSession.h
+++ b/Engine/Source/Programs/UnrealBuildAccelerator/Common/Public/UbaSession.h
@@ -236,7 +236,7 @@ namespace uba
 
 
 		bool m_runningRemote;
-		bool m_disableCustomAllocator;
+		bool m_allowCustomAllocator;
 		bool m_allowMemoryMaps;
 		bool m_allowKeepFilesInMemory;
 		bool m_allowOutputFiles;
diff --git a/Engine/Source/Programs/UnrealBuildAccelerator/Common/Public/UbaSessionCreateInfo.h b/Engine/Source/Programs/UnrealBuildAccelerator/Common/Public/UbaSessionCreateInfo.h
--- a/Engine/Source/Programs/UnrealBuildAccelerator/Common/Public/UbaSessionCreateInfo.h
+++ b/Engine/Source/Programs/UnrealBuildAccelerator/Common/Public/UbaSessionCreateInfo.h
@@ -23,7 +23,7 @@ namespace uba
 		const tchar* extraInfo = nullptr;		// Extra info that will be stored in the trace info about the session
 		bool logToFile = false;					// Set to true to have all processes write a log file with function calls.
 		bool useUniqueId = true;				// If true, id of session will be "yymmdd_hhmmss". Otherwise "Debug"
-		bool disableCustomAllocator = false;	// Disable detouring of allocator inside processes.
+		bool allowCustomAllocator = true;		// Allow detouring of allocator inside processes.
 		bool launchVisualizer = false;			// Launch a UbaVisualizer process (this automatically enable trace)
 		bool allowMemoryMaps = IsWindows;		// Use memory maps where possible. Session creates memory maps of files that processes use
 		bool allowKeepFilesInMemory = IsWindows;// Allow detoured process to keep output files in memory and send them to session through file mapping handle
diff --git a/Engine/Source/Programs/UnrealBuildAccelerator/Detours/Private/Windows/UbaDetoursFunctionsMiMalloc.inl b/Engine/Source/Programs/UnrealBuildAccelerator/Detours/Private/Windows/UbaDetoursFunctionsMiMalloc.inl
--- a/Engine/Source/Programs/UnrealBuildAccelerator/Detours/Private/Windows/UbaDetoursFunctionsMiMalloc.inl
+++ b/Engine/Source/Programs/UnrealBuildAccelerator/Detours/Private/Windows/UbaDetoursFunctionsMiMalloc.inl
@@ -4,50 +4,77 @@
 
 #define DETOURED_CALL_MEM(x) DETOURED_CALL(x)
 
+#define TRACK_UCRT_ALLOC_ENABLED 0
+
+#if TRACK_UCRT_ALLOC_ENABLED
+Atomic<u64> g_reallocCount;
+Atomic<u64> g_allocCount;
+Atomic<u64> g_freeCount;
+#define TRACK_UCRT_ALLOC(x) ++g_##x##Count;
+#else
+#define TRACK_UCRT_ALLOC(x)
+#endif
+
+inline bool IsInMiMalloc(void* ptr)
+{
+	return (((uintptr_t)ptr - 1) & ~MI_SEGMENT_MASK) && mi_is_in_heap_region(ptr);
+}
+
+#if UBA_DEBUG
+#define VALIDATE_IN_MIMALLOC(ptr) UBA_ASSERT(!ptr || IsInMiMalloc(ptr));
+#else
+#define VALIDATE_IN_MIMALLOC(ptr)
+#endif
+
 void* Detoured_malloc(size_t size)
 {
 	DETOURED_CALL_MEM(malloc);
+	TRACK_UCRT_ALLOC(alloc)
 	return mi_malloc(size);
 }
 
 void* Detoured_calloc(size_t number, size_t size)
 {
 	DETOURED_CALL_MEM(calloc);
+	TRACK_UCRT_ALLOC(alloc)
 	return mi_calloc(number, size);
 }
 
 void* Detoured__recalloc(void* memblock, size_t num, size_t size)
 {
 	DETOURED_CALL_MEM(_recalloc);
+	VALIDATE_IN_MIMALLOC(memblock);
+	TRACK_UCRT_ALLOC(realloc)
 	return mi_recalloc(memblock, num, size);
 }
 
 void* Detoured_realloc(void* memblock, size_t size)
 {
 	DETOURED_CALL_MEM(realloc);
+	VALIDATE_IN_MIMALLOC(memblock);
+	TRACK_UCRT_ALLOC(realloc)
 	return mi_realloc(memblock, size);
 }
 
 void* Detoured__expand(void* memblock, size_t size)
 {
 	DETOURED_CALL_MEM(_expand);
+	VALIDATE_IN_MIMALLOC(memblock);
+	TRACK_UCRT_ALLOC(realloc)
 	return mi_expand(memblock, size);
 }
 
 size_t Detoured__msize(void* memblock)
 {
 	DETOURED_CALL_MEM(free);
+	VALIDATE_IN_MIMALLOC(memblock);
 	return mi_usable_size(memblock);
 }
 
-inline bool IsInMiMalloc(void* ptr)
-{
-	return (((uintptr_t)ptr - 1) & ~MI_SEGMENT_MASK) && mi_is_in_heap_region(ptr);
-}
-
 void Detoured_free(void* memblock)
 {
 	DETOURED_CALL_MEM(free);
+	TRACK_UCRT_ALLOC(free)
 	if (IsInMiMalloc(memblock))
 		return mi_free(memblock);
 	else
@@ -57,60 +84,75 @@ void Detoured_free(void* memblock)
 char* Detoured__strdup(const char* s)
 {
 	DETOURED_CALL_MEM(_strdup);
+	TRACK_UCRT_ALLOC(alloc)
 	return mi_strdup(s);
 }
 
 wchar_t* Detoured__wcsdup(const wchar_t* s)
 {
 	DETOURED_CALL_MEM(_wcsdup);
+	TRACK_UCRT_ALLOC(alloc)
 	return (wchar_t*)mi_wcsdup((const unsigned short*)(s));
 }
 
 wchar_t* Detoured__mbsdup(const wchar_t* s)
 {
 	DETOURED_CALL_MEM(_mbsdup);
+	TRACK_UCRT_ALLOC(alloc)
 	return (wchar_t*)mi_mbsdup((const unsigned char*)(s));
 }
 
 void* Detoured__aligned_malloc(size_t size, size_t alignment)
 {
 	DETOURED_CALL_MEM(_aligned_malloc);
+	TRACK_UCRT_ALLOC(alloc)
 	return mi_malloc_aligned(size, alignment);
 }
 
 void* Detoured__aligned_recalloc(void* memblock, size_t num, size_t size, size_t alignment)
 {
 	DETOURED_CALL_MEM(_aligned_recalloc);
+	VALIDATE_IN_MIMALLOC(memblock);
+	TRACK_UCRT_ALLOC(realloc)
 	return mi_aligned_recalloc(memblock, num, size, alignment);
 }
 
 void* Detoured__aligned_realloc(void* memblock, size_t size, size_t alignment)
 {
 	DETOURED_CALL_MEM(_aligned_realloc);
+	VALIDATE_IN_MIMALLOC(memblock);
+	TRACK_UCRT_ALLOC(realloc)
 	return mi_realloc_aligned(memblock, size, alignment);
 }
 
 void Detoured__aligned_free(void* memblock)
 {
 	DETOURED_CALL_MEM(_aligned_free);
+	VALIDATE_IN_MIMALLOC(memblock);
+	TRACK_UCRT_ALLOC(free)
 	return mi_free(memblock);
 }
 
 void* Detoured__aligned_offset_malloc(size_t size, size_t alignment, size_t offset)
 {
 	DETOURED_CALL_MEM(_aligned_offset_malloc);
+	TRACK_UCRT_ALLOC(alloc)
 	return mi_malloc_aligned_at(size, alignment, offset);
 }
 
 void* Detoured__aligned_offset_recalloc(void* memblock, size_t num, size_t size, size_t alignment, size_t offset)
 {
 	DETOURED_CALL_MEM(_aligned_offset_recalloc);
+	VALIDATE_IN_MIMALLOC(memblock);
+	TRACK_UCRT_ALLOC(realloc)
 	return mi_recalloc_aligned_at(memblock, num, size, alignment, offset);
 }
 
 void* Detoured__aligned_offset_realloc(void* memblock, size_t size, size_t alignment, size_t offset)
 {
 	DETOURED_CALL_MEM(_aligned_offset_realloc);
+	VALIDATE_IN_MIMALLOC(memblock);
+	TRACK_UCRT_ALLOC(realloc)
 	return mi_realloc_aligned_at(memblock, size, alignment, offset);
 }
 
@@ -118,6 +160,7 @@ errno_t Detoured__wdupenv_s(wchar_t** buffer, size_t* numberOfElements, const wc
 {
 	DETOURED_CALL_MEM(_wdupenv_s);
 	auto res = mi_wdupenv_s((unsigned short**)(buffer), numberOfElements, (const unsigned short*)(varname));
+	TRACK_UCRT_ALLOC(alloc)
 	DEBUG_LOG_DETOURED(L"_wdupenv_s", L"(%ls) -> %u", varname, res);
 	return res;
 }
@@ -126,22 +169,28 @@ errno_t Detoured__dupenv_s(char** buffer, size_t* numberOfElements, const char*
 {
 	DETOURED_CALL_MEM(_dupenv_s);
 	auto res = mi_dupenv_s(buffer, numberOfElements, varname);
+	TRACK_UCRT_ALLOC(alloc)
 	DEBUG_LOG_DETOURED(L"_dupenv_s", L"(%hs) -> %u", varname, res);
 	return res;
 }
 
 void* Detoured__malloc_base(size_t _Size)
 {
+	TRACK_UCRT_ALLOC(alloc)
 	return mi_malloc(_Size);
 }
 
 void* Detoured__calloc_base(size_t _Count, size_t _Size)
 {
+	TRACK_UCRT_ALLOC(alloc)
 	return mi_calloc(_Count, _Size);
 }
 
 void* Detoured__realloc_base(void* memblock, size_t size)
 {
+	TRACK_UCRT_ALLOC(realloc)
+	if (memblock && !IsInMiMalloc(memblock))
+		return True__realloc_base(memblock, size);
 	return mi_realloc(memblock, size);
 }
 
@@ -149,6 +198,7 @@ void Detoured__free_base(void* memblock)
 {
 	if (!memblock)
 		return;
+	TRACK_UCRT_ALLOC(free)
 	if (IsInMiMalloc(memblock))
 		mi_free(memblock);
 	else
@@ -157,17 +207,25 @@ void Detoured__free_base(void* memblock)
 
 void* Detoured__expand_base(void* memblock, size_t size)
 {
+	TRACK_UCRT_ALLOC(realloc)
+	if (memblock && !IsInMiMalloc(memblock))
+		return True__expand_base(memblock, size);
 	return mi_expand(memblock, size);
 }
 
 size_t Detoured__msize_base(void* memblock)
 {
 	DETOURED_CALL_MEM(_msize_base);
+	if (memblock && !IsInMiMalloc(memblock))
+		return True__msize_base(memblock);
 	return mi_usable_size(memblock);
 }
 
 void* Detoured__recalloc_base(void* memblock, size_t num, size_t size)
 {
+	TRACK_UCRT_ALLOC(realloc)
+	if (memblock && !IsInMiMalloc(memblock))
+		return True__recalloc_base(memblock, num, size);
 	return mi_recalloc(memblock, num, size);
 }
 
@@ -177,6 +235,7 @@ size_t Detoured__aligned_msize(void* p, size_t alignment, size_t offset)
 {
 	UBA_ASSERTF(false, L"_aligned_msize called but is only detoured in debug");
 	DETOURED_CALL(_aligned_msize);
+	VALIDATE_IN_MIMALLOC(p);
 	return mi_usable_size(p);
 }
 
diff --git a/Engine/Source/Programs/UnrealBuildAccelerator/Detours/Private/Windows/UbaDetoursFunctionsWin.cpp b/Engine/Source/Programs/UnrealBuildAccelerator/Detours/Private/Windows/UbaDetoursFunctionsWin.cpp
--- a/Engine/Source/Programs/UnrealBuildAccelerator/Detours/Private/Windows/UbaDetoursFunctionsWin.cpp
+++ b/Engine/Source/Programs/UnrealBuildAccelerator/Detours/Private/Windows/UbaDetoursFunctionsWin.cpp
@@ -666,6 +666,12 @@ void SendExitMessage(DWORD exitCode, u64 startTime)
 	if (g_consoleStringIndex)
 		Shared_WriteConsole(L"\n", 1, 0);
 
+	#if TRACK_UCRT_ALLOC_ENABLED
+	StringBuffer<256> str;
+	str.Appendf(L"Alloc: %llu Realloc: %llu Free: %llu\n", g_allocCount.load(), g_reallocCount.load(), g_freeCount.load());
+	Shared_WriteConsole(str.data, str.count, 0);
+	#endif
+
 	if (g_trackInputsMem)
 		SendInput();
 
diff --git a/Engine/Source/Programs/UnrealBuildAccelerator/Host/Private/UbaExports.cpp b/Engine/Source/Programs/UnrealBuildAccelerator/Host/Private/UbaExports.cpp
--- a/Engine/Source/Programs/UnrealBuildAccelerator/Host/Private/UbaExports.cpp
+++ b/Engine/Source/Programs/UnrealBuildAccelerator/Host/Private/UbaExports.cpp
@@ -430,7 +430,7 @@ extern "C"
 		info->Apply(uba::GetConfig());
 		info->rootDir = TStrdup(rootDir);
 		info->traceOutputFile = TStrdup(traceOutputFile);
-		info->disableCustomAllocator = disableCustomAllocator;
+		info->allowCustomAllocator = !disableCustomAllocator;
 		info->launchVisualizer = launchVisualizer;
 		info->resetCas = resetCas;
 		info->shouldWriteToDisk = writeToDisk;
