[UBA]
* Fixed crypto compile error in UbaCacheService
* Moved out crypto initvector stack object to outside of crypto to make it easier to change to a sequenced initialization vector if needed

diff --git a/Engine/Source/Programs/UnrealBuildAccelerator/CacheService/Private/UbaCacheService.cpp b/Engine/Source/Programs/UnrealBuildAccelerator/CacheService/Private/UbaCacheService.cpp
--- a/Engine/Source/Programs/UnrealBuildAccelerator/CacheService/Private/UbaCacheService.cpp
+++ b/Engine/Source/Programs/UnrealBuildAccelerator/CacheService/Private/UbaCacheService.cpp
@@ -525,7 +525,8 @@ namespace uba
 								{
 									u8 decryptedKey[16];
 									memcpy(decryptedKey, crypto128Data, 16);
-									if (!Crypto::Decrypt(logger, crypto, decryptedKey, 16))
+									Guid initVector;
+									if (!Crypto::Decrypt(logger, crypto, decryptedKey, 16, initVector))
 										return "Failed to decrypt key from http addcrypto";
 									networkServer.RegisterCryptoKey(decryptedKey, expirationTime, userData);
 								}
diff --git a/Engine/Source/Programs/UnrealBuildAccelerator/Cli/Private/UbaCli.cpp b/Engine/Source/Programs/UnrealBuildAccelerator/Cli/Private/UbaCli.cpp
--- a/Engine/Source/Programs/UnrealBuildAccelerator/Cli/Private/UbaCli.cpp
+++ b/Engine/Source/Programs/UnrealBuildAccelerator/Cli/Private/UbaCli.cpp
@@ -839,7 +839,8 @@ namespace uba
 			if (key == InvalidCryptoKey)
 				return false;
 
-			if (!Crypto::Encrypt(logger, key, destFa.GetData(), u32(fileSize), *(Guid*)crypto128Data))
+			Guid iv = *(Guid*)crypto128Data;
+			if (!Crypto::Encrypt(logger, key, destFa.GetData(), u32(fileSize), iv))
 				return false;
 
 			if (!destFa.Close())
@@ -866,7 +867,8 @@ namespace uba
 			if (key == InvalidCryptoKey)
 				return false;
 
-			if (!Crypto::Decrypt(logger, key, destFa.GetData(), u32(fileSize), *(Guid*)crypto128Data))
+			Guid iv = *(Guid*)crypto128Data;
+			if (!Crypto::Decrypt(logger, key, destFa.GetData(), u32(fileSize), iv))
 				return false;
 
 			if (!destFa.Close())
diff --git a/Engine/Source/Programs/UnrealBuildAccelerator/Common/Private/UbaCrypto.cpp b/Engine/Source/Programs/UnrealBuildAccelerator/Common/Private/UbaCrypto.cpp
--- a/Engine/Source/Programs/UnrealBuildAccelerator/Common/Private/UbaCrypto.cpp
+++ b/Engine/Source/Programs/UnrealBuildAccelerator/Common/Private/UbaCrypto.cpp
@@ -106,7 +106,7 @@ namespace uba
 	}
 
 #if UBA_CRYPTO_TYPE == 1
-	bool BCryptEncryptDecrypt(Logger& logger, bool encrypt, CryptoKey key, u8* data, u32 size, const Guid& initVector)
+	bool BCryptEncryptDecrypt(Logger& logger, bool encrypt, CryptoKey key, u8* data, u32 size, Guid& inOutInitVector)
 	{
 		u8 objectBuffer[1024];
 		u32 objectBufferLen = sizeof(objectBuffer);
@@ -123,19 +123,17 @@ namespace uba
 		}
 		auto g = MakeGuard([&]() { BCryptDestroyKey(newKey); });
 
-		Guid iv = initVector;
-
 		ULONG cipherTextLength = 0;
 
 		if (encrypt)
 		{
 			for (u32 i=0; i!=overflow; ++i)
 				data[alignedSize + i] ^= data[i]; 
-			res = BCryptEncrypt(newKey, data, alignedSize, NULL, (u8*)&iv, kAesBytes128, data, alignedSize, &cipherTextLength, 0);
+			res = BCryptEncrypt(newKey, data, alignedSize, NULL, (u8*)&inOutInitVector, kAesBytes128, data, alignedSize, &cipherTextLength, 0);
 		}
 		else
 		{
-			res = BCryptDecrypt(newKey, data, alignedSize, NULL, (u8*)&iv, kAesBytes128, data, alignedSize, &cipherTextLength, 0);
+			res = BCryptDecrypt(newKey, data, alignedSize, NULL, (u8*)&inOutInitVector, kAesBytes128, data, alignedSize, &cipherTextLength, 0);
 			for (u32 i=0; i!=overflow; ++i)
 				data[alignedSize + i] ^= data[i]; 
 		}
@@ -153,11 +151,10 @@ namespace uba
 		return true;
 	}
 #elif UBA_CRYPTO_TYPE == 2
-	bool OpenSslEncryptDecrypt(CryptoKey key, u8* data, u32 size, bool enc, const Guid& initVector)
+	bool OpenSslEncryptDecrypt(CryptoKey key, u8* data, u32 size, bool enc, Guid& inOutInitVector)
 	{
 		auto& keyData = *(KeyData*)uintptr_t(key);
-		Guid iv = initVector;
-
+ 
 		u32 alignedSize = (size / kAesBytes128) * kAesBytes128;
 		u32 overflow = size - alignedSize;
 
@@ -165,11 +162,11 @@ namespace uba
 		{
 			for (u32 i=0; i!=overflow; ++i)
 				data[alignedSize + i] ^= data[i]; 
-			AES_cbc_encrypt(data, data, alignedSize, &keyData.encryptKey, (u8*)&iv, AES_ENCRYPT);
+			AES_cbc_encrypt(data, data, alignedSize, &keyData.encryptKey, (u8*)&inOutInitVector, AES_ENCRYPT);
 		}
 		else
 		{
-			AES_cbc_encrypt(data, data, alignedSize, &keyData.decryptKey, (u8*)&iv, AES_DECRYPT);
+			AES_cbc_encrypt(data, data, alignedSize, &keyData.decryptKey, (u8*)&inOutInitVector, AES_DECRYPT);
 			for (u32 i=0; i!=overflow; ++i)
 				data[alignedSize + i] ^= data[i]; 
 		}
@@ -177,7 +174,7 @@ namespace uba
 	}
 #endif
 
-	bool EncryptDecrypt(Logger& logger, CryptoKey key, u8* data, u32 size, const Guid& initVector, bool enc)
+	bool EncryptDecrypt(Logger& logger, CryptoKey key, u8* data, u32 size, Guid& inOutInitVector, bool enc)
 	{
 		// We don't fully encrypt data smaller than 16 bytes. There are a few messages that are harmless that does not need to be encrypted since they contain nothing of value
 		if (size < kAesBytes128)
@@ -189,22 +186,22 @@ namespace uba
 		}
 
 #if UBA_CRYPTO_TYPE == 1
-		return BCryptEncryptDecrypt(logger, enc, key, data, size, initVector);
+		return BCryptEncryptDecrypt(logger, enc, key, data, size, inOutInitVector);
 #elif UBA_CRYPTO_TYPE == 2
-		return OpenSslEncryptDecrypt(key, data, size, enc, initVector);
+		return OpenSslEncryptDecrypt(key, data, size, enc, inOutInitVector);
 #else
 		return false;
 #endif
 	}
 
-	bool Crypto::Encrypt(Logger& logger, CryptoKey key, u8* data, u32 size, const Guid& initVector)
+	bool Crypto::Encrypt(Logger& logger, CryptoKey key, u8* data, u32 size, Guid& inOutInitVector)
 	{
-		return EncryptDecrypt(logger, key, data, size, initVector, true);
+		return EncryptDecrypt(logger, key, data, size, inOutInitVector, true);
 	}
 
-	bool Crypto::Decrypt(Logger& logger, CryptoKey key, u8* data, u32 size, const Guid& initVector)
+	bool Crypto::Decrypt(Logger& logger, CryptoKey key, u8* data, u32 size, Guid& inOutInitVector)
 	{
-		return EncryptDecrypt(logger, key, data, size, initVector, false);
+		return EncryptDecrypt(logger, key, data, size, inOutInitVector, false);
 	}
 	
 	bool CryptoFromString(u8* out, u32 outSize, const tchar* str)
diff --git a/Engine/Source/Programs/UnrealBuildAccelerator/Common/Private/UbaNetworkClient.cpp b/Engine/Source/Programs/UnrealBuildAccelerator/Common/Private/UbaNetworkClient.cpp
--- a/Engine/Source/Programs/UnrealBuildAccelerator/Common/Private/UbaNetworkClient.cpp
+++ b/Engine/Source/Programs/UnrealBuildAccelerator/Common/Private/UbaNetworkClient.cpp
@@ -150,9 +150,11 @@ namespace uba
 				{
 					CryptoKey key = client.m_cryptoKey;
 					if (key != InvalidCryptoKey)
-						if (!Crypto::Decrypt(client.m_logger, key, (u8*)&cryptoIV, sizeof(cryptoIV), {}))
+					{
+						Guid initVector;
+						if (!Crypto::Decrypt(client.m_logger, key, (u8*)&cryptoIV, sizeof(cryptoIV), initVector))
 							return false;
-
+					}
 
 					if (!client.ConnectedCallback(rc.backend, rc.backendConnection, cryptoIV))
 						rc.error = 4;
@@ -176,7 +178,8 @@ namespace uba
 			// If server decrypt it to the same blob, we're good on that part
 			u8* encryptedBuffer = handshakeData.AllocWrite(sizeof(EncryptionHandshakeString));//[1024];
 			memcpy(encryptedBuffer, EncryptionHandshakeString, sizeof(EncryptionHandshakeString));
-			if (!Crypto::Encrypt(m_logger, m_cryptoKey, encryptedBuffer, sizeof(EncryptionHandshakeString), {}))
+			Guid initVector;
+			if (!Crypto::Encrypt(m_logger, m_cryptoKey, encryptedBuffer, sizeof(EncryptionHandshakeString), initVector))
 				return false;
 		}
 
@@ -707,7 +710,8 @@ namespace uba
 		if (m_cryptoKey && bodySize)
 		{
 			TimerScope ts(m_encryptTimer);
-			if (!Crypto::Encrypt(m_logger, m_cryptoKey, data + SendHeaderSize, bodySize, connection.cryptoIV))
+			Guid iv = connection.cryptoIV;
+			if (!Crypto::Encrypt(m_logger, m_cryptoKey, data + SendHeaderSize, bodySize, iv))
 			{
 				message.m_error = 8;
 				OnDisconnected(connection, 8);
@@ -748,7 +752,8 @@ namespace uba
 			else if (m_cryptoKey && !message.m_error && message.m_responseSize)
 			{
 				TimerScope ts(m_decryptTimer);
-				if (!Crypto::Decrypt(m_logger, m_cryptoKey, (u8*)message.m_response, message.m_responseSize, connection.cryptoIV))
+				Guid iv = connection.cryptoIV;
+				if (!Crypto::Decrypt(m_logger, m_cryptoKey, (u8*)message.m_response, message.m_responseSize, iv))
 				{
 					message.m_error = 5;
 					OnDisconnected(connection, 5);
@@ -831,7 +836,8 @@ namespace uba
 			UBA_ASSERT(!response.GetPosition());
 			UBA_ASSERT(m_connection);
 			TimerScope ts(m_client->m_decryptTimer);
-			if (!Crypto::Decrypt(m_client->m_logger, m_client->m_cryptoKey, (u8*)m_response, m_responseSize, m_connection->cryptoIV))
+			Guid iv = m_connection->cryptoIV;
+			if (!Crypto::Decrypt(m_client->m_logger, m_client->m_cryptoKey, (u8*)m_response, m_responseSize, iv))
 			{
 				m_error = 10;
 				return false;
diff --git a/Engine/Source/Programs/UnrealBuildAccelerator/Common/Private/UbaNetworkServer.cpp b/Engine/Source/Programs/UnrealBuildAccelerator/Common/Private/UbaNetworkServer.cpp
--- a/Engine/Source/Programs/UnrealBuildAccelerator/Common/Private/UbaNetworkServer.cpp
+++ b/Engine/Source/Programs/UnrealBuildAccelerator/Common/Private/UbaNetworkServer.cpp
@@ -242,7 +242,8 @@ namespace uba
 			{
 				u8 temp[sizeof(EncryptionHandshakeString)];
 				memcpy(temp, handshakeData, sizeof(temp));
-				if (!Crypto::Decrypt(logger, key, temp, sizeof(EncryptionHandshakeString), {}))
+				Guid initVector;
+				if (!Crypto::Decrypt(logger, key, temp, sizeof(EncryptionHandshakeString), initVector))
 					return false;
 				return memcmp(temp, EncryptionHandshakeString, sizeof(EncryptionHandshakeString)) == 0;
 			};
@@ -348,7 +349,8 @@ namespace uba
 			{
 				CreateGuid(cryptoIV);
 				conn.m_cryptoIV = cryptoIV;
-				if (!Crypto::Encrypt(server.m_logger, conn.m_cryptoKey, (u8*)&cryptoIV, sizeof(cryptoIV), {}))
+				Guid initVector;
+				if (!Crypto::Encrypt(server.m_logger, conn.m_cryptoKey, (u8*)&cryptoIV, sizeof(cryptoIV), initVector))
 					return false;
 			}
 
@@ -590,7 +592,8 @@ namespace uba
 		{
 			//TrackHintScope ths(tws, TCV("Decrypt"));
 			TimerScope ts(connection.m_decryptTimer);
-			if (!Crypto::Decrypt(server.m_logger, cryptoKey, context.buffer.data(), context.dataSize, connection.m_cryptoIV))
+			Guid iv = connection.m_cryptoIV;
+			if (!Crypto::Decrypt(server.m_logger, cryptoKey, context.buffer.data(), context.dataSize, iv))
 			{
 				connection.SetShouldDisconnect();
 				connection.Release();
@@ -646,7 +649,8 @@ namespace uba
 				//TrackHintScope ths(tws, TCV("Encrypt"));
 				TimerScope ts(connection.m_encryptTimer);
 				u8* bodyData = writer.GetData() + HeaderSize;
-				if (!Crypto::Encrypt(server.m_logger, cryptoKey, bodyData, bodySize, connection.m_cryptoIV))
+				Guid iv = connection.m_cryptoIV;
+				if (!Crypto::Encrypt(server.m_logger, cryptoKey, bodyData, bodySize, iv))
 				{
 					connection.SetShouldDisconnect();
 					size = MessageErrorSize;
@@ -1215,7 +1219,8 @@ namespace uba
 			{
 				TimerScope ts(connection.m_encryptTimer);
 				u8* bodyData = writer.GetData() + HeaderSize;
-				if (!Crypto::Encrypt(m_logger, connection.m_cryptoKey, bodyData, bodySize, connection.m_cryptoIV))
+				Guid iv = connection.m_cryptoIV;
+				if (!Crypto::Encrypt(m_logger, connection.m_cryptoKey, bodyData, bodySize, iv))
 				{
 					connection.SetShouldDisconnect();
 					return false;
diff --git a/Engine/Source/Programs/UnrealBuildAccelerator/Common/Public/UbaCrypto.h b/Engine/Source/Programs/UnrealBuildAccelerator/Common/Public/UbaCrypto.h
--- a/Engine/Source/Programs/UnrealBuildAccelerator/Common/Public/UbaCrypto.h
+++ b/Engine/Source/Programs/UnrealBuildAccelerator/Common/Public/UbaCrypto.h
@@ -18,8 +18,8 @@ namespace uba
 		static CryptoKey DuplicateKey(Logger& logger, CryptoKey original);
 		static void DestroyKey(CryptoKey key);
 
-		static bool Encrypt(Logger& logger, CryptoKey key, u8* data, u32 size, const Guid& initVector);
-		static bool Decrypt(Logger& logger, CryptoKey key, u8* data, u32 size, const Guid& initVector);
+		static bool Encrypt(Logger& logger, CryptoKey key, u8* data, u32 size, Guid& inOutInitVector);
+		static bool Decrypt(Logger& logger, CryptoKey key, u8* data, u32 size, Guid& inOutInitVector);
 	};
 
 	bool CryptoFromString(u8* out, u32 outSize, const tchar* str);
diff --git a/Engine/Source/Programs/UnrealBuildAccelerator/Test/Private/UbaTestCrypto.cpp b/Engine/Source/Programs/UnrealBuildAccelerator/Test/Private/UbaTestCrypto.cpp
--- a/Engine/Source/Programs/UnrealBuildAccelerator/Test/Private/UbaTestCrypto.cpp
+++ b/Engine/Source/Programs/UnrealBuildAccelerator/Test/Private/UbaTestCrypto.cpp
@@ -25,13 +25,15 @@ namespace uba
 
 			for (u32 i=0; i!=3; ++i)
 			{
-				if (!Crypto::Encrypt(logger, encryptKey, encryptedData, dataSize, {}))
+				Guid iv;
+				if (!Crypto::Encrypt(logger, encryptKey, encryptedData, dataSize, iv))
 					return false;
 
 				if (memcmp(encryptedData, originalData, dataSize) == 0)
 					return false;
 
-				if (!Crypto::Decrypt(logger, decryptKey, encryptedData, dataSize, {}))
+				Guid iv2;
+				if (!Crypto::Decrypt(logger, decryptKey, encryptedData, dataSize, iv2))
 					return false;
 
 				if (memcmp(encryptedData, originalData, dataSize) != 0)
