[UBA]
* Fixed "bug" in scheduler where it would not run cache fetches if all local process slots were full or if forceremote was set. Cache fetching only happened after validating that local processes could be popped.

diff --git a/Engine/Source/Programs/UnrealBuildAccelerator/Common/Private/UbaNetworkServer.cpp b/Engine/Source/Programs/UnrealBuildAccelerator/Common/Private/UbaNetworkServer.cpp
--- a/Engine/Source/Programs/UnrealBuildAccelerator/Common/Private/UbaNetworkServer.cpp
+++ b/Engine/Source/Programs/UnrealBuildAccelerator/Common/Private/UbaNetworkServer.cpp
@@ -893,7 +893,7 @@ namespace uba
 					Crypto::DestroyKey(cryptoKey);
 				done.Set();
 				return 0;
-			});
+			}, TC("UbaAddConSrv"));
 		
 		lock.Leave();
 		done.IsSet();
diff --git a/Engine/Source/Programs/UnrealBuildAccelerator/Common/Private/UbaScheduler.cpp b/Engine/Source/Programs/UnrealBuildAccelerator/Common/Private/UbaScheduler.cpp
--- a/Engine/Source/Programs/UnrealBuildAccelerator/Common/Private/UbaScheduler.cpp
+++ b/Engine/Source/Programs/UnrealBuildAccelerator/Common/Private/UbaScheduler.cpp
@@ -87,6 +87,7 @@ namespace uba
 			table->GetValueAsBool(forceRemote, TC("ForceRemote"));
 			table->GetValueAsBool(forceNative, TC("ForceNative"));
 			table->GetValueAsU32(maxLocalProcessors, TC("MaxLocalProcessors"));
+			table->GetValueAsU32(maxParallelCacheQueries, TC("MaxParallelCacheQueries"));
 		}
 	}
 
@@ -94,6 +95,7 @@ namespace uba
 	Scheduler::Scheduler(const SchedulerCreateInfo& info)
 	:	m_session(info.session)
 	,	m_maxLocalProcessors(info.maxLocalProcessors != ~0u ? info.maxLocalProcessors : GetLogicalProcessorCount())
+	,	m_maxParallelCacheQueries(info.maxParallelCacheQueries)
 	,	m_updateThreadLoop(false)
 	,	m_enableProcessReuse(info.enableProcessReuse)
 	,	m_forceRemote(info.forceRemote)
@@ -131,6 +133,7 @@ namespace uba
 		m_loop = false;
 		m_updateThreadLoop.Set();
 		m_thread.Wait();
+		m_cacheDownloadThreads.clear();
 		m_session.WaitOnAllTasks();
 		SkipAllQueued();
 		Cleanup();
@@ -148,7 +151,7 @@ namespace uba
 		Vector<ProcessStartInfo2*> skipped;
 		SCOPED_FUTEX(m_processEntriesLock, lock);
 		for (auto& entry : m_processEntries)
-			if (entry.status < ProcessStatus_Running)
+			if (entry.status == ProcessStatus_QueuedForCache || entry.status == ProcessStatus_QueuedForRun)
 			{
 				entry.status = ProcessStatus_Skipped;
 				skipped.push_back(entry.info);
@@ -352,12 +355,8 @@ namespace uba
 
 	void Scheduler::HandleCacheMissed(u32 processIndex)
 	{
-		if (processIndex == ~0u)
-			return;
-
 		SCOPED_FUTEX(m_processEntriesLock, lock);
-		if (m_processEntries[processIndex].status != ProcessStatus_Running)
-			return;
+		UBA_ASSERT(m_processEntries[processIndex].status == ProcessStatus_DownloadingFromCache);
 		m_processEntries[processIndex].status = ProcessStatus_QueuedForRun;
 		m_processEntriesStart = Min(m_processEntriesStart, processIndex);
 		--m_activeCacheQueries;
@@ -408,7 +407,6 @@ namespace uba
 	u32 Scheduler::PopProcess(bool isLocal, ProcessStatus& outPrevStatus)
 	{
 		bool atMaxLocalWeight = m_activeLocalProcessWeight >= float(m_maxLocalProcessors);
-		bool atMaxCacheQueries = m_activeCacheQueries >= 16;
 		auto processEntries = m_processEntries.data();
 		bool allFinished = true;
 
@@ -416,11 +414,13 @@ namespace uba
 		{
 			auto& entry = processEntries[i];
 			auto status = entry.status;
+
+			// Check if we can move search starting point forward
 			if (status != ProcessStatus_QueuedForCache && status != ProcessStatus_QueuedForRun)
 			{
 				if (allFinished)
 				{
-					if (status != ProcessStatus_Running)
+					if (status != ProcessStatus_Running && status != ProcessStatus_DownloadingFromCache)
 						m_processEntriesStart = i;
 					else
 						allFinished = false;
@@ -429,21 +429,26 @@ namespace uba
 			}
 			allFinished = false;
 
-			if (isLocal)
+			
+			if (status == ProcessStatus_DownloadingFromCache)
+				continue;
+
+			if (status == ProcessStatus_QueuedForCache)
+			{
+				if (m_activeCacheQueries >= m_maxParallelCacheQueries)
+					continue;
+			}
+			else if (isLocal)
 			{
 				if (m_forceRemote && entry.canExecuteRemotely)
 					continue;
 				if (status == ProcessStatus_QueuedForRun && atMaxLocalWeight)
 					continue;
-				if (status == ProcessStatus_QueuedForCache && atMaxCacheQueries)
-					continue;
 			}
 			else
 			{
 				if (!entry.canExecuteRemotely)
 					continue;
-				if (status == ProcessStatus_QueuedForCache)
-					continue;
 			}
 
 			bool canRun = true;
@@ -471,18 +476,65 @@ namespace uba
 			if (!canRun)
 				continue;
 
-			if (isLocal)
+			if (status != ProcessStatus_QueuedForCache)
 			{
-				if (status == ProcessStatus_QueuedForRun)
+				if (isLocal)
 					m_activeLocalProcessWeight += entry.info->weight;
-				else
-					++m_activeCacheQueries;
+
+				outPrevStatus = entry.status;
+				entry.status = ProcessStatus_Running;
+				return i;
 			}
 
-			outPrevStatus = entry.status;
-			entry.status = ProcessStatus_Running;
-			return i;
+			// Clear out existing download threads if they're done
+			FlushFinishedDownloadThreads();
+
+			++m_activeCacheQueries;
+
+			// Start download
+			entry.status = ProcessStatus_DownloadingFromCache;
+
+			m_cacheDownloadThreads.emplace_back([this, processIndex = i, info = entry.info]()
+				{
+					ProcessStartInfo2& si = *info;
+
+					CacheResult cacheResult;
+					bool isHit = false;
+					for (auto& cacheClient : m_cacheClients)
+					{
+						if (!cacheClient->FetchFromCache(cacheResult, si.rootsHandle, si.cacheBucketId, si) || !cacheResult.hit)
+							continue;
+
+						isHit = true;
+						auto process = new CachedProcess(si);
+						process->logLines.swap(cacheResult.logLines);
+						ProcessHandle ph(process);
+						ExitProcessInfo exitInfo;
+						exitInfo.scheduler = this;
+						exitInfo.startInfo = info;
+						exitInfo.isLocal = true;
+						exitInfo.processIndex = processIndex;
+						ExitProcess(exitInfo, *process, 0, true);
+						++m_cacheHitCount;
+						break;
+					}
+
+					if (!isHit)
+					{
+						HandleCacheMissed(processIndex);
+						++m_cacheMissCount;
+					}
+
+					StringBuffer<> str;
+					str.Appendf(TC("Hits %u Misses %u"), m_cacheHitCount.load(), m_cacheMissCount.load());
+					m_session.GetTrace().StatusUpdate(1, 6, str.data, LogEntryType_Info, nullptr);
+					return 0;
+
+				}, TC("UbaScheDlCach"));
+			continue;
 		}
+
+		FlushFinishedDownloadThreads();
 		return ~0u;
 	}
 
@@ -513,48 +565,6 @@ namespace uba
 
 			UpdateActiveProcessCounter(isLocal, 1);
 	
-			if (prevStatus == ProcessStatus_QueuedForCache)
-			{
-				// TODO: This should not use work manager since it is mostly waiting on network
-				m_session.GetServer().AddWork([this, indexToRun, info](const WorkContext&)
-					{
-						ProcessStartInfo2& si = *info;
-
-						CacheResult cacheResult;
-						bool isHit = false;
-						for (auto& cacheClient : m_cacheClients)
-						{
-							if (!cacheClient->FetchFromCache(cacheResult, si.rootsHandle, si.cacheBucketId, si) || !cacheResult.hit)
-								continue;
-
-							isHit = true;
-							auto process = new CachedProcess(si);
-							process->logLines.swap(cacheResult.logLines);
-							ProcessHandle ph(process);
-							ExitProcessInfo exitInfo;
-							exitInfo.scheduler = this;
-							exitInfo.startInfo = info;
-							exitInfo.isLocal = true;
-							exitInfo.processIndex = indexToRun;
-							ExitProcess(exitInfo, *process, 0, true);
-							++m_cacheHitCount;
-							break;
-						}
-
-						if (!isHit)
-						{
-							HandleCacheMissed(indexToRun);
-							++m_cacheMissCount;
-						}
-
-						StringBuffer<> str;
-						str.Appendf(TC("Hits %u Misses %u"), m_cacheHitCount.load(), m_cacheMissCount.load());
-						m_session.GetTrace().StatusUpdate(1, 6, str.data, LogEntryType_Info, nullptr);
-
-					}, 1, TC("DownloadCache"));
-				return true;
-			}
-
 			auto exitInfo = new ExitProcessInfo();
 			exitInfo->scheduler = this;
 			exitInfo->startInfo = info;
@@ -748,6 +758,16 @@ namespace uba
 		m_session.UpdateProgress(m_totalProcesses, m_finishedProcesses, m_errorCount);
 	}
 
+	void Scheduler::FlushFinishedDownloadThreads()
+	{
+		while (!m_cacheDownloadThreads.empty())
+		{
+			if (!m_cacheDownloadThreads.front().Wait(0))
+				break;
+			m_cacheDownloadThreads.pop_front();
+		}
+	}
+
 	bool Scheduler::EnqueueFromFile(const tchar* yamlFilename, const Function<void(EnqueueProcessInfo&)>& enqueued)
 	{
 		auto& logger = m_session.GetLogger();
diff --git a/Engine/Source/Programs/UnrealBuildAccelerator/Common/Public/UbaScheduler.h b/Engine/Source/Programs/UnrealBuildAccelerator/Common/Public/UbaScheduler.h
--- a/Engine/Source/Programs/UnrealBuildAccelerator/Common/Public/UbaScheduler.h
+++ b/Engine/Source/Programs/UnrealBuildAccelerator/Common/Public/UbaScheduler.h
@@ -27,6 +27,7 @@ namespace uba
 		CacheClient** cacheClients = nullptr; // Set cache clients for scheduler to use when building
 		u32 cacheClientCount = 0;
 		u32 maxLocalProcessors = ~0u; // Max local processors to use. ~0u means it will use all processors
+		u32 maxParallelCacheQueries = 16u; // Max number of parallel cache queries allowed
 		bool enableProcessReuse = false; // If this is true, the system will allow processes to be reused when they're asking for it.
 		bool forceRemote = false; // Force all processes that can run remotely to run remotely.
 		bool forceNative = false; // Force all processes to run native (not detoured)
@@ -89,6 +90,7 @@ namespace uba
 		enum ProcessStatus : u8
 		{
 			ProcessStatus_QueuedForCache,
+			ProcessStatus_DownloadingFromCache,
 			ProcessStatus_QueuedForRun,
 			ProcessStatus_Running,
 			ProcessStatus_Success,
@@ -111,9 +113,11 @@ namespace uba
 		void UpdateQueueCounter(int offset);
 		void UpdateActiveProcessCounter(bool isLocal, int offset);
 		void FinishProcess(const ProcessHandle& handle);
+		void FlushFinishedDownloadThreads();
 
 		SessionServer& m_session;
 		u32 m_maxLocalProcessors;
+		u32 m_maxParallelCacheQueries;
 		
 		struct ProcessEntry
 		{
@@ -156,6 +160,8 @@ namespace uba
 		Vector<RootPaths*> m_rootPaths;
 		bool m_writeToCache;
 
+		List<Thread> m_cacheDownloadThreads;
+
 		Scheduler(const Scheduler&) = delete;
 		void operator=(const Scheduler&) = delete;
 	};
