[UBA]
* Fixed so uba client/server initialization messages does not contain init vector if crypto is not enabled. This makes protocol the same as before for non-crypto connections

diff --git a/Engine/Source/Programs/UnrealBuildAccelerator/Common/Private/UbaNetworkClient.cpp b/Engine/Source/Programs/UnrealBuildAccelerator/Common/Private/UbaNetworkClient.cpp
--- a/Engine/Source/Programs/UnrealBuildAccelerator/Common/Private/UbaNetworkClient.cpp
+++ b/Engine/Source/Programs/UnrealBuildAccelerator/Common/Private/UbaNetworkClient.cpp
@@ -123,7 +123,9 @@ namespace uba
 				rc.exitScopeEvent.Set();
 			});
 
-		constexpr u32 recvSize = 1 + sizeof(Guid)*2;
+		u32 recvSize = 1 + sizeof(Guid);
+		if (m_cryptoKey !=InvalidCryptoKey)
+			recvSize += sizeof(Guid);
 
 		backend.SetRecvCallbacks(backendConnection, &rc, recvSize, [](void* context, const Guid& connectionUid, u8* headerData, void*& outBodyContext, u8*& outBodyData, u32& outBodySize)
 			{
@@ -132,7 +134,11 @@ namespace uba
 
 				rc.error = *headerData;
 				Guid serverUid = *(Guid*)(headerData+1);
-				Guid cryptoIV = *(Guid*)(headerData+17);
+				Guid cryptoIV;
+				
+				CryptoKey key = client.m_cryptoKey;
+				if (key !=InvalidCryptoKey)
+					cryptoIV = *(Guid*)(headerData+17);
 
 				if (serverUid == Guid())
 					rc.error = 5;
@@ -148,7 +154,6 @@ namespace uba
 
 				if (!rc.error)
 				{
-					CryptoKey key = client.m_cryptoKey;
 					if (key != InvalidCryptoKey)
 					{
 						Guid initVector;
diff --git a/Engine/Source/Programs/UnrealBuildAccelerator/Common/Private/UbaNetworkServer.cpp b/Engine/Source/Programs/UnrealBuildAccelerator/Common/Private/UbaNetworkServer.cpp
--- a/Engine/Source/Programs/UnrealBuildAccelerator/Common/Private/UbaNetworkServer.cpp
+++ b/Engine/Source/Programs/UnrealBuildAccelerator/Common/Private/UbaNetworkServer.cpp
@@ -214,9 +214,14 @@ namespace uba
 			u8 data[33];
 			*data = value;
 			*(Guid*)(data+1) = m_server.m_uid;
-			*(Guid*)(data+17) = cryptoIV;
+			u32 dataSize = 1 + sizeof(Guid);
+			if (m_cryptoKey != InvalidCryptoKey)
+			{
+				*(Guid*)(data+17) = cryptoIV;
+				dataSize += sizeof(Guid);
+			}
 			NetworkBackend::SendContext context(NetworkBackend::SendFlags_Async);
-			return m_backend.Send(m_server.m_logger, m_backendConnection, data, sizeof(data), context, TC("UidResponse"));
+			return m_backend.Send(m_server.m_logger, m_backendConnection, data, dataSize, context, TC("UidResponse"));
 		}
 
 		static bool ReceiveHandshakeHeader(void* context, const Guid& connectionUid, u8* headerData, void*& outBodyContext, u8*& outBodyData, u32& outBodySize)
