[UBA]
* Fixed so all applications are cachable by default to get back to the behavior we had before moving to uba scheduler
* Fixed so .exi file symbols are always sorted alphabetically
* Added so file entries in casdb are wiped when hit wipe threshold. Threshold defaults to 1'000'000 entries. Unfortunately we can't incrementally handle this because db has no concept of paths and can't figure out which ones that are outdated or not
* Removed casdb max size in UbaCli and removed SetScheduler call when scheduler is created since it is automatically set

diff --git a/Engine/Source/Programs/UnrealBuildAccelerator/Cli/Private/UbaCli.cpp b/Engine/Source/Programs/UnrealBuildAccelerator/Cli/Private/UbaCli.cpp
--- a/Engine/Source/Programs/UnrealBuildAccelerator/Cli/Private/UbaCli.cpp
+++ b/Engine/Source/Programs/UnrealBuildAccelerator/Cli/Private/UbaCli.cpp
@@ -33,7 +33,7 @@
 namespace uba
 {
 	const tchar*	Version = GetVersionString();
-	u32				DefaultCapacityGb = 20;
+	u32				DefaultCapacityGb = 0;
 	const tchar*	DefaultRootDir = []() {
 		static tchar buf[256];
 		if constexpr (IsWindows)
@@ -1323,8 +1323,6 @@ namespace uba
 			info.writeToCache = populateCache;
 			Scheduler scheduler(info);
 
-			sessionServer.SetOuterScheduler(&scheduler);
-
 			if (StringView(app).EndsWith(TCV(".yaml")))
 			{
 				if (!scheduler.EnqueueFromFile(app.c_str(), [&](EnqueueProcessInfo& epi)
diff --git a/Engine/Source/Programs/UnrealBuildAccelerator/Common/Private/UbaObjectFile.cpp b/Engine/Source/Programs/UnrealBuildAccelerator/Common/Private/UbaObjectFile.cpp
--- a/Engine/Source/Programs/UnrealBuildAccelerator/Common/Private/UbaObjectFile.cpp
+++ b/Engine/Source/Programs/UnrealBuildAccelerator/Common/Private/UbaObjectFile.cpp
@@ -9,6 +9,7 @@
 #include "UbaObjectFileImportLib.h"
 #include "UbaObjectFileLLVMIR.h"
 #include <oodle2.h>
+#include <algorithm>
 
 namespace uba
 {
@@ -153,13 +154,24 @@ namespace uba
 		write(&m_type, 1);
 		write(&verbose, 1);
 
-		// Write all imports
-		for (auto& symbol : m_imports)
 		{
-			write(symbol.c_str(), symbol.size());
+			// Sort imports based on name
+			Vector<const std::string*> sortedImports;
+			sortedImports.resize(m_imports.size());
+			u32 i = 0;
+			for (auto& symbol : m_imports)
+				sortedImports[i++] = &symbol;
+			std::sort(sortedImports.begin(), sortedImports.end(), [](const std::string* a, const std::string* b) { return *a < *b; });
+
+			// Write all imports
+			for (auto symbol : sortedImports)
+			{
+				write(symbol->c_str(), symbol->size());
+				write("\n", 1);
+			}
 			write("\n", 1);
 		}
-		write("\n", 1);
+
 
 		if (verbose)
 		{
@@ -175,12 +187,22 @@ namespace uba
 		}
 		else
 		{
-			u32 count = u32(m_exports.size());
-			write(&count, sizeof(count));
+			// Sort exports based on stringkey
+			struct SortedExport { const StringKey* key; u8 isData; };
+			Vector<SortedExport> sortedExports;
+			sortedExports.resize(m_exports.size());
+			u32 i = 0;
 			for (auto& kv : m_exports)
+				sortedExports[i++] = { &kv.first, kv.second.isData };
+			std::sort(sortedExports.begin(), sortedExports.end(), [](const SortedExport& a, const SortedExport& b) { return *a.key < *b.key; });
+
+			// Write all exports
+			u32 count = u32(sortedExports.size());
+			write(&count, sizeof(count));
+			for (auto& e : sortedExports)
 			{
-				write(&kv.first, sizeof(StringKey));
-				write(&kv.second.isData, 1);
+				write(e.key, sizeof(StringKey));
+				write(&e.isData, 1);
 			}
 		}
 		return true;
diff --git a/Engine/Source/Programs/UnrealBuildAccelerator/Common/Private/UbaScheduler.cpp b/Engine/Source/Programs/UnrealBuildAccelerator/Common/Private/UbaScheduler.cpp
--- a/Engine/Source/Programs/UnrealBuildAccelerator/Common/Private/UbaScheduler.cpp
+++ b/Engine/Source/Programs/UnrealBuildAccelerator/Common/Private/UbaScheduler.cpp
@@ -210,7 +210,7 @@ namespace uba
 		const ApplicationRules* rules = m_session.GetRules(*info2);
 		info2->rules = rules;
 
-		bool useCache = info.cacheBucketId && !m_cacheClients.empty() && !m_writeToCache && rules->IsCacheable();
+		bool useCache = info.cacheBucketId && !m_cacheClients.empty() && rules->IsCacheable();
 		bool canDetour = info.canDetour;
 		bool canExecuteRemotely = info.canExecuteRemotely && info.canDetour;
 
@@ -573,7 +573,7 @@ namespace uba
 
 			ProcessStartInfo si = *info;
 			si.userData = exitInfo;
-			si.trackInputs = m_writeToCache && si.rules->IsCacheable();
+			//si.trackInputs = m_writeToCache && info->cacheBucketId && !m_cacheClients.empty() && info->rules->IsCacheable();
 			si.exitedFunc = [](void* userData, const ProcessHandle& handle, ProcessExitedResponse& response)
 				{
 					auto ei = (ExitProcessInfo*)userData;
diff --git a/Engine/Source/Programs/UnrealBuildAccelerator/Common/Private/UbaStorage.cpp b/Engine/Source/Programs/UnrealBuildAccelerator/Common/Private/UbaStorage.cpp
--- a/Engine/Source/Programs/UnrealBuildAccelerator/Common/Private/UbaStorage.cpp
+++ b/Engine/Source/Programs/UnrealBuildAccelerator/Common/Private/UbaStorage.cpp
@@ -92,6 +92,8 @@ namespace uba
 		table.GetValueAsU64(casCapacityBytes, TC("CasCapacityBytes"));
 		table.GetValueAsBool(keepTransientDataMapped, TC("KeepTransientDataMapped"));
 		table.GetValueAsBool(storeCompressed, TC("StoreCompressed"));
+		table.GetValueAsU32(fileEntryTableWipeThreshold, TC("FileEntryTableWipeThreshold"));
+		table.GetValueAsU32(maxParallelCopyOrLink, TC("MaxParallelCopyOrLink"));
 
 		const tchar* compressor;
 		if (table.GetValueAsString(compressor, TC("Compressor")))
@@ -834,6 +836,7 @@ namespace uba
 	,	m_casDataBuffer(m_logger, &info.workManager)
 	{
 		m_casCapacityBytes = info.casCapacityBytes;
+		m_fileEntryTableWipeThreshold = info.fileEntryTableWipeThreshold;
 		m_storeCompressed = info.storeCompressed;
 		m_manuallyHandleOverflow = info.manuallyHandleOverflow;
 		m_asyncUnmapViewOfFile = info.asyncUnmapViewOfFile;
@@ -1124,6 +1127,9 @@ namespace uba
 		tempFileName.Append(fileName).Append(TCV(".tmp"));
 
 		u32 deleteCount = 0;
+		u32 fileTableSize = 0;
+		u32 casTableSize = 0;
+
 		{
 			FileAccessor tempFile(m_logger, tempFileName.data);
 			if (!tempFile.CreateWrite(false, DefaultAttributes(), 0, m_tempPath.data))
@@ -1150,23 +1156,33 @@ namespace uba
 			auto casLookupSizePtr = (u32*)writer.AllocWrite(4);
 
 			// File table
-			u32 fileTableSize = 0;
-			for (auto& pair : m_fileTableLookup)
+
+			// If we have a million entries in here we likely have lots of old entries we don't care about.
+			// These accumulate over time and never clean up (since there is no way to know what entry belong to what file.
+			// Best is to just wipe it when it gets too big and start from the beginning..
+			// It will force all files cas to be recalculated but it should be really rare.
+			if (m_fileEntryTableWipeThreshold && m_fileEntryTableWipeThreshold < m_fileTableLookup.size())
 			{
-				FileEntry& fileEntry = pair.second;
-				//SCOPED_READ_LOCK(fileEntry.lock, entryLock); // TODO: Remove this? tsan complains about it
-				if (fileEntry.casKey == CasKeyZero || fileEntry.lastWritten == 0 || fileEntry.isTemp || fileEntry.casKey == CasKeyInvalid)
-					continue;
-				writer.WriteStringKey(pair.first);
-				writer.WriteU64(fileEntry.size);
-				writer.WriteU64(fileEntry.lastWritten);
-				writer.WriteCasKey(fileEntry.casKey);
-				++fileTableSize;
+				m_logger.Info(TC("File entry table larger than threshold (%llu/%u), wiping... (Threshold can be tweaked in config)"), m_casLookup.size(), m_fileEntryTableWipeThreshold);
+			}
+			else
+			{
+				for (auto& pair : m_fileTableLookup)
+				{
+					FileEntry& fileEntry = pair.second;
+					//SCOPED_READ_LOCK(fileEntry.lock, entryLock); // TODO: Remove this? tsan complains about it
+					if (fileEntry.casKey == CasKeyZero || fileEntry.lastWritten == 0 || fileEntry.isTemp || fileEntry.casKey == CasKeyInvalid)
+						continue;
+					writer.WriteStringKey(pair.first);
+					writer.WriteU64(fileEntry.size);
+					writer.WriteU64(fileEntry.lastWritten);
+					writer.WriteCasKey(fileEntry.casKey);
+					++fileTableSize;
+				}
 			}
 			*fileTableSizePtr = fileTableSize;
 
 			// Cas table
-			u32 casTableSize = 0;
 			CasEntry* last = nullptr;
 			for (CasEntry* it = m_newestAccessed; it; it = it->nextAccessed)
 			{
@@ -1218,7 +1234,7 @@ namespace uba
 			m_logger.Info(TC("Session needs at least %s to not overflow."), BytesToText(m_casMaxBytes).str);
 
 		u64 duration = GetTime() - startTime;
-		m_logger.Detail(TC("Database saved in %s (deleted %u dropped cas files)"), TimeToText(duration).str, deleteCount);
+		m_logger.Detail(TC("Database saved in %s (%u file entries, %u cas entries. deleted %u dropped cas files)"), TimeToText(duration).str, fileTableSize, casTableSize, deleteCount);
 		return true;
 	}
 
diff --git a/Engine/Source/Programs/UnrealBuildAccelerator/Common/Public/UbaStorage.h b/Engine/Source/Programs/UnrealBuildAccelerator/Common/Public/UbaStorage.h
--- a/Engine/Source/Programs/UnrealBuildAccelerator/Common/Public/UbaStorage.h
+++ b/Engine/Source/Programs/UnrealBuildAccelerator/Common/Public/UbaStorage.h
@@ -103,6 +103,7 @@ namespace uba
 		WorkManager& workManager;
 		const tchar* rootDir;
 		u64 casCapacityBytes = 20llu * 1024 * 1024 * 1024;
+		u32 fileEntryTableWipeThreshold = 1'000'000;
 		u32 maxParallelCopyOrLink = 1000;
 		bool storeCompressed = true;
 		bool manuallyHandleOverflow = false;
@@ -267,6 +268,7 @@ namespace uba
 		u32 m_casEvictedCount = 0;
 		u64 m_casDroppedBytes = 0;
 		u32 m_casDroppedCount = 0;
+		u32 m_fileEntryTableWipeThreshold = 0;
 		bool m_overflowReported = false;
 		bool m_storeCompressed = false;
 		bool m_manuallyHandleOverflow = false;
diff --git a/Engine/Source/Programs/UnrealBuildAccelerator/Core/Private/UbaApplicationRules.cpp b/Engine/Source/Programs/UnrealBuildAccelerator/Core/Private/UbaApplicationRules.cpp
--- a/Engine/Source/Programs/UnrealBuildAccelerator/Core/Private/UbaApplicationRules.cpp
+++ b/Engine/Source/Programs/UnrealBuildAccelerator/Core/Private/UbaApplicationRules.cpp
@@ -111,11 +111,6 @@ namespace uba
 				|| fileName.EndsWith(TCV(".dll"));
 		}
 
-		virtual bool IsCacheable() const override
-		{
-			return true;
-		}
-
 		virtual bool ShouldExtractSymbols(const StringView& fileName) const override
 		{
 			return fileName.EndsWith(TCV(".obj"));
@@ -201,11 +196,6 @@ namespace uba
 				|| fileName.EndsWith(TCV(".dll"));
 		}
 
-		virtual bool IsCacheable() const override
-		{
-			return true;
-		}
-
 		virtual bool CanDependOnCompressedFiles() override
 		{
 			return true; // Yes, .obj files
@@ -344,11 +334,6 @@ namespace uba
 			return true;
 		}
 
-		virtual bool IsCacheable() const override
-		{
-			return true;
-		}
-
 		virtual bool ShouldExtractSymbols(const StringView& fileName) const override
 		{
 			return fileName.EndsWith(TCV(".obj")) || fileName.EndsWith(TCV(".o"));
@@ -590,11 +575,6 @@ namespace uba
 			return Super::FileTypeMaxSize(file, isSystemOrTempFile);
 		}
 
-		virtual bool IsCacheable() const override
-		{
-			return true;
-		}
-
 		virtual bool ShouldExtractSymbols(const StringView& fileName) const override
 		{
 			return fileName.EndsWith(TCV(".obj")) || fileName.EndsWith(TCV(".o"));
diff --git a/Engine/Source/Programs/UnrealBuildAccelerator/Core/Public/UbaApplicationRules.h b/Engine/Source/Programs/UnrealBuildAccelerator/Core/Public/UbaApplicationRules.h
--- a/Engine/Source/Programs/UnrealBuildAccelerator/Core/Public/UbaApplicationRules.h
+++ b/Engine/Source/Programs/UnrealBuildAccelerator/Core/Public/UbaApplicationRules.h
@@ -139,10 +139,10 @@ namespace uba
 			// Do nothing
 		}
 
-		// Is used by uba scheduler to decide if application can be cached. (This is not used by UBT)
+		// Is used by uba scheduler to decide if application can be cached.
 		virtual bool IsCacheable() const
 		{
-			return false;
+			return true;
 		}
 
 		virtual bool CanDependOnCompressedFiles()
