[UBA]
* Fixed so cache hits including executables get the right permissions on posix (+x)

diff --git a/Engine/Source/Programs/UnrealBuildAccelerator/Common/Private/UbaCacheClient.cpp b/Engine/Source/Programs/UnrealBuildAccelerator/Common/Private/UbaCacheClient.cpp
--- a/Engine/Source/Programs/UnrealBuildAccelerator/Common/Private/UbaCacheClient.cpp
+++ b/Engine/Source/Programs/UnrealBuildAccelerator/Common/Private/UbaCacheClient.cpp
@@ -400,6 +400,15 @@ namespace uba
 							path.data, CasKeyString(oldKey).str, CasKeyString(newKey).str, fileEntry.size, fileEntry.lastWritten, ToString(fileEntry.verified), fileSize, fileInfo.lastWriteTime, info.GetDescription());
 					}
 				}
+
+				#if UBA_TRACK_IS_EXECUTABLE
+				if (isOutput)
+				{
+					auto& fileEntry = m_storage.GetOrCreateFileEntry(CaseInsensitiveFs ? ToStringKeyLower(path) : ToStringKey(path));
+					if (fileEntry.isExecutable)
+						casKey = AsExecutable(casKey, true);
+				}
+				#endif
 			}
 
 			#ifndef __clang_analyzer__ // Must be bug in analyzer
@@ -847,6 +856,8 @@ namespace uba
 		// TODO: This probably is not needed but is correct. Since there could be deferred entries of the file we are downloading and writing to
 		//m_storage.ReportFileWrite(path.data, 
 
+		u32 attributes = DefaultAttributes();
+
 		if (IsNormalized(casKey))
 		{
 			DowngradedLogger logger(m_connected, m_logger.m_writer, TC("UbaCacheClientNormalizedDownload"));
@@ -921,9 +932,17 @@ namespace uba
 		}
 		else
 		{
+			#if UBA_TRACK_IS_EXECUTABLE
+			if (IsExecutable(casKey))
+			{
+				casKey = AsExecutable(casKey, false); // Remove exec bit when fetching from casdb
+				attributes = DefaultAttributes(true);
+			}
+			#endif
+
 			DowngradedLogger logger(m_connected, m_logger.m_writer, TC("UbaCacheClientDownload"));
 			bool writeCompressed = m_session.ShouldStoreIntermediateFilesCompressed() && g_globalRules.FileCanBeCompressed(path);
-			if (!fetcher.RetrieveFile(logger, m_client, casKey, path.data, writeCompressed))
+			if (!fetcher.RetrieveFile(logger, m_client, casKey, path.data, writeCompressed, nullptr, attributes))
 				return logger.Debug(TC("Failed to download cache output %s for %s"), path.data, info.GetDescription()).ToFalse();
 		}
 
@@ -946,7 +965,7 @@ namespace uba
 			}
 		}
 #endif
-		if (!m_storage.FakeCopy(casKey, path.data, fetcher.sizeOnDisk, fetcher.lastWritten, false))
+		if (!m_storage.FakeCopy(casKey, path.data, fetcher.sizeOnDisk, fetcher.lastWritten, attributes, false))
 			return false;
 		if (!m_session.RegisterNewFile(path.data))
 			return false;
@@ -1109,6 +1128,11 @@ namespace uba
 				return false;
 
 			casKey = AsCompressed(casKey, true);
+			
+			#if UBA_TRACK_IS_EXECUTABLE
+			if (IsExecutable(casKey))
+				casKey = AsExecutable(casKey, false); // Remove exec bit when storing in casdb
+			#endif
 
 			StorageImpl::CasEntry* casEntry;
 			if (m_storage.HasCasFile(casKey, &casEntry))
diff --git a/Engine/Source/Programs/UnrealBuildAccelerator/Common/Private/UbaCacheServer.cpp b/Engine/Source/Programs/UnrealBuildAccelerator/Common/Private/UbaCacheServer.cpp
--- a/Engine/Source/Programs/UnrealBuildAccelerator/Common/Private/UbaCacheServer.cpp
+++ b/Engine/Source/Programs/UnrealBuildAccelerator/Common/Private/UbaCacheServer.cpp
@@ -1053,6 +1053,10 @@ namespace uba
 								CasKey casKey;
 								bucket.m_casKeyTable.GetKey(casKey, offset);
 								UBA_ASSERT(IsCompressed(casKey));
+
+								if (IsExecutable(casKey))
+									casKey = AsExecutable(casKey, false);
+
 								if (auto value = existingCas.Find(casKey))
 								{
 									touchedCas.push_back(&value->isUsed);
@@ -1997,6 +2001,8 @@ namespace uba
 				u64 offset = outputsReader.Read7BitEncoded();
 				CasKey casKey;
 				bucket.m_casKeyTable.GetKey(casKey, offset);
+				if (IsExecutable(casKey))
+					casKey = AsExecutable(casKey, false);
 				casKeysUsed.push_back(casKey);
 			}
 
diff --git a/Engine/Source/Programs/UnrealBuildAccelerator/Common/Private/UbaSessionServer.cpp b/Engine/Source/Programs/UnrealBuildAccelerator/Common/Private/UbaSessionServer.cpp
--- a/Engine/Source/Programs/UnrealBuildAccelerator/Common/Private/UbaSessionServer.cpp
+++ b/Engine/Source/Programs/UnrealBuildAccelerator/Common/Private/UbaSessionServer.cpp
@@ -1196,7 +1196,7 @@ namespace uba
 		{
 			if (!rules.IsInvisible(destination))
 			{
-				success = m_storage.FakeCopy(casKey, destination.data);
+				success = m_storage.FakeCopy(casKey, destination.data, 0, 0, attributes);
 				if (!success)
 					m_logger.Error(TC("Failed to fake copy cas from %s to %s (%s)"), CasKeyString(casKey).str, destination.data, GetProcessDescription(processId).c_str());
 			}
diff --git a/Engine/Source/Programs/UnrealBuildAccelerator/Common/Private/UbaStorage.cpp b/Engine/Source/Programs/UnrealBuildAccelerator/Common/Private/UbaStorage.cpp
--- a/Engine/Source/Programs/UnrealBuildAccelerator/Common/Private/UbaStorage.cpp
+++ b/Engine/Source/Programs/UnrealBuildAccelerator/Common/Private/UbaStorage.cpp
@@ -15,7 +15,9 @@
 
 namespace uba
 {
-	constexpr u32 CasTableVersion = IsWindows ? 32 : 34;
+	constexpr u32 CasTableVersionWithoutIsExecutable = 34;
+
+	constexpr u32 CasTableVersion = IsWindows ? 32 : 35;
 	constexpr u32 MaxWorkItemsPerAction = 128; // Cap this to not starve other things
 
     CasKey EmptyFileKey = []() { CasKeyHasher hasher; return ToCasKey(hasher, false); }();
@@ -845,6 +847,8 @@ namespace uba
 		if (!m_writeToDisk)
 			m_casCapacityBytes = 0;
 
+		m_casTableVersion = CasTableVersion;
+
 		m_exclusiveMutex = info.exclusiveMutex;
 
 		m_rootDir.count = GetFullPathNameW(info.rootDir, m_rootDir.capacity, m_rootDir.data, NULL);
@@ -993,11 +997,24 @@ namespace uba
 			return false;
 
 		BinaryReader reader(buffer);
-		u32 version = reader.ReadU32();
-		if (version != CasTableVersion)
+		m_casTableVersion = reader.ReadU32();
+
+
+		#if UBA_TRACK_IS_EXECUTABLE
+		bool checkExecutableBit = true;
+		if (m_casTableVersion == CasTableVersionWithoutIsExecutable)
+		{
+			// We don't want to upgrade cas table version automatically since this machine might swap between branches
+			checkExecutableBit = false;
+		}
+		else
+		#endif
+			
+		if (m_casTableVersion != CasTableVersion)
 		{
 			fileGuard.Execute();
-			m_logger.Info(TC("New CasTable version (%u). Deleting all cas files..."), CasTableVersion);
+			m_logger.Info(TC("New CasTable version v%u (Old %u). Deleting all cas files..."), CasTableVersion, m_casTableVersion);
+			m_casTableVersion = CasTableVersion;
 			DeleteAllCas();
 			loadSaveLock.Leave();
 			if (!SaveCasTable(false, false)) // Save it just to get the version down
@@ -1033,6 +1050,15 @@ namespace uba
 
 				auto insres = m_fileTableLookup.try_emplace(fileNameKey);
 				FileEntry& fileEntry = insres.first->second;
+
+				#if UBA_TRACK_IS_EXECUTABLE
+				if (checkExecutableBit && IsExecutable(casKey))
+				{
+					casKey = AsExecutable(casKey, false);
+					fileEntry.isExecutable = true;
+				}
+				#endif
+
 				fileEntry.verified = false;
 				fileEntry.size = size;
 				fileEntry.lastWritten = lastWritten;
@@ -1106,7 +1132,7 @@ namespace uba
 		if (logStats)
 		{
 			u64 duration = GetTime() - startTime;
-			m_logger.Detail(TC("Database loaded from %s (v%u) in %s (contained %llu entries estimated to %s)"), fileName.data, version, TimeToText(duration).str, m_casLookup.size(), BytesToText(m_casTotalBytes).str);
+			m_logger.Detail(TC("Database loaded from %s (v%u) in %s (contained %llu entries estimated to %s)"), fileName.data, m_casTableVersion, TimeToText(duration).str, m_casLookup.size(), BytesToText(m_casTotalBytes).str);
 		}
 
 		return true;
@@ -1167,16 +1193,28 @@ namespace uba
 			}
 			else
 			{
+				#if UBA_TRACK_IS_EXECUTABLE
+				bool useExecutableBit = m_casTableVersion != CasTableVersionWithoutIsExecutable;
+				#endif
+
 				for (auto& pair : m_fileTableLookup)
 				{
 					FileEntry& fileEntry = pair.second;
+					CasKey casKey = fileEntry.casKey;
 					//SCOPED_READ_LOCK(fileEntry.lock, entryLock); // TODO: Remove this? tsan complains about it
-					if (fileEntry.casKey == CasKeyZero || fileEntry.lastWritten == 0 || fileEntry.isTemp || fileEntry.casKey == CasKeyInvalid)
+					if (casKey == CasKeyZero || fileEntry.lastWritten == 0 || fileEntry.isTemp || casKey == CasKeyInvalid)
 						continue;
+
+					#if UBA_TRACK_IS_EXECUTABLE
+					if (fileEntry.isExecutable && useExecutableBit)
+						casKey = AsExecutable(casKey, true);
+					#endif
+
+
 					writer.WriteStringKey(pair.first);
 					writer.WriteU64(fileEntry.size);
 					writer.WriteU64(fileEntry.lastWritten);
-					writer.WriteCasKey(fileEntry.casKey);
+					writer.WriteCasKey(casKey);
 					++fileTableSize;
 				}
 			}
@@ -1573,6 +1611,14 @@ namespace uba
 			fileSize = info.size;
 			storedSize = fileSize;
 			lastWriteTime = info.lastWriteTime;
+
+			#if UBA_TRACK_IS_EXECUTABLE
+			#if PLATFORM_WINDOWS
+			fileEntry.isExecutable = ToView(fileName).EndsWith(TCV(".exe"));
+			#else
+			fileEntry.isExecutable = IsExecutable(info.attributes);
+			#endif
+			#endif
 		}
 
 		if (fileEntry.casKey != CasKeyZero)
@@ -2295,7 +2341,7 @@ namespace uba
 
 	}
 
-	bool StorageImpl::FakeCopy(const CasKey& casKey, const tchar* destination, u64 size, u64 lastWritten, bool deleteExisting)
+	bool StorageImpl::FakeCopy(const CasKey& casKey, const tchar* destination, u64 size, u64 lastWritten, u32 fileAttributes, bool deleteExisting)
 	{
 		if (deleteExisting) 
 			DeleteFileW(destination);
@@ -2311,7 +2357,15 @@ namespace uba
 		fileEntry.lastWritten = lastWritten;
 		fileEntry.size = size;
 		fileEntry.verified = true;
-		fileEntry.isTemp = true;
+		fileEntry.isTemp = deleteExisting; // For now, keep entries created by cache client downloads
+
+		#if UBA_TRACK_IS_EXECUTABLE
+		#if PLATFORM_WINDOWS
+		fileEntry.isExecutable = ToView(destination).EndsWith(TCV(".exe"));
+		#else
+		fileEntry.isExecutable = IsExecutable(fileAttributes);
+		#endif
+		#endif
 		return true;
 	}
 
diff --git a/Engine/Source/Programs/UnrealBuildAccelerator/Common/Private/UbaStorageClient.cpp b/Engine/Source/Programs/UnrealBuildAccelerator/Common/Private/UbaStorageClient.cpp
--- a/Engine/Source/Programs/UnrealBuildAccelerator/Common/Private/UbaStorageClient.cpp
+++ b/Engine/Source/Programs/UnrealBuildAccelerator/Common/Private/UbaStorageClient.cpp
@@ -18,7 +18,7 @@ namespace uba
 		CasKey newKey = key;
 		#ifndef __clang_analyzer__
 		u8 flagField = ((u8*)&key)[19];
-		((u8*)&newKey)[19] = (flagField | u8(4));
+		((u8*)&newKey)[19] = (flagField | IsProxyMask);
 		#endif
 		return newKey;
 	}
diff --git a/Engine/Source/Programs/UnrealBuildAccelerator/Common/Private/UbaStorageUtils.cpp b/Engine/Source/Programs/UnrealBuildAccelerator/Common/Private/UbaStorageUtils.cpp
--- a/Engine/Source/Programs/UnrealBuildAccelerator/Common/Private/UbaStorageUtils.cpp
+++ b/Engine/Source/Programs/UnrealBuildAccelerator/Common/Private/UbaStorageUtils.cpp
@@ -483,7 +483,7 @@ namespace uba
 
 	////////////////////////////////////////////////////////////////////////////////////////////////////
 
-	bool FileFetcher::RetrieveFile(Logger& logger, NetworkClient& client, const CasKey& casKey, const tchar* destination, bool writeCompressed, MemoryBlock* destinationMem)
+	bool FileFetcher::RetrieveFile(Logger& logger, NetworkClient& client, const CasKey& casKey, const tchar* destination, bool writeCompressed, MemoryBlock* destinationMem, u32 attributes)
 	{
 		TimerScope ts(m_stats.recvCas);
 		u8* slot = m_bufferSlots.Pop();
@@ -547,13 +547,13 @@ namespace uba
 		{
 			if (useFileMapping)
 			{
-				if (!destinationFile.CreateMemoryWrite(false, DefaultAttributes(), sizeOnDisk))
+				if (!destinationFile.CreateMemoryWrite(false, attributes, sizeOnDisk))
 					return false;
 				fileMappingMem = destinationFile.GetData();
 			}
 			else
 			{
-				if (!destinationFile.CreateWrite(false, DefaultAttributes(), sizeOnDisk, m_tempPath.data))
+				if (!destinationFile.CreateWrite(false, attributes, sizeOnDisk, m_tempPath.data))
 					return false;
 			}
 		}
diff --git a/Engine/Source/Programs/UnrealBuildAccelerator/Common/Public/UbaRootPaths.h b/Engine/Source/Programs/UnrealBuildAccelerator/Common/Public/UbaRootPaths.h
--- a/Engine/Source/Programs/UnrealBuildAccelerator/Common/Public/UbaRootPaths.h
+++ b/Engine/Source/Programs/UnrealBuildAccelerator/Common/Public/UbaRootPaths.h
@@ -157,7 +157,7 @@ namespace uba
 	inline bool IsNormalized(const CasKey& key)
 	{
 		UBA_ASSERT(key != CasKeyZero);
-		return (((u8*)&key)[19] & 2) == 2;
+		return (((u8*)&key)[19] & IsNormalizedMask) == IsNormalizedMask;
 	}
 
 	inline CasKey AsNormalized(const CasKey& key, bool normalized)
@@ -165,7 +165,7 @@ namespace uba
 		UBA_ASSERT(key != CasKeyZero);
 		CasKey newKey = key;
 		u8 flagField = ((u8*)&key)[19];
-		((u8*)&newKey)[19] = normalized ? (flagField | 2) : (flagField & ~2);
+		((u8*)&newKey)[19] = normalized ? (flagField | IsNormalizedMask) : (flagField & ~IsNormalizedMask);
 		return newKey;
 	}
 }
diff --git a/Engine/Source/Programs/UnrealBuildAccelerator/Common/Public/UbaStorage.h b/Engine/Source/Programs/UnrealBuildAccelerator/Common/Public/UbaStorage.h
--- a/Engine/Source/Programs/UnrealBuildAccelerator/Common/Public/UbaStorage.h
+++ b/Engine/Source/Programs/UnrealBuildAccelerator/Common/Public/UbaStorage.h
@@ -10,6 +10,8 @@
 #include "UbaPathUtils.h"
 #include "UbaStats.h"
 
+#define UBA_TRACK_IS_EXECUTABLE !PLATFORM_WINDOWS
+
 namespace uba
 {
 	class Config;
@@ -54,7 +56,7 @@ namespace uba
 
 		using FormattingFunc = Function<bool(MemoryBlock& destData, const void* sourceData, u64 sourceSize, const tchar* hint)>;
 		virtual bool CopyOrLink(const CasKey& casKey, const tchar* destination, u32 fileAttributes, bool writeCompressed = false, const FormattingFunc& formattingFunc = {}, bool isTemp = false, bool allowHardLink = true) = 0;
-		virtual bool FakeCopy(const CasKey& casKey, const tchar* destination, u64 size = 0, u64 lastWritten = 0, bool deleteExisting = true) = 0;
+		virtual bool FakeCopy(const CasKey& casKey, const tchar* destination, u64 size, u64 lastWritten, u32 fileAttributes, bool deleteExisting = true) = 0;
 		virtual bool GetCasFileName(StringBufferBase& out, const CasKey& casKey) = 0;
 
 		virtual MappedView MapView(const CasKey& casKey, const tchar* hint) = 0;
@@ -165,7 +167,7 @@ namespace uba
 		virtual bool ReportBadCasFile(const CasKey& casKey) override;
 		virtual bool CalculateCasKey(CasKey& out, const tchar* fileName) override;
 		virtual bool CopyOrLink(const CasKey& casKey, const tchar* destination, u32 fileAttributes, bool writeCompressed = false, const FormattingFunc& formattingFunc = {}, bool isTemp = false, bool allowHardLink = true) override;
-		virtual bool FakeCopy(const CasKey& casKey, const tchar* destination, u64 size = 0, u64 lastWritten = 0, bool deleteExisting = true) override;
+		virtual bool FakeCopy(const CasKey& casKey, const tchar* destination, u64 size, u64 lastWritten, u32 fileAttributes, bool deleteExisting = true) override;
 		virtual bool GetCasFileName(StringBufferBase& out, const CasKey& casKey) override;
 		virtual MappedView MapView(const CasKey& casKey, const tchar* hint) override;
 		virtual void UnmapView(const MappedView& view, const tchar* hint) override;
@@ -231,6 +233,10 @@ namespace uba
 			u64 lastWritten = 0;
 			bool verified = false;
 			bool isTemp = false;
+
+			#if UBA_TRACK_IS_EXECUTABLE
+			bool isExecutable = false;
+			#endif
 		};
 		UnorderedMap<StringKey, FileEntry> m_fileTableLookup;
 
@@ -284,6 +290,7 @@ namespace uba
 
 		Futex m_casTableLoadSaveLock;
 		bool m_casTableLoaded = false;
+		u32 m_casTableVersion = 0;
 
 		FileMappingBuffer m_casDataBuffer;
 
diff --git a/Engine/Source/Programs/UnrealBuildAccelerator/Common/Public/UbaStorageUtils.h b/Engine/Source/Programs/UnrealBuildAccelerator/Common/Public/UbaStorageUtils.h
--- a/Engine/Source/Programs/UnrealBuildAccelerator/Common/Public/UbaStorageUtils.h
+++ b/Engine/Source/Programs/UnrealBuildAccelerator/Common/Public/UbaStorageUtils.h
@@ -3,6 +3,7 @@
 #pragma once
 
 #include "UbaBinaryReaderWriter.h"
+#include "UbaFile.h"
 #include "UbaHash.h"
 #include "UbaNetworkMessage.h"
 #include <oodle2.h>
@@ -46,7 +47,7 @@ namespace uba
 
 	struct FileFetcher
 	{
-		bool RetrieveFile(Logger& logger, NetworkClient& client, const CasKey& casKey, const tchar* destination, bool writeCompressed, MemoryBlock* destinationMem = nullptr);
+		bool RetrieveFile(Logger& logger, NetworkClient& client, const CasKey& casKey, const tchar* destination, bool writeCompressed, MemoryBlock* destinationMem = nullptr, u32 attributes = DefaultAttributes());
 
 		BufferSlots& m_bufferSlots;
 		StorageStats& m_stats;
diff --git a/Engine/Source/Programs/UnrealBuildAccelerator/Core/Private/UbaHash.cpp b/Engine/Source/Programs/UnrealBuildAccelerator/Core/Private/UbaHash.cpp
--- a/Engine/Source/Programs/UnrealBuildAccelerator/Core/Private/UbaHash.cpp
+++ b/Engine/Source/Programs/UnrealBuildAccelerator/Core/Private/UbaHash.cpp
@@ -111,7 +111,7 @@ namespace uba
 		CasKeyHasher temp(hasher);
 		uint8_t output[BLAKE3_OUT_LEN];
 		blake3_hasher_finalize((blake3_hasher*)&temp.hasher, output, BLAKE3_OUT_LEN);
-		output[19] = compressed ? 1 : 0;
+		output[19] = compressed ? IsCompressedMask : EmptyMask;
 		return (CasKey&)output;
 	}
 
diff --git a/Engine/Source/Programs/UnrealBuildAccelerator/Core/Public/UbaHash.h b/Engine/Source/Programs/UnrealBuildAccelerator/Core/Public/UbaHash.h
--- a/Engine/Source/Programs/UnrealBuildAccelerator/Core/Public/UbaHash.h
+++ b/Engine/Source/Programs/UnrealBuildAccelerator/Core/Public/UbaHash.h
@@ -114,10 +114,16 @@ namespace uba
 	constexpr inline u64 AlignUp(u64 v, u64 a) { return ((v + a - 1)/a) * a; }
 	constexpr u64 InvalidValue = 0x7ffffffff; // ~34gb - the largest number that can 7bit encode to 5 bytes
 
+	constexpr u8 EmptyMask = 0;
+	constexpr u8 IsCompressedMask = 1;
+	constexpr u8 IsNormalizedMask = 2;
+	constexpr u8 IsProxyMask = 4;
+	constexpr u8 IsExecutableMask = 8;
+
 	inline bool IsCompressed(const CasKey& key)
 	{
 		UBA_ASSERT(key != CasKeyZero);
-		return (((u8*)&key)[19] & 1) == 1;
+		return (((u8*)&key)[19] & IsCompressedMask) == IsCompressedMask;
 	}
 
 	inline CasKey AsCompressed(const CasKey& key, bool compressed)
@@ -126,7 +132,25 @@ namespace uba
 		CasKey newKey = key;
 		#ifndef __clang_analyzer__
 		u8 flagField = ((u8*)&key)[19];
-		((u8*)&newKey)[19] = compressed ? (flagField | u8(1)) : (flagField & ~u8(1));
+		((u8*)&newKey)[19] = compressed ? (flagField | IsCompressedMask) : (flagField & ~IsCompressedMask);
+		#endif
+		return newKey;
+	}
+
+	// For posix we need to set +x on file so we track this on the cas key because it is easier in terms of downloading from cache server etc
+	inline bool IsExecutable(const CasKey& key)
+	{
+		UBA_ASSERT(key != CasKeyZero);
+		return (((u8*)&key)[19] & IsExecutableMask) == IsExecutableMask;
+	}
+
+	inline CasKey AsExecutable(const CasKey& key, bool executable)
+	{
+		UBA_ASSERT(key != CasKeyZero);
+		CasKey newKey = key;
+		#ifndef __clang_analyzer__
+		u8 flagField = ((u8*)&key)[19];
+		((u8*)&newKey)[19] = executable ? (flagField | IsExecutableMask) : (flagField & ~IsExecutableMask);
 		#endif
 		return newKey;
 	}
diff --git a/Engine/Source/Programs/UnrealBuildAccelerator/Test/Private/UbaTest.cpp b/Engine/Source/Programs/UnrealBuildAccelerator/Test/Private/UbaTest.cpp
--- a/Engine/Source/Programs/UnrealBuildAccelerator/Test/Private/UbaTest.cpp
+++ b/Engine/Source/Programs/UnrealBuildAccelerator/Test/Private/UbaTest.cpp
@@ -5,20 +5,25 @@
 
 namespace uba
 {
-	bool CreateTestFile(Logger& logger, StringView testRootDir, StringView fileName, StringView content)
+	bool CreateTestFile(Logger& logger, StringView testRootDir, StringView fileName, StringView content, u32 attributes)
 	{
 		StringBuffer<> testFileName;
-		testFileName.Append(testRootDir).EnsureEndsWithSlash().Append(fileName).FixPathSeparators();
+		return CreateTestFile(testFileName, logger, testRootDir, fileName, content, attributes);
+	}
+
+	bool CreateTestFile(StringBufferBase& outFile, Logger& logger, StringView testRootDir, StringView fileName, StringView content, u32 attributes)
+	{
+		outFile.Clear().Append(testRootDir).EnsureEndsWithSlash().Append(fileName).FixPathSeparators();
 
 		if (fileName.Contains(PathSeparator) || fileName.Contains(NonPathSeparator))
 		{
 			StringBuffer<> testFileDir;
-			testFileDir.AppendDir(testFileName);
+			testFileDir.AppendDir(outFile);
 			DirectoryCache().CreateDirectory(logger, testFileDir.data);
 		}
 
-		FileAccessor file(logger, testFileName.data);
-		if (!file.CreateMemoryWrite(false, DefaultAttributes(), content.count))
+		FileAccessor file(logger, outFile.data);
+		if (!file.CreateMemoryWrite(false, attributes, content.count))
 			return logger.Error(TC("Failed to create file for write"));
 		memcpy(file.GetData(), content.data, content.count);
 		return file.Close();
diff --git a/Engine/Source/Programs/UnrealBuildAccelerator/Test/Private/UbaTest.h b/Engine/Source/Programs/UnrealBuildAccelerator/Test/Private/UbaTest.h
--- a/Engine/Source/Programs/UnrealBuildAccelerator/Test/Private/UbaTest.h
+++ b/Engine/Source/Programs/UnrealBuildAccelerator/Test/Private/UbaTest.h
@@ -2,6 +2,7 @@
 
 #pragma once
 
+#include "UbaFile.h"
 #include "UbaStringBuffer.h"
 
 #define CHECK_TRUE(x) \
@@ -17,7 +18,8 @@ namespace uba
 	class SessionServer;
 	struct ProcessStartInfo;
 
-	bool CreateTestFile(Logger& logger, StringView testRootDir, StringView fileName, StringView content);
+	bool CreateTestFile(Logger& logger, StringView testRootDir, StringView fileName, StringView content, u32 attributes = DefaultAttributes());
+	bool CreateTestFile(StringBufferBase& outFile, Logger& logger, StringView testRootDir, StringView fileName, StringView content, u32 attributes = DefaultAttributes());
 	bool DeleteTestFile(Logger& logger, StringView testRootDir, StringView fileName);
 	bool FileExists(Logger& logger, StringView dir, StringView fileName);
 
diff --git a/Engine/Source/Programs/UnrealBuildAccelerator/Test/Private/UbaTestCache.cpp b/Engine/Source/Programs/UnrealBuildAccelerator/Test/Private/UbaTestCache.cpp
--- a/Engine/Source/Programs/UnrealBuildAccelerator/Test/Private/UbaTestCache.cpp
+++ b/Engine/Source/Programs/UnrealBuildAccelerator/Test/Private/UbaTestCache.cpp
@@ -289,7 +289,6 @@ namespace uba
 	}
 
 	void GetTestAppPath(LoggerWithWriter& logger, StringBufferBase& out);
-	bool CreateTextFile(StringBufferBase& outPath, LoggerWithWriter& logger, const tchar* workingDir, const tchar* fileName, const char* text);
 	StringKey GetKeyAndFixedName(StringBuffer<>& fixedFilePath, const tchar* filePath);
 	void InvalidateCachedInfo(StorageImpl& storage, StringBufferBase& fileName)
 	{
@@ -513,165 +512,153 @@ namespace uba
 		LogWriter& logWriter = logger.m_writer;
 		NetworkBackendTcp tcpBackend(logWriter);
 
-		bool ctorSuccess = true;
-		NetworkServer server(ctorSuccess, { logWriter });
-
 		StringBuffer<MaxPath> rootDir;
 		rootDir.Append(testRootDir).Append(TCV("Uba"));
-		if (!DeleteAllFiles(logger, rootDir.data))
-			return false;
 
-		StorageServerCreateInfo storageServerInfo(server, rootDir.data, logWriter);
-		storageServerInfo.casCapacityBytes = 1024ull * 1024 * 1024;
-		auto& storageServer = *new StorageServer(storageServerInfo);
-		auto ssg = MakeGuard([&]() { delete &storageServer; });
+		StringBuffer<MaxPath> inputFile;
+		StringBuffer<MaxPath> outputFile;
 
 		{
-		CacheServerCreateInfo csci(storageServer, rootDir.data, logWriter);
-		CacheServer cacheServer(csci);
-		if (!cacheServer.Load(false))
-			return false;
+			bool ctorSuccess = true;
+			NetworkServer server(ctorSuccess, { logWriter });
 
+			CHECK_TRUE(DeleteAllFiles(logger, rootDir.data));
 
-		SessionServerCreateInfo sessionInfo(storageServer, server, logWriter);
-		sessionInfo.rootDir = rootDir.data;
-		auto& session = *new SessionServer(sessionInfo);
-		auto sg = MakeGuard([&]() { delete &session; });
+			StorageServerCreateInfo storageServerInfo(server, rootDir.data, logWriter);
+			storageServerInfo.casCapacityBytes = 1024ull * 1024 * 1024;
+			auto& storageServer = *new StorageServer(storageServerInfo);
+			auto ssg = MakeGuard([&]() { delete &storageServer; });
 
-		u16 port = 1356;
-		if (!server.StartListen(tcpBackend, port))
-			return logger.Error(TC("Failed to listen"));
-		auto disconnectServer = MakeGuard([&]() { server.DisconnectClients(); });
+			CacheServerCreateInfo csci(storageServer, rootDir.data, logWriter);
+			CacheServer cacheServer(csci);
+			CHECK_TRUE(cacheServer.Load(false));
 
-		StringBuffer<MaxPath> workingDir;
-		workingDir.Append(testRootDir).Append(TCV("WorkingDir"));
-		if (!DeleteAllFiles(logger, workingDir.data))
-			return false;
-		if (!storageServer.CreateDirectory(workingDir.data))
-			return false;
-		if (!DeleteAllFiles(logger, workingDir.data, false))
-			return false;
+			SessionServerCreateInfo sessionInfo(storageServer, server, logWriter);
+			sessionInfo.rootDir = rootDir.data;
+			auto& session = *new SessionServer(sessionInfo);
+			auto sg = MakeGuard([&]() { delete &session; });
 
-		StringBuffer<> testApp;
-		GetTestAppPath(logger, testApp);
+			u16 port = 1356;
+			CHECK_TRUE(server.StartListen(tcpBackend, port));
+			auto disconnectServer = MakeGuard([&]() { server.DisconnectClients(); });
 
-		StringBuffer<MaxPath> inputFile;
-		if (!CreateTextFile(inputFile, logger, workingDir.data, TC("Input.txt"), "Foo"))
-			return false;
-		StringBuffer<MaxPath> outputFile;
-		if (!CreateTextFile(outputFile, logger, workingDir.data, TC("Output.txt"), "Foo"))
-			return false;
+			StringBuffer<MaxPath> workingDir;
+			workingDir.Append(testRootDir).Append(TCV("WorkingDir"));
+			CHECK_TRUE(DeleteAllFiles(logger, workingDir.data));
+			CHECK_TRUE(storageServer.CreateDirectory(workingDir.data));
+			CHECK_TRUE(DeleteAllFiles(logger, workingDir.data, false));
 
-		StackBinaryWriter<256> inputs;
-		inputs.WriteString(inputFile);
+			StringBuffer<> testApp;
+			GetTestAppPath(logger, testApp);
 
-		StackBinaryWriter<256> outputs;
-		outputs.WriteString(outputFile);
+			CHECK_TRUE(CreateTestFile(inputFile, logger, workingDir, TCV("Input.txt"), TCV("Foo")));
+			CHECK_TRUE(CreateTestFile(outputFile, logger, workingDir, TCV("Output.exe"), TCV("Foo"), DefaultAttributes(true)));
 
-		StackBinaryWriter<256> logLines;
-		logLines.WriteString(TC("Hello"));
-		logLines.WriteByte(1);
+			StackBinaryWriter<256> inputs;
+			inputs.WriteString(inputFile);
 
-		ProcessStartInfo psi;
-		psi.application = testApp.data;
+			StackBinaryWriter<256> outputs;
+			outputs.WriteString(outputFile);
 
-		{
-			NetworkClient client(ctorSuccess, { logWriter });
-			CacheClientCreateInfo ccci(logWriter, storageServer, client, session);
-			ccci.useRoots = false;
-			CacheClient cacheClient(ccci);
-
-			if (!client.Connect(tcpBackend, TC("127.0.0.1"), port))
-				return logger.Error(TC("Failed to connect"));
-			auto disconnectClient = MakeGuard([&]() { client.Disconnect(); });
+			StackBinaryWriter<256> logLines;
+			logLines.WriteString(TC("Hello"));
+			logLines.WriteByte(1);
 
-			//cacheClient.RegisterPathHash(workingDir.data, CasKey(1, 2, 3));
+			ProcessStartInfo psi;
+			psi.application = testApp.data;
 
 			{
-				CacheResult result;
-				if (cacheClient.FetchFromCache(result, RootPaths(), 0, psi) || result.hit)
-					return false;
+				NetworkClient client(ctorSuccess, { logWriter });
+				CacheClientCreateInfo ccci(logWriter, storageServer, client, session);
+				ccci.useRoots = false;
+				CacheClient cacheClient(ccci);
 
-				if (!cacheClient.WriteToCache(RootPaths(), 0, psi, inputs.GetData(), inputs.GetPosition(), outputs.GetData(), outputs.GetPosition(), logLines.GetData(), logLines.GetPosition()))
-					return false;
+				CHECK_TRUE(client.Connect(tcpBackend, TC("127.0.0.1"), port));
+				auto disconnectClient = MakeGuard([&]() { client.Disconnect(); });
 
-				if (!DeleteFileW(outputFile.data))
-					return false;
-				if (FileExists(logger, outputFile.data))
-					return false;
-				if (!cacheClient.FetchFromCache(result, RootPaths(), 0, psi))
-					return false;
-				if (!FileExists(logger, outputFile.data))
-					return false;
-				if (result.logLines.size() != 1)
-					return false;
-				if (result.logLines[0].text != TC("Hello"))
-					return false;
-			}
+				//cacheClient.RegisterPathHash(workingDir.data, CasKey(1, 2, 3));
 
-			{
-				if (!DeleteFileW(inputFile.data))
-					return false;
-				if (!CreateTextFile(inputFile, logger, workingDir.data, TC("Input.txt"), "Bar"))
-					return false;
-				InvalidateCachedInfo(storageServer, inputFile);
+				{
+					CacheResult result;
+					CHECK_TRUE(!cacheClient.FetchFromCache(result, RootPaths(), 0, psi) && !result.hit);
+					CHECK_TRUE(cacheClient.WriteToCache(RootPaths(), 0, psi, inputs.GetData(), inputs.GetPosition(), outputs.GetData(), outputs.GetPosition(), logLines.GetData(), logLines.GetPosition()));
+					CHECK_TRUE(DeleteFileW(outputFile.data));
+					CHECK_TRUE(!FileExists(logger, outputFile.data));
+					CHECK_TRUE(cacheClient.FetchFromCache(result, RootPaths(), 0, psi));
+					CHECK_TRUE(FileExists(logger, outputFile.data));
+					CHECK_TRUE(result.logLines.size() == 1);
+					CHECK_TRUE(result.logLines[0].text == TC("Hello"));
+				}
 
-				CacheResult result;
-				if (cacheClient.FetchFromCache(result, RootPaths(), 0, psi) || result.hit)
-					return false;
+				{
+					CHECK_TRUE(DeleteFileW(inputFile.data));
+					CHECK_TRUE(CreateTestFile(inputFile, logger, workingDir, TCV("Input.txt"), TCV("Bar")));
+					InvalidateCachedInfo(storageServer, inputFile);
 
-				if (!cacheClient.WriteToCache(RootPaths(), 0, psi, inputs.GetData(), inputs.GetPosition(), outputs.GetData(), outputs.GetPosition(), logLines.GetData(), logLines.GetPosition()))
-					return false;
+					CacheResult result;
+					CHECK_TRUE(!cacheClient.FetchFromCache(result, RootPaths(), 0, psi) && !result.hit);
 
-				if (!DeleteFileW(outputFile.data))
-					return false;
-				if (FileExists(logger, outputFile.data))
-					return false;
-				if (!cacheClient.FetchFromCache(result, RootPaths(), 0, psi))
-					return false;
-				if (!FileExists(logger, outputFile.data))
-					return false;
-				if (result.logLines.size() != 1)
-					return false;
-				if (result.logLines[0].text != TC("Hello"))
-					return false;
-			}
-		}
+					CHECK_TRUE(cacheClient.WriteToCache(RootPaths(), 0, psi, inputs.GetData(), inputs.GetPosition(), outputs.GetData(), outputs.GetPosition(), logLines.GetData(), logLines.GetPosition()));
 
-		if (!cacheServer.RunMaintenance(true, true, []() { return false; }))
-			return false;
+					CHECK_TRUE(DeleteFileW(outputFile.data));
+					CHECK_TRUE(!FileExists(logger, outputFile.data));
+					CHECK_TRUE(cacheClient.FetchFromCache(result, RootPaths(), 0, psi));
+					CHECK_TRUE(FileExists(logger, outputFile.data));
 
-		{
-			NetworkClient client(ctorSuccess, { logWriter });
-			CacheClientCreateInfo ccci(logWriter, storageServer, client, session);
-			ccci.useRoots = false;
-			CacheClient cacheClient(ccci);
+					#if !PLATFORM_WINDOWS
+					CHECK_TRUE(IsExecutable(GetFileAttributesW(outputFile.data)));
+					#endif
+
+					CHECK_TRUE(result.logLines.size() == 1);
+					CHECK_TRUE(result.logLines[0].text == TC("Hello"));
+				}
+			}
 
-			if (!client.Connect(tcpBackend, TC("127.0.0.1"), port))
-				return logger.Error(TC("Failed to connect"));
-			auto disconnectClient = MakeGuard([&]() { client.Disconnect(); });
+			CHECK_TRUE(cacheServer.RunMaintenance(true, true, []() { return false; }));
 
 			{
-				CacheResult result;
-				if (!cacheClient.FetchFromCache(result, RootPaths(), 0, psi))
-					return false;
-				if (result.logLines.size() != 1)
-					return false;
-				if (result.logLines[0].text != TC("Hello"))
-					return false;
-			}
-		}
+				NetworkClient client(ctorSuccess, { logWriter });
+				CacheClientCreateInfo ccci(logWriter, storageServer, client, session);
+				ccci.useRoots = false;
+				CacheClient cacheClient(ccci);
 
-		if (!cacheServer.Save())
-			return false;
+				CHECK_TRUE(client.Connect(tcpBackend, TC("127.0.0.1"), port));
+				auto disconnectClient = MakeGuard([&]() { client.Disconnect(); });
 
-		}
+				{
+					CacheResult result;
+					CHECK_TRUE(cacheClient.FetchFromCache(result, RootPaths(), 0, psi));
+					CHECK_TRUE(result.logLines.size() == 1);
+					CHECK_TRUE(result.logLines[0].text == TC("Hello"));
+				}
+			}
 
-		CacheServerCreateInfo csci(storageServer, rootDir.data, logWriter);
-		CacheServer cacheServer(csci);
-		if (!cacheServer.Load(false))
-			return false;
+			CHECK_TRUE(cacheServer.Save());
+			CHECK_TRUE(storageServer.SaveCasTable(true));
+		}
 
+		{
+			bool ctorSuccess = true;
+			NetworkServer server(ctorSuccess, { logWriter });
+			StorageServerCreateInfo storageServerInfo(server, rootDir.data, logWriter);
+			storageServerInfo.casCapacityBytes = 1024ull * 1024 * 1024;
+			StorageServer storageServer(storageServerInfo);
+			storageServer.LoadCasTable();
+
+			StorageImpl::FileEntry* inputFileEntry;
+			StorageImpl::FileEntry* outputFileEntry;
+			CHECK_TRUE(storageServer.GetFileEntry(inputFileEntry, inputFile.data));
+			CHECK_TRUE(storageServer.GetFileEntry(outputFileEntry, outputFile.data));
+
+			#if UBA_TRACK_IS_EXECUTABLE
+			CHECK_TRUE(!inputFileEntry->isExecutable);
+			CHECK_TRUE(outputFileEntry->isExecutable);
+			#endif
+
+			CacheServerCreateInfo csci(storageServer, rootDir.data, logWriter);
+			CacheServer cacheServer(csci);
+			CHECK_TRUE(cacheServer.Load(false));
+		}
 		return true;
 	}
 
diff --git a/Engine/Source/Programs/UnrealBuildAccelerator/Visualizer/Private/UbaVisualizer.cpp b/Engine/Source/Programs/UnrealBuildAccelerator/Visualizer/Private/UbaVisualizer.cpp
--- a/Engine/Source/Programs/UnrealBuildAccelerator/Visualizer/Private/UbaVisualizer.cpp
+++ b/Engine/Source/Programs/UnrealBuildAccelerator/Visualizer/Private/UbaVisualizer.cpp
@@ -2754,7 +2754,7 @@ namespace uba
 						if (file.start < playTime)
 							time = playTime - file.start;
 
-						if ((((u8*)&file.key)[19] & 4) != 0)
+						if ((((u8*)&file.key)[19] & IsProxyMask) != 0)
 							drawText(L"%s (proxy) %s", file.hint.c_str(), TimeToText(time, true).str);
 						else if (file.size == 0)
 							drawText(L"%s (calc) %s", file.hint.c_str(), TimeToText(time, true).str);
