[UBA]
* Implemented setting file time for posix platforms
* Fixed so .isRunning file is touched when casdb is saved without deleting it

diff --git a/Engine/Source/Programs/UnrealBuildAccelerator/Common/Private/UbaFile.cpp b/Engine/Source/Programs/UnrealBuildAccelerator/Common/Private/UbaFile.cpp
--- a/Engine/Source/Programs/UnrealBuildAccelerator/Common/Private/UbaFile.cpp
+++ b/Engine/Source/Programs/UnrealBuildAccelerator/Common/Private/UbaFile.cpp
@@ -539,7 +539,7 @@ namespace uba
 
 		if ((desiredAccess & (GENERIC_WRITE | GENERIC_READ)) == (GENERIC_WRITE | GENERIC_READ))
 			flags |= O_RDWR;
-		else if (desiredAccess & GENERIC_WRITE)
+		else if (desiredAccess & GENERIC_WRITE || desiredAccess == FILE_WRITE_ATTRIBUTES)
 			flags |= O_WRONLY;
 		else if (desiredAccess & GENERIC_READ)
 			flags |= O_RDONLY;
@@ -705,17 +705,17 @@ namespace uba
 #endif
 	}
 
-	bool GetFileLastWriteTime(u64& outTime, FileHandle hFile)
+	bool GetFileLastWriteTime(u64& outTime, FileHandle fileHandle)
 	{
 		ExtendedTimerScope ts(KernelStats::GetCurrent().getFileTime);
 #if PLATFORM_WINDOWS
 		FILETIME lastWriteTime;
-		auto res = ::GetFileTime(asHANDLE(hFile), NULL, NULL, &lastWriteTime);
+		auto res = ::GetFileTime(asHANDLE(fileHandle), NULL, NULL, &lastWriteTime);
 		outTime = (u64&)lastWriteTime;
 		return res;
 #else
 		struct stat attr;
-		int res = fstat(asFileDescriptor(hFile), &attr);
+		int res = fstat(asFileDescriptor(fileHandle), &attr);
 		if (res != 0)
 		{
 			UBA_ASSERTF(errno == ENOENT, TC("GetFileLastWriteTime (fstat) error handling not implemented: %s"), strerror(errno));
@@ -732,7 +732,12 @@ namespace uba
 #if PLATFORM_WINDOWS
 		return SetFileTime(asHANDLE(fileHandle), (FILETIME*)&writeTime, NULL, (FILETIME*)&writeTime);
 #else
-		return true;
+		struct timespec times[2];
+		times[0] = times[1] = ToTimeSpec(writeTime);
+		if (futimens(asFileDescriptor(fileHandle), times) == 0)
+			return true;
+		UBA_ASSERTF(false, TC("SetFileLastWriteTime (futimens) error handling not implemented: %s"), strerror(errno));
+		return false;
 #endif
 	}
 
diff --git a/Engine/Source/Programs/UnrealBuildAccelerator/Common/Private/UbaStorage.cpp b/Engine/Source/Programs/UnrealBuildAccelerator/Common/Private/UbaStorage.cpp
--- a/Engine/Source/Programs/UnrealBuildAccelerator/Common/Private/UbaStorage.cpp
+++ b/Engine/Source/Programs/UnrealBuildAccelerator/Common/Private/UbaStorage.cpp
@@ -831,6 +831,18 @@ namespace uba
 		return m_logger.Warning(TC("Failed to delete %s (%s)"), isRunningName.data, LastErrorToText(lastError).data);
 	}
 
+	bool StorageImpl::TouchIsRunningFile()
+	{
+		StringBuffer<256> isRunningName;
+		isRunningName.Append(m_rootDir).Append(TCV(".isRunning"));
+		FileHandle fh = uba::CreateFileW(isRunningName.data, FILE_WRITE_ATTRIBUTES, FILE_SHARE_ALL, OPEN_EXISTING, FILE_FLAG_BACKUP_SEMANTICS);
+		if (fh == InvalidFileHandle)
+			return false;
+		SetFileLastWriteTime(fh, GetSystemTimeAsFileTime());
+		uba::CloseFile(isRunningName.data, fh);
+		return true;
+	}
+
 	StorageImpl::StorageImpl(const StorageCreateInfo& info, const tchar* logPrefix)
 	:	m_logger(info.writer, logPrefix)
 	,	m_workManager(info.workManager)
@@ -1102,7 +1114,7 @@ namespace uba
 		if (wasTerminated)
 		{
 			u64 fileTime = 0;
-			FileHandle fh = uba::CreateFileW(isRunningName, 0, 0x00000007, 0x00000003, FILE_FLAG_BACKUP_SEMANTICS);
+			FileHandle fh = uba::CreateFileW(isRunningName, 0, FILE_SHARE_ALL, OPEN_EXISTING, FILE_FLAG_BACKUP_SEMANTICS);
 			auto fhg = MakeGuard([&]() { uba::CloseFile(isRunningName, fh); });
 			if (fh != InvalidFileHandle)
 				GetFileLastWriteTime(fileTime, fh);
@@ -1267,6 +1279,8 @@ namespace uba
 
 		if (deleteIsRunningfile)
 			DeleteIsRunningFile();
+		else
+			TouchIsRunningFile();
 
 		if (m_overflowReported)
 			m_logger.Info(TC("Session needs at least %s to not overflow."), BytesToText(m_casMaxBytes).str);
diff --git a/Engine/Source/Programs/UnrealBuildAccelerator/Common/Public/UbaFile.h b/Engine/Source/Programs/UnrealBuildAccelerator/Common/Public/UbaFile.h
--- a/Engine/Source/Programs/UnrealBuildAccelerator/Common/Public/UbaFile.h
+++ b/Engine/Source/Programs/UnrealBuildAccelerator/Common/Public/UbaFile.h
@@ -15,20 +15,27 @@ namespace uba
 	inline constexpr u32 ERROR_ACCESS_DENIED = EACCES;
 	inline constexpr u32 ERROR_DIRECTORY = ENOTDIR;
 	inline constexpr u32 MOVEFILE_REPLACE_EXISTING = 0x00000001;
-	inline constexpr u32 FILE_FLAG_NO_BUFFERING = 0;
-	inline constexpr u32 FILE_FLAG_OVERLAPPED = 0;
 	inline constexpr u32 INVALID_FILE_ATTRIBUTES = (u32)-1;
-	inline constexpr u32 CREATE_ALWAYS = 2;
+
 	inline constexpr u32 GENERIC_WRITE = 0x40000000L;
 	inline constexpr u32 DELETE = 0x00010000L;
 	inline constexpr u32 GENERIC_READ = 0x80000000L;
-	inline constexpr u32 FILE_SHARE_WRITE = 0x00000002;
+	inline constexpr u32 FILE_WRITE_ATTRIBUTES = 0x100L;
+
 	inline constexpr u32 FILE_SHARE_READ = 0x00000001;
+	inline constexpr u32 FILE_SHARE_WRITE = 0x00000002;
+	inline constexpr u32 FILE_SHARE_DELETE = 0x00000004;
+
+	inline constexpr u32 FILE_FLAG_NO_BUFFERING = 0;
+	inline constexpr u32 FILE_FLAG_OVERLAPPED = 0;
 	inline constexpr u32 FILE_FLAG_BACKUP_SEMANTICS = 0x02000000;
 	inline constexpr u32 FILE_FLAG_SEQUENTIAL_SCAN = 0;
+
+	inline constexpr u32 CREATE_ALWAYS = 2;
 	inline constexpr u32 OPEN_EXISTING = 3;
 	inline constexpr u32 PAGE_READONLY = 0x02;
 	#endif
+	inline constexpr u32 FILE_SHARE_ALL = 0x00000007;
 
 	inline constexpr u64 FileHandleFlagMask = 0x0000'0000'ffff'ffff;
 
diff --git a/Engine/Source/Programs/UnrealBuildAccelerator/Common/Public/UbaStorage.h b/Engine/Source/Programs/UnrealBuildAccelerator/Common/Public/UbaStorage.h
--- a/Engine/Source/Programs/UnrealBuildAccelerator/Common/Public/UbaStorage.h
+++ b/Engine/Source/Programs/UnrealBuildAccelerator/Common/Public/UbaStorage.h
@@ -211,6 +211,7 @@ namespace uba
 
 		static MutexHandle GetExclusiveAccess(Logger& logger, const StringView& rootDir, bool reportError);
 		bool DeleteIsRunningFile();
+		bool TouchIsRunningFile();
 
 		struct FileEntry;
 		FileEntry& GetOrCreateFileEntry(const StringKey& fileNameKey);
