[UBA]
* Implemented NetworkServer::FlushWork. It is used when cleaning out cas files but has not been used since we did some refactorings

diff --git a/Engine/Source/Programs/UnrealBuildAccelerator/Common/Private/UbaNetworkServer.cpp b/Engine/Source/Programs/UnrealBuildAccelerator/Common/Private/UbaNetworkServer.cpp
--- a/Engine/Source/Programs/UnrealBuildAccelerator/Common/Private/UbaNetworkServer.cpp
+++ b/Engine/Source/Programs/UnrealBuildAccelerator/Common/Private/UbaNetworkServer.cpp
@@ -1042,7 +1042,29 @@ namespace uba
 
 	bool NetworkServer::FlushWork(u32 timeoutMs)
 	{
-		return m_logger.Error(TC("NetworkServer::FlushWork not implemented"));
+		u64 start = GetTime();
+		if (!timeoutMs)
+			timeoutMs = 2*60*60*1000; // Two hours is the same as forever
+
+		while (true)
+		{
+			SCOPED_FUTEX_READ(m_additionalWorkLock, lock1);
+			bool noAdditionalWork = m_additionalWork.empty();
+			lock1.Leave();
+
+			SCOPED_FUTEX_READ(m_availableWorkersLock, lock2);
+			bool noActiveWorkers = !m_firstActiveWorker;
+			lock2.Leave();
+
+			if (noAdditionalWork && noActiveWorkers)
+				return true;
+
+			Sleep(50);
+
+			u64 delta = GetTime() - start;
+			if (TimeToMs(delta) > timeoutMs)
+				return m_logger.Error(TC("NetworkServer::FlushWork timed out after %s"), TimeToText(delta).str);
+		}
 	}
 
 	MutableLogger& NetworkServer::GetLogger()
