[UBA]
* Added include-what-you-use.exe to application rules. Needed for .d file to be devirtualized when written out when building with vfs paths

diff --git a/Engine/Source/Programs/UnrealBuildAccelerator/Core/Private/UbaApplicationRules.cpp b/Engine/Source/Programs/UnrealBuildAccelerator/Core/Private/UbaApplicationRules.cpp
--- a/Engine/Source/Programs/UnrealBuildAccelerator/Core/Private/UbaApplicationRules.cpp
+++ b/Engine/Source/Programs/UnrealBuildAccelerator/Core/Private/UbaApplicationRules.cpp
@@ -329,11 +329,6 @@ namespace uba
 			return fileName.EndsWith(TCV(".d"));
 		}
 
-		virtual bool AllowMiMalloc() const override
-		{
-			return true;
-		}
-
 		virtual bool ShouldExtractSymbols(const StringView& fileName) const override
 		{
 			return fileName.EndsWith(TCV(".obj")) || fileName.EndsWith(TCV(".o"));
@@ -629,6 +624,20 @@ namespace uba
 #endif // #if PLATFORM_WINDOWS
 	};
 
+	class ApplicationRulesIWYU : public ApplicationRulesClangClExe
+	{
+		virtual bool IsOutputFile(StringView fileName, StringView systemTemp) const override
+		{
+			return fileName.EndsWith(TCV(".c.d"))
+				|| fileName.EndsWith(TCV(".h.d"))
+				|| fileName.EndsWith(TCV(".h.d"))
+				|| fileName.EndsWith(TCV(".cc.d"))
+				|| fileName.EndsWith(TCV(".cpp.d"))
+				|| fileName.EndsWith(TCV(".iwyu"))
+				;
+		}
+	};
+
 	class ApplicationRulesShaderCompileWorker : public ApplicationRules
 	{
 		virtual bool IsRarelyRead(const StringView& file) const override
@@ -738,6 +747,7 @@ namespace uba
 			{ ApplicationHash(TC("ubaobjtool.exe")),			new ApplicationRulesUbaObjTool() },
 			{ ApplicationHash(TC("shadercompileworker.exe")),	new ApplicationRulesShaderCompileWorker() },
 			{ ApplicationHash(TC("instr-clang-cl.exe")),		new ApplicationRulesClangClExe() },
+			{ ApplicationHash(TC("include-what-you-use.exe")),	new ApplicationRulesIWYU() },
 			{ ApplicationHash(TC("ubatestapp.exe")),			new ApplicationRulesUbaTestApp() },
 			//{ L"MSBuild.dll"),				new ApplicationRules() },
 			//{ L"BreakpadSymbolEncoder.exe"),	new ApplicationRulesClang() },
