[UBA]
* Changed so Scheduler::ThreadLoop is never inlined to be able to more easily debug shipping builds

diff --git a/Engine/Source/Programs/UnrealBuildAccelerator/Common/Public/UbaScheduler.h b/Engine/Source/Programs/UnrealBuildAccelerator/Common/Public/UbaScheduler.h
--- a/Engine/Source/Programs/UnrealBuildAccelerator/Common/Public/UbaScheduler.h
+++ b/Engine/Source/Programs/UnrealBuildAccelerator/Common/Public/UbaScheduler.h
@@ -98,7 +98,7 @@ namespace uba
 			ProcessStatus_Skipped,
 		};
 
-		void ThreadLoop();
+		UBA_NOINLINE void ThreadLoop();
 		void SkipAllQueued();
 		void Cleanup();
 		void RemoteProcessReturned(Process& process);
