[UBA]
* Changed assert to error in FileAccessor when failing to write file to disk on posix platforms

diff --git a/Engine/Source/Programs/UnrealBuildAccelerator/Common/Private/UbaFileAccessor.cpp b/Engine/Source/Programs/UnrealBuildAccelerator/Common/Private/UbaFileAccessor.cpp
--- a/Engine/Source/Programs/UnrealBuildAccelerator/Common/Private/UbaFileAccessor.cpp
+++ b/Engine/Source/Programs/UnrealBuildAccelerator/Common/Private/UbaFileAccessor.cpp
@@ -383,10 +383,7 @@ namespace uba
 #else
 			ssize_t written = write(asFileDescriptor(m_fileHandle), pos, toActuallyWrite);
 			if (written == -1)
-			{
-				UBA_ASSERTF(false, TC("WriteFile error handling not implemented for %i (%s)"), errno, strerror(errno));
-				return false;
-			}
+				return m_logger.Error(TC("ERROR writing file %s writing %u bytes (%llu bytes written out of %llu) (%s)"), m_fileName, toActuallyWrite, (dataLen - writeLeft), dataLen, strerror(errno));
 #endif
 			writeLeft -= written;
 			pos += written;
