[UBA]
* Fixed so Create/DestroyDirectory works with vfs
* Fixed so CreateHardLink works with vfs
* Fixed so GetDiskFreeSpaceExW works with vfs

diff --git a/Engine/Source/Programs/UnrealBuildAccelerator/Detours/Private/Windows/UbaDetoursFunctionsKernelBase.inl b/Engine/Source/Programs/UnrealBuildAccelerator/Detours/Private/Windows/UbaDetoursFunctionsKernelBase.inl
--- a/Engine/Source/Programs/UnrealBuildAccelerator/Detours/Private/Windows/UbaDetoursFunctionsKernelBase.inl
+++ b/Engine/Source/Programs/UnrealBuildAccelerator/Detours/Private/Windows/UbaDetoursFunctionsKernelBase.inl
@@ -216,6 +216,7 @@ BOOL Detoured_CreateDirectoryW(LPCWSTR lpPathName, LPSECURITY_ATTRIBUTES lpSecur
 
 	StringBuffer<> pathName;
 	FixPath(pathName, lpPathName);
+
 	if (pathName.StartsWith(g_systemTemp.data))
 	{
 		SuppressCreateFileDetourScope s;
@@ -224,6 +225,8 @@ BOOL Detoured_CreateDirectoryW(LPCWSTR lpPathName, LPSECURITY_ATTRIBUTES lpSecur
 		return res;
 	}
 
+	DevirtualizePath(pathName);
+
 	u32 directoryTableSize;
 	BOOL res;
 	u32 errorCode = 0;
@@ -243,7 +246,7 @@ BOOL Detoured_CreateDirectoryW(LPCWSTR lpPathName, LPSECURITY_ATTRIBUTES lpSecur
 	g_directoryTable.ParseDirectoryTable(directoryTableSize);
 
 	SetLastError(errorCode);
-	DEBUG_LOG_DETOURED(L"CreateDirectoryW", L"%ls -> %ls (%u)", lpPathName, ToString(res), errorCode);
+	DEBUG_LOG_DETOURED(L"CreateDirectoryW", L"%ls -> %ls (%u)", pathName.data, ToString(res), errorCode);
 	return res;
 }
 
@@ -262,6 +265,8 @@ BOOL Detoured_RemoveDirectoryW(LPCWSTR lpPathName)
 		return res;
 	}
 
+	DevirtualizePath(pathName);
+
 	u32 directoryTableSize;
 	BOOL res;
 	u32 errorCode = 0;
@@ -281,7 +286,7 @@ BOOL Detoured_RemoveDirectoryW(LPCWSTR lpPathName)
 	g_directoryTable.ParseDirectoryTable(directoryTableSize);
 
 	SetLastError(errorCode);
-	DEBUG_LOG_DETOURED(L"RemoveDirectoryW", L"%ls -> %ls (%u)", lpPathName, ToString(res), errorCode);
+	DEBUG_LOG_DETOURED(L"RemoveDirectoryW", L"%ls -> %ls (%u)", pathName.data, ToString(res), errorCode);
 	return res;
 }
 
@@ -453,9 +458,9 @@ BOOL Detoured_GetDiskFreeSpaceExW(LPCWSTR lpDirectoryName, PULARGE_INTEGER lpFre
 {
 	DETOURED_CALL(GetDiskFreeSpaceExW);
 	StringBuffer<MaxPath> path;
-	if (g_runningRemote)
+	if (lpDirectoryName)
 	{
-		if (lpDirectoryName)
+		if (g_runningRemote)
 		{
 			UBA_ASSERT(lpDirectoryName[1] == ':');
 			if (ToLower(lpDirectoryName[0]) == ToLower(g_virtualWorkingDir[0]))
@@ -467,6 +472,12 @@ BOOL Detoured_GetDiskFreeSpaceExW(LPCWSTR lpDirectoryName, PULARGE_INTEGER lpFre
 				lpDirectoryName = path.data;
 			}
 		}
+		else
+		{
+			FixPath(path, lpDirectoryName);
+			DevirtualizePath(path);
+			lpDirectoryName = path.data;
+		}
 	}
 
 	DEBUG_LOG_TRUE(L"GetDiskFreeSpaceExW", L"%ls", lpDirectoryName);
@@ -1506,7 +1517,7 @@ BOOL Detoured_CreateHardLinkW(LPCWSTR lpFileName, LPCWSTR lpExistingFileName, LP
 	bool res;
 	{
 		SuppressCreateFileDetourScope cfs;
-		res = True_CreateHardLinkW(lpFileName, lpExistingFileName, lpSecurityAttributes);
+		res = True_CreateHardLinkW(newToName.data, newFromName.data, lpSecurityAttributes);
 	}
 	DEBUG_LOG_TRUE(L"CreateHardLinkW", L"%ls to %ls  (%ls to %ls) -> %ls", lpExistingFileName, lpFileName, newFromName.data, newToName.data, ToString(res));
 
