[UBA]
* Added undocumented -crypto to UbaCacheService to be able to be used for testing
* Fixed so multiple network traffic providers can be registered to session and fixed so all cache servers register/unregister themselves

diff --git a/Engine/Source/Programs/UnrealBuildAccelerator/CacheService/Private/UbaCacheService.cpp b/Engine/Source/Programs/UnrealBuildAccelerator/CacheService/Private/UbaCacheService.cpp
--- a/Engine/Source/Programs/UnrealBuildAccelerator/CacheService/Private/UbaCacheService.cpp
+++ b/Engine/Source/Programs/UnrealBuildAccelerator/CacheService/Private/UbaCacheService.cpp
@@ -135,6 +135,9 @@ namespace uba
 		TString configFile;
 		u32 iocpWorkerCount = 0;
 
+		u8 permanentCrypto[16];
+		bool hasPermanentCrypto = false;
+
 		#if PLATFORM_LINUX
 		bool forkProcess = false;
 		#endif
@@ -209,6 +212,15 @@ namespace uba
 			{
 				fullMaintenance = true;
 			}
+			else if (name.Equals(TCV("-crypto")))
+			{
+				if (value.count != 32)
+					return PrintHelp(TC("Invalid number of characters in crypto string. Should be 32"));
+				((u64*)permanentCrypto)[0] = StringToValue(value.data, 16);
+				((u64*)permanentCrypto)[1] = StringToValue(value.data + 16, 16);
+				hasPermanentCrypto = true;
+				requiresCrypto = true;
+			}
 			else if (name.Equals(TCV("-crash")))
 			{
 				shouldCrash = true;
@@ -505,6 +517,13 @@ namespace uba
 			if (!cacheServer.RunMaintenance(true, allowSave, ShouldExit))
 				return false;
 
+
+		if (hasPermanentCrypto)
+		{
+			logger.Info(TC("Adding permanent crypto. crypto now required"));
+			networkServer.RegisterCryptoKey(permanentCrypto, ~0llu, 0);
+		}
+
 		HttpServer httpServer(logWriter, networkBackend);
 
 		if (httpPort)
diff --git a/Engine/Source/Programs/UnrealBuildAccelerator/Common/Private/UbaSessionServer.cpp b/Engine/Source/Programs/UnrealBuildAccelerator/Common/Private/UbaSessionServer.cpp
--- a/Engine/Source/Programs/UnrealBuildAccelerator/Common/Private/UbaSessionServer.cpp
+++ b/Engine/Source/Programs/UnrealBuildAccelerator/Common/Private/UbaSessionServer.cpp
@@ -617,10 +617,22 @@ namespace uba
 		return m_server;
 	}
 
-	void SessionServer::RegisterNetworkTrafficProvider(const NetworkTrafficProvider& provider)
+	void SessionServer::RegisterNetworkTrafficProvider(u64 key, const NetworkTrafficProvider& provider)
 	{
 		SCOPED_CRITICAL_SECTION(m_remoteProcessAndSessionLock, lock);
-		m_provider = provider;
+		m_providers.emplace_back(key, provider);
+	}
+
+	void SessionServer::UnregisterNetworkTrafficProvider(u64 key)
+	{
+		SCOPED_CRITICAL_SECTION(m_remoteProcessAndSessionLock, lock);
+		for (auto i=m_providers.begin();i!=m_providers.end();++i)
+		{
+			if (i->key != key)
+				continue;
+			m_providers.erase(i);
+			break;
+		}
 	}
 
 	void SessionServer::RegisterCrossArchitectureMapping(const tchar* from, const tchar* to)
@@ -2359,11 +2371,11 @@ namespace uba
 				m_trace.SessionUpdate(sessionIndex, stats.connectionCount, stats.send, stats.recv, s.lastPing, s.memAvail, s.memTotal, s.cpuLoad);
 			++sessionIndex;
 		}
-		if (m_provider)
+		for (auto& rec : m_providers)
 		{
 			u64 send;
 			u64 recv;
-			m_provider(send, recv);
+			rec.provider(send, recv);
 			serverSend += send;
 			serverRecv += recv;
 		}
diff --git a/Engine/Source/Programs/UnrealBuildAccelerator/Common/Public/UbaSessionServer.h b/Engine/Source/Programs/UnrealBuildAccelerator/Common/Public/UbaSessionServer.h
--- a/Engine/Source/Programs/UnrealBuildAccelerator/Common/Public/UbaSessionServer.h
+++ b/Engine/Source/Programs/UnrealBuildAccelerator/Common/Public/UbaSessionServer.h
@@ -74,7 +74,8 @@ namespace uba
 
 		// Additional network traffic provider. Can be used to add more network stats to trace
 		using NetworkTrafficProvider = Function<void(u64& outSent, u64& outReceive)>;
-		void RegisterNetworkTrafficProvider(const NetworkTrafficProvider& provider);
+		void RegisterNetworkTrafficProvider(u64 key, const NetworkTrafficProvider& provider);
+		void UnregisterNetworkTrafficProvider(u64 key);
 
 		// Register mappings for cross architecture helpers (helpers with different architecture will use the mapping)
 		void RegisterCrossArchitectureMapping(const tchar* from, const tchar* to);
@@ -194,7 +195,8 @@ namespace uba
 		struct ApplicationData { Futex lock; Vector<u8> bytes; };
 		UnorderedMap<StringKey, ApplicationData> m_applicationData;
 
-		NetworkTrafficProvider m_provider;
+		struct NetworkTrafficProviderRec { u64 key; NetworkTrafficProvider provider; };
+		Vector<NetworkTrafficProviderRec> m_providers;
 
 		ReaderWriterLock m_crossArchitectureMappingsLock;
 		struct CrossArchitectureMapping { TString from; TString to; };
diff --git a/Engine/Source/Programs/UnrealBuildAccelerator/Host/Private/UbaExports.cpp b/Engine/Source/Programs/UnrealBuildAccelerator/Host/Private/UbaExports.cpp
--- a/Engine/Source/Programs/UnrealBuildAccelerator/Host/Private/UbaExports.cpp
+++ b/Engine/Source/Programs/UnrealBuildAccelerator/Host/Private/UbaExports.cpp
@@ -821,12 +821,6 @@ extern "C"
 			return nullptr;
 		}
 
-		session->RegisterNetworkTrafficProvider([networkClient](u64& outSent, u64& outReceive)
-			{
-				outSent = networkClient->GetTotalSentBytes();
-				outReceive = networkClient->GetTotalRecvBytes();
-			});
-
 		CacheClientCreateInfo info{writer, storage, *networkClient, *session};
 		info.Apply(GetConfig());
 		if (hint && *hint)
@@ -850,6 +844,12 @@ extern "C"
 		while (--desired)
 			networkClient.Connect(*networkClient.backend, host, u16(port));
 
+		((SessionServer&)cacheClient->GetSession()).RegisterNetworkTrafficProvider(u64(cacheClient), [nc = &networkClient](u64& outSent, u64& outReceive)
+			{
+				outSent = nc->GetTotalSentBytes();
+				outReceive = nc->GetTotalRecvBytes();
+			});
+
 		cacheClient->GetStorage().LoadCasTable();
 		return true;
 	}
@@ -935,7 +935,7 @@ extern "C"
 		using namespace uba;
 
 		auto& session = (SessionServer&)cacheClient->GetSession();
-		session.RegisterNetworkTrafficProvider({});
+		session.UnregisterNetworkTrafficProvider(u64(cacheClient));
 
 
 		auto& networkClient = (NetworkClientWithBackend&)cacheClient->GetClient();
