[UBA]
* Added support for "-trace" on UbaCacheService to be able to trace out work. Trace file will be written on process exit
* Added so connection count is visible when hovering over session graph in visualizer
* Added so connect timeout time is set to 10 seconds so a timeout does not take forever if there are things like wrong encryption etc set

diff --git a/Engine/Source/Programs/UnrealBuildAccelerator/CacheService/Private/UbaCacheService.cpp b/Engine/Source/Programs/UnrealBuildAccelerator/CacheService/Private/UbaCacheService.cpp
--- a/Engine/Source/Programs/UnrealBuildAccelerator/CacheService/Private/UbaCacheService.cpp
+++ b/Engine/Source/Programs/UnrealBuildAccelerator/CacheService/Private/UbaCacheService.cpp
@@ -12,6 +12,7 @@
 #include "UbaPlatform.h"
 #include "UbaProtocol.h"
 #include "UbaStorageServer.h"
+#include "UbaTrace.h"
 #include "UbaVersion.h"
 
 #if PLATFORM_WINDOWS
@@ -65,6 +66,7 @@ namespace uba
 		logger.Info(TC("  -nosignalhandler        Will not hook up signal handler"));
 		logger.Info(TC("  -maxworkers=<number>    Max number of workers used by cache server. Defaults to \"%u\""), DefaultWorkerCount);
 		logger.Info(TC("  -reportinterval=<sec>   How often the service should report status. Defaults to \"%s\""), TimeToText(MsToTime(DefaultReportIntervalSeconds*1000)).str);
+		logger.Info(TC("  -trace                  Create a trace file that is written out when cache service exits"));
 #if PLATFORM_LINUX
 		logger.Info(TC("  -fork                   Will handle segfaults and restart"));
 #elif PLATFORM_WINDOWS
@@ -129,6 +131,7 @@ namespace uba
 		bool signalHandlerEnabled = true;
 		bool shouldCrash = false;
 		bool requiresCrypto = false;
+		bool writeTrace = false;
 		u32 expirationTimeSeconds = DefaultExpiration;
 		u32 maxWorkerCount = DefaultWorkerCount;
 		u32 reportIntervalSeconds = DefaultReportIntervalSeconds;
@@ -225,6 +228,10 @@ namespace uba
 			{
 				shouldCrash = true;
 			}
+			else if (name.Equals(TCV("-trace")))
+			{
+				writeTrace = true;
+			}
 			else if (name.Equals(TCV("-nosignalhandler")))
 			{
 				signalHandlerEnabled = false;
@@ -497,6 +504,48 @@ namespace uba
 		storageInfo.allowDeleteVerified = true;
 		StorageServer storageServer(storageInfo);
 
+		Trace trace(logger.m_writer);
+		Thread traceThread;
+		Event traceLoop;
+
+		if (writeTrace)
+		{
+			u64 traceReserveSize = 128ull*1024*1024;
+			trace.StartWrite(nullptr, traceReserveSize);
+			trace.SessionAdded(0, 0, TCV("CacheServer"), {});
+			trace.ProcessAdded(0, 0, TCV("Running"), {});
+			networkServer.SetWorkTracker(&trace);
+			traceLoop.Create(true);
+			traceThread.Start([&]()
+				{
+					while (!traceLoop.IsSet(500))
+					{
+						trace.SessionUpdate(0, networkServer.GetConnectionCount(), networkServer.GetTotalSentBytes(), networkServer.GetTotalRecvBytes(), 0, 0, 0, 0);
+					}
+					return 0;
+				}, TC("UbaCaTrc"));
+		}
+		auto endTrace = MakeGuard([&]()
+			{
+				if (!writeTrace)
+					return;
+				traceLoop.Set();
+				traceThread.Wait();
+				networkServer.SetWorkTracker(nullptr);
+				trace.ProcessExited(0, 0);
+				trace.WriteSessionSummary([&](Logger& logger)
+					{
+						storageServer.PrintSummary(logger);
+						networkServer.PrintSummary(logger);
+						KernelStats::GetGlobal().Print(logger, true);
+						PrintContentionSummary(logger);
+					});
+				StringBuffer<> traceFile(g_rootDir.data);
+				traceFile.EnsureEndsWithSlash().Append(TCV("UbaCacheServer.uba"));
+				trace.StopWrite(traceFile.data);
+			});
+
+
 		bool wasTerminated = false;
 		if (!storageServer.LoadCasTable(true, true, &wasTerminated))
 			return false;
diff --git a/Engine/Source/Programs/UnrealBuildAccelerator/Common/Private/UbaCacheServer.cpp b/Engine/Source/Programs/UnrealBuildAccelerator/Common/Private/UbaCacheServer.cpp
--- a/Engine/Source/Programs/UnrealBuildAccelerator/Common/Private/UbaCacheServer.cpp
+++ b/Engine/Source/Programs/UnrealBuildAccelerator/Common/Private/UbaCacheServer.cpp
@@ -719,13 +719,10 @@ namespace uba
 		}
 		auto endTrace = MakeGuard([&]()
 			{
+				if (!forceAllSteps)
+					return;
 				m_server.SetWorkTracker(nullptr);
-				StackBinaryWriter<1024> writer;
-				ProcessStats().Write(writer);
-				SessionStats().Write(writer);
-				StorageStats().Write(writer);
-				KernelStats().Write(writer);
-				trace.ProcessExited(0, 0, writer.GetData(), writer.GetPosition(), {});
+				trace.ProcessExited(0, 0);
 				StringBuffer<> traceFile(m_rootDir);
 				traceFile.Append(TCV("UbaCacheServer.uba"));
 				trace.StopWrite(traceFile.data);
diff --git a/Engine/Source/Programs/UnrealBuildAccelerator/Common/Private/UbaFileMapping.cpp b/Engine/Source/Programs/UnrealBuildAccelerator/Common/Private/UbaFileMapping.cpp
--- a/Engine/Source/Programs/UnrealBuildAccelerator/Common/Private/UbaFileMapping.cpp
+++ b/Engine/Source/Programs/UnrealBuildAccelerator/Common/Private/UbaFileMapping.cpp
@@ -193,6 +193,7 @@ namespace uba
 				return {};
 			}
 		}
+		SetLastError(0);
 		return { shmFd, lockFd, uid };
 #endif
 	}
diff --git a/Engine/Source/Programs/UnrealBuildAccelerator/Common/Private/UbaNetworkClient.cpp b/Engine/Source/Programs/UnrealBuildAccelerator/Common/Private/UbaNetworkClient.cpp
--- a/Engine/Source/Programs/UnrealBuildAccelerator/Common/Private/UbaNetworkClient.cpp
+++ b/Engine/Source/Programs/UnrealBuildAccelerator/Common/Private/UbaNetworkClient.cpp
@@ -94,10 +94,17 @@ namespace uba
 			}
 
 			~RecvContext()
+			{
+				Wait();
+			}
+
+			bool Wait()
 			{
 				if (error)
 					backend.Shutdown(backendConnection);
-				exitScopeEvent.IsSet(~0u);
+				else
+					backend.SetRecvTimeout(backendConnection, client.m_receiveTimeoutSeconds*1000, nullptr, nullptr);
+				return exitScopeEvent.IsSet(~0u) && error == 0;
 			}
 
 			NetworkClient& client;
@@ -108,7 +115,7 @@ namespace uba
 			Atomic<u8> error;
 		};
 
-		backend.SetRecvTimeout(backendConnection, m_receiveTimeoutSeconds*1000, nullptr, nullptr);
+		backend.SetRecvTimeout(backendConnection, 10*1000, nullptr, nullptr); // 10 seconds timeout during adding connection. RecvContext set it to right time in ctor
 
 		RecvContext rc(*this, backend, backendConnection);
 
@@ -239,7 +246,7 @@ namespace uba
 		m_isConnected.Set();
 		lock.Leave();
 
-		return true;
+		return rc.Wait();
 	}
 
 	constexpr u32 SendHeaderSize = 6;
diff --git a/Engine/Source/Programs/UnrealBuildAccelerator/Common/Private/UbaSession.cpp b/Engine/Source/Programs/UnrealBuildAccelerator/Common/Private/UbaSession.cpp
--- a/Engine/Source/Programs/UnrealBuildAccelerator/Common/Private/UbaSession.cpp
+++ b/Engine/Source/Programs/UnrealBuildAccelerator/Common/Private/UbaSession.cpp
@@ -3880,32 +3880,6 @@ namespace uba
 		return true;
 	}
 
-	void Session::WriteSummary(BinaryWriter& writer, const Function<void(Logger& logger)>& summaryFunc)
-	{
-		struct SummaryLogWriter : public LogWriter
-		{
-			SummaryLogWriter(BinaryWriter& w) : writer(w) {}
-			virtual void BeginScope() override {}
-			virtual void EndScope() override {}
-			virtual void Log(LogEntryType type, const tchar* str, u32 strLen, const tchar* prefix, u32 prefixLen) override
-			{
-				writer.WriteString(str, strLen);
-				++count;
-			}
-
-			BinaryWriter& writer;
-			u32 count = 0;
-		};
-		u32* lineCount = (u32*)writer.AllocWrite(4);
-
-		SummaryLogWriter logWriter(writer);
-		LoggerWithWriter logger(logWriter, TC(""));
-
-		summaryFunc(logger);
-
-		*lineCount = logWriter.count;
-	}
-
 	bool GetCpuTime(u64& outTotalTime, u64& outIdleTime)
 	{
 		outTotalTime = 0;
diff --git a/Engine/Source/Programs/UnrealBuildAccelerator/Common/Private/UbaSessionClient.cpp b/Engine/Source/Programs/UnrealBuildAccelerator/Common/Private/UbaSessionClient.cpp
--- a/Engine/Source/Programs/UnrealBuildAccelerator/Common/Private/UbaSessionClient.cpp
+++ b/Engine/Source/Programs/UnrealBuildAccelerator/Common/Private/UbaSessionClient.cpp
@@ -1767,7 +1767,7 @@ namespace uba
 
 		writer.WriteU32(m_sessionId);
 
-		WriteSummary(writer, [&](Logger& logger)
+		Trace::WriteSummaryText(writer, [&](Logger& logger)
 			{
 				PrintSummary(logger);
 				m_storage.PrintSummary(logger);
@@ -2197,8 +2197,7 @@ namespace uba
 			// Can't disable this, it can cause race conditions
 			//m_client.SetWorkTracker(nullptr);
 
-			StackBinaryWriter<SendMaxSize> writer;
-			WriteSummary(writer, [&](Logger& logger)
+			m_trace.WriteSessionSummary([&](Logger& logger)
 				{
 					PrintSummary(logger);
 					m_storage.PrintSummary(logger);
@@ -2206,7 +2205,6 @@ namespace uba
 					KernelStats::GetGlobal().Print(logger, true);
 					PrintContentionSummary(logger);
 				});
-			m_trace.SessionSummary(0, writer.GetData(), writer.GetPosition());
 
 			StringBuffer<> ubaFile(m_sessionLogDir);
 			ubaFile.Append(TCV("Trace.uba"));
diff --git a/Engine/Source/Programs/UnrealBuildAccelerator/Common/Private/UbaSessionServer.cpp b/Engine/Source/Programs/UnrealBuildAccelerator/Common/Private/UbaSessionServer.cpp
--- a/Engine/Source/Programs/UnrealBuildAccelerator/Common/Private/UbaSessionServer.cpp
+++ b/Engine/Source/Programs/UnrealBuildAccelerator/Common/Private/UbaSessionServer.cpp
@@ -274,8 +274,7 @@ namespace uba
 
 		if (m_trace.IsWriting())
 		{
-			StackBinaryWriter<SendMaxSize> writer;
-			WriteSummary(writer, [&](Logger& logger)
+			m_trace.WriteSessionSummary([&](Logger& logger)
 				{
 					PrintSummary(logger);
 					m_storage.PrintSummary(logger);
@@ -283,7 +282,6 @@ namespace uba
 					KernelStats::GetGlobal().Print(logger, true);
 					PrintContentionSummary(logger);
 				});
-			m_trace.SessionSummary(0, writer.GetData(), writer.GetPosition());
 		}
 
 		for (auto s : m_clientSessions)
diff --git a/Engine/Source/Programs/UnrealBuildAccelerator/Common/Private/UbaTrace.cpp b/Engine/Source/Programs/UnrealBuildAccelerator/Common/Private/UbaTrace.cpp
--- a/Engine/Source/Programs/UnrealBuildAccelerator/Common/Private/UbaTrace.cpp
+++ b/Engine/Source/Programs/UnrealBuildAccelerator/Common/Private/UbaTrace.cpp
@@ -3,6 +3,7 @@
 #include "UbaTrace.h"
 #include "UbaFileAccessor.h"
 #include "UbaProcessHandle.h"
+#include "UbaStats.h"
 
 #if PLATFORM_WINDOWS
 #include <tlhelp32.h>
@@ -170,6 +171,39 @@ namespace uba
 		WorkEnd(id);
 	}
 
+	void Trace::WriteSessionSummary(const Function<void(Logger& logger)>& summaryFunc)
+	{
+		StackBinaryWriter<256*1024> writer;
+		WriteSummaryText(writer, summaryFunc);
+		SessionSummary(0, writer.GetData(), writer.GetPosition());
+	}
+
+	void Trace::WriteSummaryText(BinaryWriter& writer, const Function<void(Logger& logger)>& summaryFunc)
+	{
+		struct SummaryLogWriter : public LogWriter
+		{
+			SummaryLogWriter(BinaryWriter& w) : writer(w) {}
+			virtual void BeginScope() override {}
+			virtual void EndScope() override {}
+			virtual void Log(LogEntryType type, const tchar* str, u32 strLen, const tchar* prefix, u32 prefixLen) override
+			{
+				writer.WriteString(str, strLen);
+				++count;
+			}
+
+			BinaryWriter& writer;
+			u32 count = 0;
+		};
+		u32* lineCount = (u32*)writer.AllocWrite(4);
+
+		SummaryLogWriter logWriter(writer);
+		LoggerWithWriter logger(logWriter, TC(""));
+
+		summaryFunc(logger);
+
+		*lineCount = logWriter.count;
+	}
+
 	void Trace::FreeMemory()
 	{
 		if (m_memoryBegin)
@@ -320,6 +354,16 @@ namespace uba
 		writer.WriteByte(255);
 	}
 
+	void Trace::ProcessExited(u32 processId, u32 exitCode)
+	{
+		StackBinaryWriter<1024> writer;
+		ProcessStats().Write(writer);
+		SessionStats().Write(writer);
+		StorageStats().Write(writer);
+		KernelStats().Write(writer);
+		ProcessExited(processId, exitCode, writer.GetData(), writer.GetPosition(), {});
+	}
+
 	void Trace::ProcessReturned(u32 processId, const StringView& reason)
 	{
 		BEGIN_TRACE_ENTRY(ProcessReturned);
diff --git a/Engine/Source/Programs/UnrealBuildAccelerator/Common/Private/UbaWorkManager.cpp b/Engine/Source/Programs/UnrealBuildAccelerator/Common/Private/UbaWorkManager.cpp
--- a/Engine/Source/Programs/UnrealBuildAccelerator/Common/Private/UbaWorkManager.cpp
+++ b/Engine/Source/Programs/UnrealBuildAccelerator/Common/Private/UbaWorkManager.cpp
@@ -26,6 +26,10 @@ namespace uba
 			return t->TrackWorkEnd(id);
 	}
 
+	void WorkManager::SetWorkTracker(WorkTracker* workTracker)
+	{
+		m_workTracker = workTracker;
+	}
 
 	struct WorkManagerImpl::Worker
 	{
diff --git a/Engine/Source/Programs/UnrealBuildAccelerator/Common/Public/UbaSession.h b/Engine/Source/Programs/UnrealBuildAccelerator/Common/Public/UbaSession.h
--- a/Engine/Source/Programs/UnrealBuildAccelerator/Common/Public/UbaSession.h
+++ b/Engine/Source/Programs/UnrealBuildAccelerator/Common/Public/UbaSession.h
@@ -203,7 +203,6 @@ namespace uba
 		void EnsureDirectoryTableMemory(u64 neededSize);
 		void GetSystemInfo(StringBufferBase& out);
 		bool GetMemoryInfo(u64& outAvailable, u64& outTotal);
-		void WriteSummary(BinaryWriter& writer, const Function<void(Logger& logger)>& summaryFunc);
 		float UpdateCpuLoad();
 
 		bool HasVfs(RootsHandle handle) const;
diff --git a/Engine/Source/Programs/UnrealBuildAccelerator/Common/Public/UbaTrace.h b/Engine/Source/Programs/UnrealBuildAccelerator/Common/Public/UbaTrace.h
--- a/Engine/Source/Programs/UnrealBuildAccelerator/Common/Public/UbaTrace.h
+++ b/Engine/Source/Programs/UnrealBuildAccelerator/Common/Public/UbaTrace.h
@@ -89,6 +89,7 @@ namespace uba
 		void ProcessAdded(u32 sessionId, u32 processId, const StringView& description, const StringView& breadcrumbs);
 		void ProcessEnvironmentUpdated(u32 processId, const StringView& reason, const u8* data, u64 dataSize, const StringView& breadcrumbs);
 		void ProcessExited(u32 processId, u32 exitCode, const u8* data, u64 dataSize, const Vector<ProcessLogLine>& logLines);
+		void ProcessExited(u32 processId, u32 exitCode);
 		void ProcessReturned(u32 processId, const StringView& reason);
 		void ProcessAddBreadcrumbs(u32 processId, const StringView& breadcrumbs, bool deleteOld);
 		void ProxyCreated(u32 clientId, const tchar* proxyName);
@@ -122,6 +123,9 @@ namespace uba
 		virtual void TrackWorkHint(u32 id, const StringView& hint, u64 startTime = 0) final override;
 		virtual void TrackWorkEnd(u32 id) final override;
 
+		void WriteSessionSummary(const Function<void(Logger& logger)>& summaryFunc);
+		static void WriteSummaryText(BinaryWriter& writer, const Function<void(Logger& logger)>& summaryFunc);
+
 	private:
 		struct WriterScope;
 		void FreeMemory();
diff --git a/Engine/Source/Programs/UnrealBuildAccelerator/Common/Public/UbaWorkManager.h b/Engine/Source/Programs/UnrealBuildAccelerator/Common/Public/UbaWorkManager.h
--- a/Engine/Source/Programs/UnrealBuildAccelerator/Common/Public/UbaWorkManager.h
+++ b/Engine/Source/Programs/UnrealBuildAccelerator/Common/Public/UbaWorkManager.h
@@ -44,7 +44,7 @@ namespace uba
 		void TrackWorkHint(u32 id, const StringView& hint, u64 startTime = 0);
 		void TrackWorkEnd(u32 id);
 
-		void SetWorkTracker(WorkTracker* workTracker) { m_workTracker = workTracker; }
+		void SetWorkTracker(WorkTracker* workTracker);
 		WorkTracker* GetWorkTracker() { return m_workTracker; }
 
 		template<u32 BatchSize = 1, typename TContainer, typename TFunc>
@@ -85,8 +85,8 @@ namespace uba
 	struct TrackWorkScope
 	{
 		TrackWorkScope() : tracker(nullptr), id(0) {}
-		TrackWorkScope(WorkTracker& t, const StringView& desc, const Color& color = ColorWork) : tracker(&t), id(t.TrackWorkStart(desc, color)) {}
-		void AddHint(const StringView& hint, u64 startTime = 0) { if (tracker) tracker->TrackWorkHint(id, hint, startTime); }
+		TrackWorkScope(WorkTracker& t, StringView desc, const Color& color = ColorWork) : tracker(&t), id(t.TrackWorkStart(desc, color)) {}
+		void AddHint(StringView hint, u64 startTime = 0) { if (tracker) tracker->TrackWorkHint(id, hint, startTime); }
 		~TrackWorkScope() { if (tracker) tracker->TrackWorkEnd(id); }
 		WorkTracker* tracker;
 		u32 id;
@@ -94,10 +94,10 @@ namespace uba
 
 	struct TrackHintScope
 	{
-		TrackHintScope(TrackWorkScope& t, const StringView& h) : tws(t), hint(h), startTime(GetTime()) {}
+		TrackHintScope(TrackWorkScope& t, StringView h) : tws(t), hint(h), startTime(GetTime()) {}
 		~TrackHintScope() { tws.AddHint(hint, startTime); }
 		TrackWorkScope& tws;
-		const StringView& hint;
+		StringView hint;
 		u64 startTime;
 	};
 
diff --git a/Engine/Source/Programs/UnrealBuildAccelerator/Visualizer/Private/UbaVisualizer.cpp b/Engine/Source/Programs/UnrealBuildAccelerator/Visualizer/Private/UbaVisualizer.cpp
--- a/Engine/Source/Programs/UnrealBuildAccelerator/Visualizer/Private/UbaVisualizer.cpp
+++ b/Engine/Source/Programs/UnrealBuildAccelerator/Visualizer/Private/UbaVisualizer.cpp
@@ -2308,9 +2308,12 @@ namespace uba
 				logger.SetColor(m_recvColor).Info(L"  Recv: %lsit/s", BytesToText(m_stats.recvBytesPerSecond*8).str);
 				logger.SetColor(m_sendColor).Info(L"  Send: %lsit/s", BytesToText(m_stats.sendBytesPerSecond*8).str);
 			}
+			logger.SetColor(m_textColor);
+
 			if (m_stats.ping)
 				logger.Info(L"  Ping: %ls", TimeToText(m_stats.ping, false, m_traceView.frequency).str);
-
+			if (m_stats.connectionCount)
+				logger.Info(L"  Conn: %u", u32(m_stats.connectionCount));
 
 			for (auto& pair : m_stats.drives)
 				logger.SetColor(m_driveColor).Info(L"  %c: %u%% R:%s/s W:%s/s", pair.first, pair.second.busyPercent, BytesToText(pair.second.readPerSecond).str, BytesToText(pair.second.writePerSecond).str);
@@ -3107,6 +3110,7 @@ namespace uba
 								outResult.stats.memAvail = session.memAvail[i];
 								outResult.stats.cpuLoad = session.cpuLoad[i];
 								outResult.stats.memTotal = session.memTotal;
+								outResult.stats.connectionCount = session.connectionCount[i];
 								outResult.statsSelected = true;
 
 								for (auto& pair : session.drives)
diff --git a/Engine/Source/Programs/UnrealBuildAccelerator/Visualizer/Private/UbaVisualizer.h b/Engine/Source/Programs/UnrealBuildAccelerator/Visualizer/Private/UbaVisualizer.h
--- a/Engine/Source/Programs/UnrealBuildAccelerator/Visualizer/Private/UbaVisualizer.h
+++ b/Engine/Source/Programs/UnrealBuildAccelerator/Visualizer/Private/UbaVisualizer.h
@@ -103,6 +103,7 @@ namespace uba
 			u64 recvBytes = 0;
 			u64 sendBytes = 0;
 			u64 procActive = 0;
+			u16 connectionCount = 0;
 			float cpuLoad = 0;
 			struct Drive
 			{
