[UBA]
* Fixed so when stdout and/or stderr are redirected to nul, they will not output in detoured paths either.
* Removed echoOn variable which was just totally wrong.

diff --git a/Engine/Source/Programs/UnrealBuildAccelerator/Common/Private/UbaProcess.cpp b/Engine/Source/Programs/UnrealBuildAccelerator/Common/Private/UbaProcess.cpp
--- a/Engine/Source/Programs/UnrealBuildAccelerator/Common/Private/UbaProcess.cpp
+++ b/Engine/Source/Programs/UnrealBuildAccelerator/Common/Private/UbaProcess.cpp
@@ -822,7 +822,6 @@ namespace uba
 		InitResponse response;
 		CallSession(MessageType_Init, m_session.GetInitResponse(response, msg));
 		writer.WriteU32(m_id);
-		writer.WriteBool(m_echoOn);
 		writer.WriteBool(m_parentProcess != nullptr);
 		writer.WriteString(m_startInfo.application);
 		writer.WriteString(m_startInfo.workingDir);
@@ -1173,7 +1172,6 @@ namespace uba
 		}
 
 		auto& process = *(ProcessImpl*)h.m_process;
-		process.m_echoOn = m_echoOn;
 
 		const char* detoursLib = m_session.m_detoursLibrary[IsArmBinary].c_str();
 		u32 detoursLibLen = u32(m_session.m_detoursLibrary[IsArmBinary].size());
@@ -1301,12 +1299,6 @@ namespace uba
 		return true;
 	}
 
-	bool ProcessImpl::HandleEchoOn(BinaryReader& reader, BinaryWriter& writer)
-	{
-		m_echoOn = reader.ReadBool();
-		return true;
-	}
-
 	bool ProcessImpl::HandleInputDependencies(BinaryReader& reader, BinaryWriter& writer)
 	{
 		UBA_ASSERT(m_startInfo.trackInputs);
diff --git a/Engine/Source/Programs/UnrealBuildAccelerator/Common/Public/UbaProcess.h b/Engine/Source/Programs/UnrealBuildAccelerator/Common/Public/UbaProcess.h
--- a/Engine/Source/Programs/UnrealBuildAccelerator/Common/Public/UbaProcess.h
+++ b/Engine/Source/Programs/UnrealBuildAccelerator/Common/Public/UbaProcess.h
@@ -153,7 +153,6 @@ namespace uba
 		Atomic<bool> m_hasExited;
 		MessageType m_firstMessageError = MessageType_None;
 		bool m_runningRemote = false;
-		bool m_echoOn = true;
 		bool m_gotExitMessage = false;
 		bool m_parentReportedExit = false;
 		bool m_detourEnabled = true;
diff --git a/Engine/Source/Programs/UnrealBuildAccelerator/Core/Public/UbaProtocol.h b/Engine/Source/Programs/UnrealBuildAccelerator/Core/Public/UbaProtocol.h
--- a/Engine/Source/Programs/UnrealBuildAccelerator/Core/Public/UbaProtocol.h
+++ b/Engine/Source/Programs/UnrealBuildAccelerator/Core/Public/UbaProtocol.h
@@ -26,7 +26,6 @@ namespace uba
 		UBA_PROCESS_MESSAGE(ExitChildProcess) \
 		UBA_PROCESS_MESSAGE(VirtualAllocFailed) \
 		UBA_PROCESS_MESSAGE(Log) \
-		UBA_PROCESS_MESSAGE(EchoOn) \
 		UBA_PROCESS_MESSAGE(InputDependencies) \
 		UBA_PROCESS_MESSAGE(Exit) \
 		UBA_PROCESS_MESSAGE(FlushWrittenFiles) \
@@ -48,7 +47,7 @@ namespace uba
 		#undef UBA_PROCESS_MESSAGE
 	};
 
-	inline constexpr u32 ProcessMessageVersion = 1346;
+	inline constexpr u32 ProcessMessageVersion = 1347;
 
 	inline constexpr u32 CommunicationMemSize = IsWindows ? 64*1024 : 64*1024*2; // Macos expands some commandlines to be crazy long
 
diff --git a/Engine/Source/Programs/UnrealBuildAccelerator/Detours/Private/UbaDetoursShared.cpp b/Engine/Source/Programs/UnrealBuildAccelerator/Detours/Private/UbaDetoursShared.cpp
--- a/Engine/Source/Programs/UnrealBuildAccelerator/Detours/Private/UbaDetoursShared.cpp
+++ b/Engine/Source/Programs/UnrealBuildAccelerator/Detours/Private/UbaDetoursShared.cpp
@@ -24,7 +24,7 @@ namespace uba
 	VARIABLE_MEM(MappedFileTable, g_mappedFileTable);
 	VARIABLE_MEM(ReaderWriterLock, g_consoleStringCs);
 
-	bool g_echoOn = true;
+	bool g_conEnabled[2];
 	u32 g_rulesIndex;
 	ApplicationRules* g_rules;
 	bool g_runningRemote;
@@ -35,6 +35,8 @@ namespace uba
 
 	void InitSharedVariables()
 	{
+		g_conEnabled[0] = true;
+		g_conEnabled[1] = true;
 		g_virtualApplicationMem.Create();
 		g_virtualApplicationDirMem.Create();
 		g_statsMem.Create();
@@ -412,7 +414,7 @@ namespace uba
 	template<typename CharType>
 	void Shared_WriteConsoleT(const CharType* chars, u32 charCount, bool isError)
 	{
-		if (!g_echoOn || g_suppressLogging)
+		if (!g_conEnabled[isError?0:1] || g_suppressLogging)
 			return;
 
 		SCOPED_WRITE_LOCK(g_consoleStringCs, lock);
diff --git a/Engine/Source/Programs/UnrealBuildAccelerator/Detours/Private/UbaDetoursShared.h b/Engine/Source/Programs/UnrealBuildAccelerator/Detours/Private/UbaDetoursShared.h
--- a/Engine/Source/Programs/UnrealBuildAccelerator/Detours/Private/UbaDetoursShared.h
+++ b/Engine/Source/Programs/UnrealBuildAccelerator/Detours/Private/UbaDetoursShared.h
@@ -63,7 +63,7 @@ namespace uba
 
 	extern ProcessStats& g_stats;
 	extern KernelStats& g_kernelStats;
-	extern bool g_echoOn;
+	extern bool g_conEnabled[2]; // 0=stderr, 1=stdout
 	extern ReaderWriterLock& g_communicationLock;
 	extern MemoryBlock& g_memoryBlock;
 	extern DirectoryTable& g_directoryTable;
diff --git a/Engine/Source/Programs/UnrealBuildAccelerator/Detours/Private/Windows/UbaDetoursFunctionsKernelBase.inl b/Engine/Source/Programs/UnrealBuildAccelerator/Detours/Private/Windows/UbaDetoursFunctionsKernelBase.inl
--- a/Engine/Source/Programs/UnrealBuildAccelerator/Detours/Private/Windows/UbaDetoursFunctionsKernelBase.inl
+++ b/Engine/Source/Programs/UnrealBuildAccelerator/Detours/Private/Windows/UbaDetoursFunctionsKernelBase.inl
@@ -619,7 +619,7 @@ ReaderWriterLock& g_stdFileLock = *new ReaderWriterLock();
 
 void WriteStdFile(LPCVOID lpBuffer, DWORD nNumberOfBytesToWrite, bool isError)
 {
-	if (!g_echoOn || g_suppressLogging)
+	if (!g_conEnabled[isError?0:1] || g_suppressLogging)
 		return;
 
 	SCOPED_WRITE_LOCK(g_stdFileLock, lock);
@@ -2929,6 +2929,8 @@ HANDLE Detoured_GetStdHandle(DWORD nStdHandle)
 	return res;
 }
 
+bool g_setStdHandleCalled;
+
 BOOL Detoured_SetStdHandle(DWORD nStdHandle, HANDLE hHandle)
 {
 	DETOURED_CALL(SetStdHandle);
@@ -2953,6 +2955,9 @@ BOOL Detoured_SetStdHandle(DWORD nStdHandle, HANDLE hHandle)
 		UBA_ASSERT(trueHandle != INVALID_HANDLE_VALUE);
 	}
 	DEBUG_LOG_TRUE(L"SetStdHandle", L"%u -> %llu", nStdHandle, u64(trueHandle));
+
+	g_setStdHandleCalled = true;
+
 	return True_SetStdHandle(nStdHandle, trueHandle);
 }
 
@@ -2983,45 +2988,44 @@ BOOL Detoured_GetConsoleMode(HANDLE hConsoleHandle, LPDWORD lpMode)
 	return res;
 }
 
-bool g_setConsoleModeCalled;
-
 BOOL Detoured_SetConsoleMode(HANDLE hConsoleHandle, DWORD mode)
 {
 	DETOURED_CALL(SetConsoleMode);
-	DEBUG_LOG_DETOURED(L"SetConsoleMode", L"%llu (0x%x)", u64(hConsoleHandle), mode);
 
 	HANDLE trueHandle = hConsoleHandle;
 
-	if (hConsoleHandle == g_stdHandle[0])
+	const tchar* dbgstr = TC("unknown");
+
+	if (g_isDetachedProcess)
 	{
-		if (g_isDetachedProcess)
+		if (hConsoleHandle == g_stdHandle[0])
+		{
 			trueHandle = True_GetStdHandle(STD_ERROR_HANDLE);
-	}
-	else if (hConsoleHandle == g_stdHandle[1])
-	{
-		g_echoOn = (mode & ~503) != 0; // TODO: This might be wrong. Trying to figure out how echo off in batch files work in terms of win32 calls
-		DEBUG_LOG(L"Echo %s", (g_echoOn ? L"on" : L"off"));
-
-		RPC_MESSAGE(EchoOn, log)
-		writer.WriteBool(g_echoOn);
-		writer.Flush();
-		if (g_isDetachedProcess)
+			dbgstr = TC("error");
+		}
+		else if (hConsoleHandle == g_stdHandle[1])
+		{
 			trueHandle = True_GetStdHandle(STD_OUTPUT_HANDLE);
-	}
-	else if (hConsoleHandle == g_stdHandle[0])
-	{
-		if (g_isDetachedProcess)
+			dbgstr = TC("output");
+		}
+		else if (hConsoleHandle == g_stdHandle[2])
+		{
 			trueHandle = True_GetStdHandle(STD_INPUT_HANDLE);
+			dbgstr = TC("input");
+		}
 	}
-	else if (isDetouredHandle(hConsoleHandle))
+
+	if (hConsoleHandle == trueHandle && isDetouredHandle(hConsoleHandle))
 	{
-		g_setConsoleModeCalled = true;
 		DetouredHandle& dh = asDetouredHandle(hConsoleHandle);
 		trueHandle = dh.trueHandle;
 		UBA_ASSERTF(trueHandle != INVALID_HANDLE_VALUE, TC("SetConsoleMode is using %s as handle to set mode 0x%x"), HandleToName(hConsoleHandle), mode);
+		dbgstr = TC("detoured");
 	}
 
-	return True_SetConsoleMode(trueHandle, mode);
+	BOOL res = True_SetConsoleMode(trueHandle, mode);
+	DEBUG_LOG_DETOURED(L"SetConsoleMode", L"(%s) %llu (0x%x) -> %s", dbgstr, u64(hConsoleHandle), mode, ToString(res));
+	return res;
 }
 
 BOOL Detoured_GetConsoleTitleW(LPTSTR lpConsoleTitle, DWORD nSize)
@@ -3038,7 +3042,7 @@ BOOL Detoured_CreateProcessW(LPCWSTR lpApplicationName, LPWSTR lpCommandLine, LP
 	StringView originalCmd = ToView(lpCommandLine ? lpCommandLine : TC(""));
 
 	DETOURED_CALL(CreateProcessW);
-	DEBUG_LOG_DETOURED(L"CreateProcessW", L"%ls %ls CreationFlags: 0x%x StartupFlags: 0x%u Stdin: %llu WorkDir: %s", lpApplicationName, originalCmd.data, dwCreationFlags, lpStartupInfo->dwFlags, u64(lpStartupInfo->hStdInput), (lpCurrentDirectory ? lpCurrentDirectory : L""));
+	DEBUG_LOG_DETOURED(L"CreateProcessW", L"%ls %ls CreationFlags: 0x%x StartupFlags: 0x%u Stdin: %llu Stdout: %llu Stderr: %llu WorkDir: %s", lpApplicationName, originalCmd.data, dwCreationFlags, lpStartupInfo->dwFlags, u64(lpStartupInfo->hStdInput), u64(lpStartupInfo->hStdOutput), u64(lpStartupInfo->hStdError), (lpCurrentDirectory ? lpCurrentDirectory : L""));
 
 	StringView testName = ToView(lpApplicationName ? lpApplicationName : TC(""));
 
@@ -3226,12 +3230,18 @@ BOOL Detoured_CreateProcessW(LPCWSTR lpApplicationName, LPWSTR lpCommandLine, LP
 		}
 	}
 
+	//if (dwCreationFlags & EXTENDED_STARTUPINFO_PRESENT) // 0x00080000
+	//{
+	//	auto& siex = *(STARTUPINFOEXW*)lpStartupInfo;
+	//	(void)siex;
+	//}
+
 	lpStartupInfo->dwFlags |= STARTF_USESHOWWINDOW;
 	lpStartupInfo->wShowWindow = SW_HIDE;
 
 	if (g_rules->AllowDetach())
 		dwCreationFlags |= DETACHED_PROCESS;
-	else if (!g_setConsoleModeCalled) // If console mode has been called there are probably some sort of stdin/out rerouting stuff.. so let's not set CREATE_NO_WINDOW
+	else if (!g_setStdHandleCalled) // If console mode has been called there are probably some sort of stdin/out rerouting stuff.. so let's not set CREATE_NO_WINDOW
 		dwCreationFlags |= CREATE_NO_WINDOW;
 
 	dwCreationFlags |= CREATE_SUSPENDED;
diff --git a/Engine/Source/Programs/UnrealBuildAccelerator/Detours/Private/Windows/UbaDetoursFunctionsWin.cpp b/Engine/Source/Programs/UnrealBuildAccelerator/Detours/Private/Windows/UbaDetoursFunctionsWin.cpp
--- a/Engine/Source/Programs/UnrealBuildAccelerator/Detours/Private/Windows/UbaDetoursFunctionsWin.cpp
+++ b/Engine/Source/Programs/UnrealBuildAccelerator/Detours/Private/Windows/UbaDetoursFunctionsWin.cpp
@@ -38,6 +38,10 @@ DETOURED_FUNCTIONS_MEMORY
 DETOURED_FUNCTIONS_MEMORY_NON_WINE
 #undef DETOURED_FUNCTION
 
+#if !defined(DETOURED_INCLUDE_DEBUG)
+#define True_GetEnvironmentVariableW GetEnvironmentVariableW
+#endif
+
 #include "UbaBinaryParser.h"
 #include "UbaDirectoryTable.h"
 #include "UbaProcessStats.h"
@@ -987,14 +991,41 @@ void Init(const DetoursPayload& payload, u64 startTime)
 		if (stdoutHandle == 0)
 		{
 			g_isDetachedProcess = true;
+			DEBUG_LOG(L"Detached: true (stdout == 0)");
 		}
 		else
 		{
 			HANDLE stderrHandle = True_GetStdHandle(STD_ERROR_HANDLE);
-			g_stdHandle[0] = GetFileType(stderrHandle) == FILE_TYPE_CHAR ? stderrHandle : 0;
-			g_stdHandle[1] = GetFileType(stdoutHandle) == FILE_TYPE_CHAR ? stdoutHandle : 0;
+			g_stdHandle[0] = True_GetFileType(stderrHandle) == FILE_TYPE_CHAR ? stderrHandle : 0;
+			g_stdHandle[1] = True_GetFileType(stdoutHandle) == FILE_TYPE_CHAR ? stdoutHandle : 0;
+
+			ULONG neededSize = 0;
+
+			u8 buffer0[256];
+			auto& objName0 = *reinterpret_cast<UNICODE_STRING*>(buffer0);
+			if (!g_stdHandle[0])
+				objName0.Buffer = const_cast<tchar*>(L"Unset");
+			else if (NT_SUCCESS(True_NtQueryObject(g_stdHandle[0], (OBJECT_INFORMATION_CLASS)1, buffer0, sizeof(buffer0), &neededSize)))
+				g_conEnabled[0] = !Equals(objName0.Buffer, L"\\Device\\Null");
+			else
+				objName0.Buffer = const_cast<tchar*>(L"???");
+
+			u8 buffer1[256];
+			auto& objName1 = *reinterpret_cast<UNICODE_STRING*>(buffer1);
+			if (!g_stdHandle[1])
+				objName1.Buffer = const_cast<tchar*>(L"Unset");
+			else if (NT_SUCCESS(True_NtQueryObject(g_stdHandle[1], (OBJECT_INFORMATION_CLASS)1, buffer1, sizeof(buffer1), &neededSize)))
+				g_conEnabled[1] = !Equals(objName1.Buffer, L"\\Device\\Null");
+			else
+				objName1.Buffer = const_cast<tchar*>(L"???");
+
+			DEBUG_LOG(L"Detached: false (stderr: %s, stdout: %s)", objName0.Buffer, objName1.Buffer);
 		}
 	}
+	else
+	{
+		DEBUG_LOG(L"Detached: true");
+	}
 
 	if (g_isDetachedProcess)
 	{
@@ -1006,11 +1037,11 @@ void Init(const DetoursPayload& payload, u64 startTime)
 	if (payload.trackInputs)
 		g_trackInputsMem = (u8*)g_memoryBlock.Allocate(TrackInputsMemCapacity, 1, L"TrackInputs");
 	
-	g_systemRoot.count = GetEnvironmentVariableW(L"SystemRoot", g_systemRoot.data, g_systemRoot.capacity);
+	g_systemRoot.count = True_GetEnvironmentVariableW(L"SystemRoot", g_systemRoot.data, g_systemRoot.capacity);
 	g_systemRoot.MakeLower();
 
 	wchar_t systemTemp[256];
-	GetEnvironmentVariableW(L"TEMP", systemTemp, 256);
+	True_GetEnvironmentVariableW(L"TEMP", systemTemp, 256);
 	FixPath(g_systemTemp, systemTemp);
 
 	StringBuffer<512> applicationBuffer;
@@ -1029,7 +1060,6 @@ void Init(const DetoursPayload& payload, u64 startTime)
 		BinaryReader reader;
 
 		g_processId = reader.ReadU32();
-		g_echoOn = reader.ReadBool();
 		g_isChild = reader.ReadBool();
 
 		reader.ReadString(applicationBuffer);
