[UBA]
* Added option Scheduler UseThreadToExitRemoteProcess that can be used to move exit logic for remote process off the job system and to its separate thread. This can be useful if exit is doing lots of work (writing to cache etc) and we want to return the job worker to the system

diff --git a/Engine/Source/Programs/UnrealBuildAccelerator/Common/Private/UbaScheduler.cpp b/Engine/Source/Programs/UnrealBuildAccelerator/Common/Private/UbaScheduler.cpp
--- a/Engine/Source/Programs/UnrealBuildAccelerator/Common/Private/UbaScheduler.cpp
+++ b/Engine/Source/Programs/UnrealBuildAccelerator/Common/Private/UbaScheduler.cpp
@@ -88,6 +88,7 @@ namespace uba
 			table->GetValueAsBool(forceNative, TC("ForceNative"));
 			table->GetValueAsU32(maxLocalProcessors, TC("MaxLocalProcessors"));
 			table->GetValueAsU32(maxParallelCacheQueries, TC("MaxParallelCacheQueries"));
+			table->GetValueAsBool(useThreadToExitRemoteProcess, TC("UseThreadToExitRemoteProcess"));
 		}
 	}
 
@@ -100,6 +101,7 @@ namespace uba
 	,	m_enableProcessReuse(info.enableProcessReuse)
 	,	m_forceRemote(info.forceRemote)
 	,	m_forceNative(info.forceNative)
+	,	m_useThreadToExitRemoteProcess(info.useThreadToExitRemoteProcess)
 	,	m_processConfigs(info.processConfigs)
 	,	m_writeToCache(info.writeToCache && info.cacheClientCount)
 	{
@@ -401,7 +403,21 @@ namespace uba
 			return;
 		}
 
-		ExitProcess(*info, *handle.m_process, handle.m_process->GetExitCode(), false);
+		if (!m_useThreadToExitRemoteProcess || !handle.IsRemote())
+		{
+			ExitProcess(*info, *handle.m_process, handle.m_process->GetExitCode(), false);
+			return;
+		}
+
+		// If process is remote we want to return as soon as possible to return worker thread
+		ig.Cancel();
+		m_finishedRemoteProcessThreads.emplace_back([this, handle = handle, info]()
+			{
+				ExitProcess(*info, *handle.m_process, handle.m_process->GetExitCode(), false);
+				delete info;
+				return 0;
+
+			}, TC("UbaScheRemFin"));
 	}
 
 	u32 Scheduler::PopProcess(bool isLocal, ProcessStatus& outPrevStatus)
diff --git a/Engine/Source/Programs/UnrealBuildAccelerator/Common/Public/UbaScheduler.h b/Engine/Source/Programs/UnrealBuildAccelerator/Common/Public/UbaScheduler.h
--- a/Engine/Source/Programs/UnrealBuildAccelerator/Common/Public/UbaScheduler.h
+++ b/Engine/Source/Programs/UnrealBuildAccelerator/Common/Public/UbaScheduler.h
@@ -32,6 +32,7 @@ namespace uba
 		bool forceRemote = false; // Force all processes that can run remotely to run remotely.
 		bool forceNative = false; // Force all processes to run native (not detoured)
 		bool writeToCache = false; // Set to true in combination with setting cacheClient to populate cache
+		bool useThreadToExitRemoteProcess = false; // Set this to true to move exit process call from job to thread. 
 		ConfigTable* processConfigs = nullptr;
 	};
 
@@ -141,6 +142,7 @@ namespace uba
 		bool m_enableProcessReuse;
 		bool m_forceRemote;
 		bool m_forceNative;
+		bool m_useThreadToExitRemoteProcess;
 		bool m_allowDisableRemoteExecution = false;
 		ConfigTable* m_processConfigs = nullptr;
 
@@ -161,6 +163,7 @@ namespace uba
 		bool m_writeToCache;
 
 		List<Thread> m_cacheDownloadThreads;
+		List<Thread> m_finishedRemoteProcessThreads;
 
 		Scheduler(const Scheduler&) = delete;
 		void operator=(const Scheduler&) = delete;
