[UBA]
* Optimized network file transfer of large files if sender has full file mapping. Will now only do one batched message call

diff --git a/Engine/Source/Programs/UnrealBuildAccelerator/Common/Private/UbaStorageUtils.cpp b/Engine/Source/Programs/UnrealBuildAccelerator/Common/Private/UbaStorageUtils.cpp
--- a/Engine/Source/Programs/UnrealBuildAccelerator/Common/Private/UbaStorageUtils.cpp
+++ b/Engine/Source/Programs/UnrealBuildAccelerator/Common/Private/UbaStorageUtils.cpp
@@ -163,18 +163,28 @@ namespace uba
 		u64 sendCountCapacity = capacity / messageMaxSize;
 		u64 sendCount = left/messageMaxSize;
 
-		if (sendCount > sendCountCapacity)
-			sendCount = sendCountCapacity;
-		else if (sendCount < sendCountCapacity && (left - sendCount * messageMaxSize) > 0)
-			++sendCount;
-		
-		UBA_ASSERT(sendCount);
+		if (left <= capacity)
+		{
+			if (sendCount * messageMaxSize < left)
+				++sendCount;
+		}
+		else
+		{
+			if (sendCount > sendCountCapacity)
+				sendCount = sendCountCapacity;
+			else if (sendCount < sendCountCapacity && (left - sendCount * messageMaxSize) > 0)
+				++sendCount;
 
-		u64 entriesMem[sizeof(Entry) * 8];
-		UBA_ASSERT(sizeof(Entry)*sendCount <= sizeof(entriesMem));
+			UBA_ASSERT(sendCount);
+		}
 
+		u64 entriesMem[sizeof(Entry) * 8];
 		Entry* entries = (Entry*)entriesMem;
 
+		if (sizeof(Entry)*sendCount > sizeof(entriesMem))
+			entries = (Entry*)malloc(sizeof(Entry)*sendCount);
+		auto entriesGuard = MakeGuard([&]() { if (entries != (Entry*)entriesMem) free(entries); });
+
 		bool success = true;
 		u32 error = 0;
 		u32 inFlightCount = u32(sendCount);
@@ -597,24 +607,41 @@ namespace uba
 
 			auto writePrev = [&]() { return WriteDestination(lastSource, lastResponseSize); };
 
-			u64 leftCompressed = fileSize - responseSize;
-			while (leftCompressed)
+			if (fileMappingMem)
+			{
+				if (!writePrev())
+					return false;
+				sg.Execute();
+				if (u64 leftCompressed = fileSize - responseSize)
+				{
+					u32 error;
+					if (!SendBatchMessages(logger, client, fetchId, fileMappingMem + destOffset, leftCompressed, leftCompressed, sizeOfFirstMessage, readIndex, responseSize, {}, destination, &error))
+						return logger.Error(TC("Failed to send batched messages to server while retrieving cas %s to %s. Error: %u"), CasKeyString(casKey).str, destination, error);
+					UBA_ASSERT(responseSize == leftCompressed);
+				}
+			}
+			else
 			{
-				if (fetchId == u16(~0))
-					return logger.Error(TC("Cas content error (2). Server believes %s was only one segment but client sees more. "));//UncompressedSize: %llu LeftUncompressed: %llu Size: %llu Left to read: %llu ResponseSize: %u. (%s)"), destination, actualSize, leftUncompressed, fileSize, left, responseSize, CasKeyString(casKey).str);
 
-				u32 error;
-				if (!SendBatchMessages(logger, client, fetchId, source, BufferSlotHalfSize, leftCompressed, sizeOfFirstMessage, readIndex, responseSize, writePrev, destination, &error))
-					return logger.Error(TC("Failed to send batched messages to server while retrieving cas %s to %s. Error: %u"), CasKeyString(casKey).str, destination, error);
+				u64 leftCompressed = fileSize - responseSize;
+				while (leftCompressed)
+				{
+					if (fetchId == u16(~0))
+						return logger.Error(TC("Cas content error (2). Server believes %s was only one segment but client sees more. "));//UncompressedSize: %llu LeftUncompressed: %llu Size: %llu Left to read: %llu ResponseSize: %u. (%s)"), destination, actualSize, leftUncompressed, fileSize, left, responseSize, CasKeyString(casKey).str);
+
+					u32 error;
+					if (!SendBatchMessages(logger, client, fetchId, source, BufferSlotHalfSize, leftCompressed, sizeOfFirstMessage, readIndex, responseSize, writePrev, destination, &error))
+						return logger.Error(TC("Failed to send batched messages to server while retrieving cas %s to %s. Error: %u"), CasKeyString(casKey).str, destination, error);
 
-				lastSource = source;
-				lastResponseSize = responseSize;
-				source = source == slot ? slot + BufferSlotHalfSize : slot;
+					lastSource = source;
+					lastResponseSize = responseSize;
+					source = source == slot ? slot + BufferSlotHalfSize : slot;
 
-				leftCompressed -= responseSize;
+					leftCompressed -= responseSize;
+				}
+				if (!writePrev())
+					return false;
 			}
-			if (!writePrev())
-				return false;
 		}
 		else if (actualSize)
 		{
diff --git a/Engine/Source/Programs/UnrealBuildAccelerator/Test/Private/UbaTest.cpp b/Engine/Source/Programs/UnrealBuildAccelerator/Test/Private/UbaTest.cpp
--- a/Engine/Source/Programs/UnrealBuildAccelerator/Test/Private/UbaTest.cpp
+++ b/Engine/Source/Programs/UnrealBuildAccelerator/Test/Private/UbaTest.cpp
@@ -22,10 +22,11 @@ namespace uba
 			DirectoryCache().CreateDirectory(logger, testFileDir.data);
 		}
 
+		u64 bytes = content.count*sizeof(tchar);
 		FileAccessor file(logger, outFile.data);
-		if (!file.CreateMemoryWrite(false, attributes, content.count))
+		if (!file.CreateMemoryWrite(false, attributes, bytes))
 			return logger.Error(TC("Failed to create file for write"));
-		memcpy(file.GetData(), content.data, content.count);
+		memcpy(file.GetData(), content.data, bytes);
 		return file.Close();
 	}
 
diff --git a/Engine/Source/Programs/UnrealBuildAccelerator/Test/Private/UbaTestCache.cpp b/Engine/Source/Programs/UnrealBuildAccelerator/Test/Private/UbaTestCache.cpp
--- a/Engine/Source/Programs/UnrealBuildAccelerator/Test/Private/UbaTestCache.cpp
+++ b/Engine/Source/Programs/UnrealBuildAccelerator/Test/Private/UbaTestCache.cpp
@@ -551,8 +551,17 @@ namespace uba
 			StringBuffer<> testApp;
 			GetTestAppPath(logger, testApp);
 
-			CHECK_TRUE(CreateTestFile(inputFile, logger, workingDir, TCV("Input.txt"), TCV("Foo")));
-			CHECK_TRUE(CreateTestFile(outputFile, logger, workingDir, TCV("Output.exe"), TCV("Foo"), DefaultAttributes(true)));
+			StringView content(TCV("Foo"));
+
+			CHECK_TRUE(CreateTestFile(inputFile, logger, workingDir, TCV("Input.txt"), content));
+
+			Vector<tchar> data;
+			data.resize(sizeof(Guid)*100000);
+			for (u32 i=0;i!=100000; ++i)
+				CreateGuid(((Guid*)data.data())[i]);
+			content = StringView(data.data(), u32(data.size()));
+
+			CHECK_TRUE(CreateTestFile(outputFile, logger, workingDir, TCV("Output.exe"), content, DefaultAttributes(true)));
 
 			StackBinaryWriter<256> inputs;
 			inputs.WriteString(inputFile);
