[UBA]
* Added so CacheClient_Connect can be called multiple times and only first call will initialize other things
* Added so it is possible to see total connection count visualizer for host session

diff --git a/Engine/Source/Programs/UnrealBuildAccelerator/Common/Private/UbaSessionServer.cpp b/Engine/Source/Programs/UnrealBuildAccelerator/Common/Private/UbaSessionServer.cpp
--- a/Engine/Source/Programs/UnrealBuildAccelerator/Common/Private/UbaSessionServer.cpp
+++ b/Engine/Source/Programs/UnrealBuildAccelerator/Common/Private/UbaSessionServer.cpp
@@ -2358,6 +2358,7 @@ namespace uba
 
 		u64 serverSend = m_server.GetTotalSentBytes();
 		u64 serverRecv = m_server.GetTotalRecvBytes();
+		u32 serverConnectionCount = m_server.GetConnectionCount();
 
 		SCOPED_CRITICAL_SECTION(m_remoteProcessAndSessionLock, lock);
 		for (auto sptr : m_clientSessions)
@@ -2373,9 +2374,11 @@ namespace uba
 		{
 			u64 send;
 			u64 recv;
-			rec.provider(send, recv);
+			u32 connectionCount;
+			rec.provider(send, recv, connectionCount);
 			serverSend += send;
 			serverRecv += recv;
+			serverConnectionCount += connectionCount;
 		}
 		lock.Leave();
 
@@ -2403,7 +2406,7 @@ namespace uba
 			}
 		}
 
-		m_trace.SessionUpdate(0, 0, serverSend, serverRecv, 0, memAvail, memTotal, cpuLoad);
+		m_trace.SessionUpdate(0, serverConnectionCount, serverSend, serverRecv, 0, memAvail, memTotal, cpuLoad);
 	}
 
 	void SessionServer::WriteRemoteEnvironmentVariables(BinaryWriter& writer)
diff --git a/Engine/Source/Programs/UnrealBuildAccelerator/Common/Public/UbaSessionServer.h b/Engine/Source/Programs/UnrealBuildAccelerator/Common/Public/UbaSessionServer.h
--- a/Engine/Source/Programs/UnrealBuildAccelerator/Common/Public/UbaSessionServer.h
+++ b/Engine/Source/Programs/UnrealBuildAccelerator/Common/Public/UbaSessionServer.h
@@ -73,7 +73,7 @@ namespace uba
 		NetworkServer& GetServer();
 
 		// Additional network traffic provider. Can be used to add more network stats to trace
-		using NetworkTrafficProvider = Function<void(u64& outSent, u64& outReceive)>;
+		using NetworkTrafficProvider = Function<void(u64& outSent, u64& outReceive, u32& outConnectionCount)>;
 		void RegisterNetworkTrafficProvider(u64 key, const NetworkTrafficProvider& provider);
 		void UnregisterNetworkTrafficProvider(u64 key);
 
diff --git a/Engine/Source/Programs/UnrealBuildAccelerator/Host/Private/UbaExports.cpp b/Engine/Source/Programs/UnrealBuildAccelerator/Host/Private/UbaExports.cpp
--- a/Engine/Source/Programs/UnrealBuildAccelerator/Host/Private/UbaExports.cpp
+++ b/Engine/Source/Programs/UnrealBuildAccelerator/Host/Private/UbaExports.cpp
@@ -836,6 +836,9 @@ extern "C"
 		auto& networkClient = (NetworkClientWithBackend&)cacheClient->GetClient();
 		CacheClientActiveScope ccas(cacheClient);
 
+		if (networkClient.IsConnected())
+			return networkClient.Connect(*networkClient.backend, host, u16(port));
+
 		u32 desired = networkClient.GetDesiredConnectionCount();
 		if (!desired)
 			return false;
@@ -844,10 +847,11 @@ extern "C"
 		while (--desired)
 			networkClient.Connect(*networkClient.backend, host, u16(port));
 
-		((SessionServer&)cacheClient->GetSession()).RegisterNetworkTrafficProvider(u64(cacheClient), [nc = &networkClient](u64& outSent, u64& outReceive)
+		((SessionServer&)cacheClient->GetSession()).RegisterNetworkTrafficProvider(u64(cacheClient), [nc = &networkClient](u64& outSent, u64& outReceive, u32& outConnectionCount)
 			{
 				outSent = nc->GetTotalSentBytes();
 				outReceive = nc->GetTotalRecvBytes();
+				outConnectionCount = nc->GetConnectionCount();
 			});
 
 		cacheClient->GetStorage().LoadCasTable();
