[UBA]
* Added CacheClient_Disconnect that can be called pressing ctrl-c
* Removed desired connection count variable in network client.. better to handle on the outside

diff --git a/Engine/Source/Programs/UnrealBuildAccelerator/Common/Private/UbaNetworkClient.cpp b/Engine/Source/Programs/UnrealBuildAccelerator/Common/Private/UbaNetworkClient.cpp
--- a/Engine/Source/Programs/UnrealBuildAccelerator/Common/Private/UbaNetworkClient.cpp
+++ b/Engine/Source/Programs/UnrealBuildAccelerator/Common/Private/UbaNetworkClient.cpp
@@ -17,7 +17,7 @@ namespace uba
 		if (!tablePtr)
 			return;
 		const ConfigTable& table = *tablePtr;
-		table.GetValueAsU32(desiredConnectionCount, TC("DesiredConnectionCount"));
+		(void)table;
 	}
 
 	struct NetworkClient::Connection
@@ -52,7 +52,6 @@ namespace uba
 			m_logger.Detail(TC("Adjusted msg size to %u to stay inside limits"), fixedSendSize);
 		m_sendSize = fixedSendSize;
 		m_receiveTimeoutSeconds = info.receiveTimeoutSeconds;
-		m_desiredConnectionCount = info.desiredConnectionCount;
 
 		m_connectionsIt = m_connections.end();
 
diff --git a/Engine/Source/Programs/UnrealBuildAccelerator/Common/Public/UbaNetworkClient.h b/Engine/Source/Programs/UnrealBuildAccelerator/Common/Public/UbaNetworkClient.h
--- a/Engine/Source/Programs/UnrealBuildAccelerator/Common/Public/UbaNetworkClient.h
+++ b/Engine/Source/Programs/UnrealBuildAccelerator/Common/Public/UbaNetworkClient.h
@@ -26,7 +26,6 @@ namespace uba
 		u32 sendSize = SendDefaultSize;
 		u32 receiveTimeoutSeconds = DefaultNetworkReceiveTimeoutSeconds;
 		u32 workerCount = ~0u; // ~0u means logical processor count will be used
-		u32 desiredConnectionCount = 1; // Desired connection count to server when connected
 		const u8* cryptoKey128 = nullptr;
 		void Apply(Config& config, const tchar* tableName = TC("NetworkClient"));
 	};
@@ -70,7 +69,6 @@ namespace uba
 		u32 GetConnectionCount() { return m_connectionCount; }
 		u64 GetTotalSentBytes() { return m_sendBytes; }
 		u64 GetTotalRecvBytes() { return m_recvBytes; }
-		u32 GetDesiredConnectionCount() { return m_desiredConnectionCount; } 
 
 		NetworkBackend* GetFirstConnectionBackend();
 
@@ -92,7 +90,6 @@ namespace uba
 		LoggerWithWriter m_logger;
 		u32 m_sendSize;
 		u32 m_receiveTimeoutSeconds;
-		u32 m_desiredConnectionCount;
 		Atomic<u64> m_sendBytes;
 		Atomic<u64> m_recvBytes;
 		Atomic<u32> m_recvCount;
diff --git a/Engine/Source/Programs/UnrealBuildAccelerator/Host/Private/UbaExports.cpp b/Engine/Source/Programs/UnrealBuildAccelerator/Host/Private/UbaExports.cpp
--- a/Engine/Source/Programs/UnrealBuildAccelerator/Host/Private/UbaExports.cpp
+++ b/Engine/Source/Programs/UnrealBuildAccelerator/Host/Private/UbaExports.cpp
@@ -836,16 +836,11 @@ extern "C"
 		auto& networkClient = (NetworkClientWithBackend&)cacheClient->GetClient();
 		CacheClientActiveScope ccas(cacheClient);
 
-		if (networkClient.IsConnected())
-			return networkClient.Connect(*networkClient.backend, host, u16(port));
-
-		u32 desired = networkClient.GetDesiredConnectionCount();
-		if (!desired)
-			return false;
+		bool wasConnected = networkClient.IsConnected();
 		if (!networkClient.Connect(*networkClient.backend, host, u16(port)))
 			return false;
-		while (--desired)
-			networkClient.Connect(*networkClient.backend, host, u16(port));
+		if (wasConnected)
+			return true;
 
 		((SessionServer&)cacheClient->GetSession()).RegisterNetworkTrafficProvider(u64(cacheClient), [nc = &networkClient](u64& outSent, u64& outReceive, u32& outConnectionCount)
 			{
@@ -934,6 +929,13 @@ extern "C"
 		cacheClient->RequestServerShutdown(reason);
 	}
 
+	void CacheClient_Disconnect(uba::CacheClient* cacheClient)
+	{
+		using namespace uba;
+		auto& networkClient = (NetworkClientWithBackend&)cacheClient->GetClient();
+		networkClient.Disconnect();
+	}
+
 	void CacheClient_Destroy(uba::CacheClient* cacheClient)
 	{
 		using namespace uba;
diff --git a/Engine/Source/Programs/UnrealBuildAccelerator/Host/Public/UbaExports.h b/Engine/Source/Programs/UnrealBuildAccelerator/Host/Public/UbaExports.h
--- a/Engine/Source/Programs/UnrealBuildAccelerator/Host/Public/UbaExports.h
+++ b/Engine/Source/Programs/UnrealBuildAccelerator/Host/Public/UbaExports.h
@@ -170,6 +170,7 @@ extern "C"
 	UBA_API bool CacheClient_WriteToCache2(uba::CacheClient* cacheClient, uba::u32 bucket, const uba::ProcessHandle* process, const uba::u8* inputs, uba::u32 inputsSize, const uba::u8* outputs, uba::u32 outputsSize);
 	UBA_API uba::CacheResult* CacheClient_FetchFromCache3(uba::CacheClient* cacheClient, uba::u64 rootsHandle, uba::u32 bucket, const uba::ProcessStartInfo& info);
 	UBA_API void CacheClient_RequestServerShutdown(uba::CacheClient* cacheClient, const uba::tchar* reason);
+	UBA_API void CacheClient_Disconnect(uba::CacheClient* cacheClient);
 	UBA_API void CacheClient_Destroy(uba::CacheClient* cacheClient);
 	UBA_API const uba::tchar* CacheResult_GetLogLine(uba::CacheResult* result, uba::u32 index);
 	UBA_API uba::u32 CacheResult_GetLogLineType(uba::CacheResult* result, uba::u32 index);
