[UBA]
* Fixed bug where protocol changes (removing echo flag) was not removed on posix
* Improved assert description

diff --git a/Engine/Source/Programs/UnrealBuildAccelerator/Detours/Private/PosixOS/UbaDetoursFunctionsPosix.cpp b/Engine/Source/Programs/UnrealBuildAccelerator/Detours/Private/PosixOS/UbaDetoursFunctionsPosix.cpp
--- a/Engine/Source/Programs/UnrealBuildAccelerator/Detours/Private/PosixOS/UbaDetoursFunctionsPosix.cpp
+++ b/Engine/Source/Programs/UnrealBuildAccelerator/Detours/Private/PosixOS/UbaDetoursFunctionsPosix.cpp
@@ -2757,7 +2757,6 @@ namespace uba
 			BinaryReader reader;
 
 			g_processId = reader.ReadU32();
-			bool echoOn = reader.ReadBool();
 			g_isChild = reader.ReadBool();
 
 			reader.ReadString(g_virtualApplication);
@@ -2797,7 +2796,7 @@ namespace uba
 		StringBuffer<128> mappedFileTableUidName;
 		GetMappingHandleName(mappedFileTableUidName, mappedFileTableHandle);
 		int mappedFileTableFd = shm_open(mappedFileTableUidName.data, O_RDONLY, S_IRUSR | S_IWUSR);
-		UBA_ASSERTF(mappedFileTableFd != -1, "mappedFileTableFd == %i", mappedFileTableFd);
+		UBA_ASSERTF(mappedFileTableFd != -1, "Failed to open mapped file table memory mapping %s (%s)", mappedFileTableFd, mappedFileTableUidName.data, strerror(errno));
 		u8* mappedFileTableMem = (u8*)mmap(NULL, FileMappingTableMemSize, PROT_READ, MAP_SHARED, mappedFileTableFd, 0);
 		UBA_ASSERTF(mappedFileTableMem != MAP_FAILED, "mmap failed (%s)", strerror(errno));
 		g_mappedFileTable.Init(mappedFileTableMem, mappedFileTableCount, mappedFileTableSize);
