[UBA]
* Changed so socket recv buffer size defaults to 4mb on windows and send buffer size defaults to 2mb. (disabled on wine because it caused perf degradation.. needs revisit)
* Cleaned up some logging stuff in NetworkBackendTcp

diff --git a/Engine/Source/Programs/UnrealBuildAccelerator/Common/Private/UbaNetworkBackendTcp.cpp b/Engine/Source/Programs/UnrealBuildAccelerator/Common/Private/UbaNetworkBackendTcp.cpp
--- a/Engine/Source/Programs/UnrealBuildAccelerator/Common/Private/UbaNetworkBackendTcp.cpp
+++ b/Engine/Source/Programs/UnrealBuildAccelerator/Common/Private/UbaNetworkBackendTcp.cpp
@@ -62,7 +62,7 @@ namespace uba
 	bool SetRecvBuf(Logger& logger, SOCKET socket, u32 windowSize);
 	bool SetSendBuf(Logger& logger, SOCKET socket, u32 windowSize);
 	bool DisableNagle(Logger& logger, SOCKET socket);
-	bool LogTcpInfo(Logger& logger, SOCKET socket);
+	void LogTcpInfo(Logger& logger, SOCKET socket);
 
 	////////////////////////////////////////////////////////////////////////////////////////////////////
 
@@ -154,6 +154,9 @@ namespace uba
 		table.GetValueAsBool(useOverlappedSend, TC("UseOverlappedSend"));
 		table.GetValueAsU32(iocpWorkerCount, TC("IocpWorkerCount"));
 		#endif
+
+		table.GetValueAsU32(recvBufferSize, TC("RecvBufferSize"));
+		table.GetValueAsU32(sendBufferSize, TC("SendBufferSize"));
 	}
 
 	////////////////////////////////////////////////////////////////////////////////////////////////////
@@ -202,6 +205,9 @@ namespace uba
 		m_iocpWorkerCount = u16(info.iocpWorkerCount);
 		#endif
 
+		m_recvBufferSize = info.recvBufferSize;
+		m_sendBufferSize = info.sendBufferSize;
+
 		if (info.statusUpdateSeconds)
 		{
 			m_tcpStatusLoop.Create(true);
@@ -524,9 +530,12 @@ namespace uba
 
 			SetLinger(logger, clientSocket, 10);
 
-			// Hard coding this values disable auto-tuning. Seems like creating multiple sockets is more efficient than playing around with these
-			//SetRecvBuf(logger, clientSocket, 50*1024*1024);
-			//SetSendBuf(logger, clientSocket, 50*1024*1024);
+			// For some reason performance gets destroyed on wine when changing these
+			if (IsWindows && !IsRunningWine())
+			{
+				SetRecvBuf(logger, clientSocket, m_recvBufferSize);
+				SetSendBuf(logger, clientSocket, m_sendBufferSize);
+			}
 
 			SCOPED_FUTEX(m_connectionsLock, lock);
 			auto it = m_connections.emplace(m_connections.end(), logger, clientSocket);
@@ -785,9 +794,12 @@ namespace uba
 
 		SetLinger(logger, socketFd, 10);
 
-		// Hard coding this values disable auto-tuning. Seems like creating multiple sockets is more efficient than playing around with these
-		//SetRecvBuf(logger, socketFd, 50*1024*1024);
-		//SetSendBuf(logger, socketFd, 50*1024*1024);
+		// For some reason performance gets destroyed on wine when changing these
+		if (IsWindows && !IsRunningWine())
+		{
+			SetRecvBuf(logger, socketFd, m_recvBufferSize);
+			SetSendBuf(logger, socketFd, m_sendBufferSize);
+		}
 
 		// Socket is good, cancel the socket close scope and break out of the loop.
 		socketClose.Cancel();
@@ -1322,7 +1334,6 @@ namespace uba
 			PrintTcpStatistics(m_logger, statusUpdateSeconds, &prevStats);
 			#endif
 
-			#if 0
 			SCOPED_FUTEX(m_connectionsLock, lock);
 			for (auto& conn : m_connections)
 			{
@@ -1335,7 +1346,6 @@ namespace uba
 					continue;
 				LogTcpInfo(m_logger, s);
 			}
-			#endif
 		}
 	}
 
@@ -1857,9 +1867,20 @@ namespace uba
 #endif
 
 
-	bool LogTcpInfo(Logger& logger, SOCKET socket)
+	void LogTcpInfo(Logger& logger, SOCKET socket)
 	{
 	#if PLATFORM_WINDOWS
+
+		#if 0
+		socklen_t optlen = sizeof(int);
+		int sendBufSize = 0;
+		getsockopt(socket, SOL_SOCKET, SO_SNDBUF, (char *)&sendBufSize, &optlen);
+		int recvBufSize = 0;
+		getsockopt(socket, SOL_SOCKET, SO_RCVBUF, (char *)&recvBufSize, &optlen);
+		logger.Info(TC("%i - Send: %s, Recv: %s"), (int)socket, BytesToText(sendBufSize).str, BytesToText(recvBufSize).str);
+		#endif
+
+		#if 0
 		if (IsRunningWine())
 		{
 			static HMODULE wineDll = LoadLibrary(TC("UbaWine.dll.so"));
@@ -1879,16 +1900,15 @@ namespace uba
 			logger.Info(TC("%i - RTT=%uus  InFlight=%u  Retrans=%u (Tot %u)  Dupacks=%u  State=%u RcvWnd=%u  RcvScale=%u"), (int)socket, ti.tcpi_rtt, bytes_inflight_est, ti.tcpi_retrans, ti.tcpi_total_retrans, ti.tcpi_snd_cwnd, ti.tcpi_state, ti.tcpi_rcv_space, ti.tcpi_rcv_wscale);
 			return true;
 		}
-
+		#endif
+		
 		//TCP_INFO_v0 ti;
 		//ULONG version = 0; // Specify 0 to retrieve the v0 version of this structure.
 		//ULONG bytesReturned;
 		//if (WSAIoctl(socket, SIO_TCP_INFO, &version, sizeof(ULONG), &ti, sizeof(TCP_INFO_v0), &bytesReturned, NULL, NULL) == SOCKET_ERROR)
 		//	return logger.Error(TC("WSAIoctl SIO_TCP_INFO failed (%s)"), LastErrorToText(WSAGetLastError()).data);
 		//logger.Info(TC("RTT=%uµs  InFlight=%u  Retrans=%u  Dupacks=%u  State=%u RcvWnd=%u  RcvBuf=%u"), ti.RttUs, ti.BytesInFlight, ti.BytesRetrans, ti.DupAcksIn, ti.State, ti.RcvWnd, ti.RcvBuf);
-		return true;
 	#else
-		return false;
 		/*
 		struct tcp_info ti = {};
 		socklen_t len = sizeof(ti);
diff --git a/Engine/Source/Programs/UnrealBuildAccelerator/Common/Public/UbaNetworkBackendTcp.h b/Engine/Source/Programs/UnrealBuildAccelerator/Common/Public/UbaNetworkBackendTcp.h
--- a/Engine/Source/Programs/UnrealBuildAccelerator/Common/Public/UbaNetworkBackendTcp.h
+++ b/Engine/Source/Programs/UnrealBuildAccelerator/Common/Public/UbaNetworkBackendTcp.h
@@ -21,6 +21,9 @@ namespace uba
 		bool useOverlappedSend = true;
 		u32 iocpWorkerCount = 0; // Iocp is turned off by default
 
+		u32 recvBufferSize = 4*1024*1024;
+		u32 sendBufferSize = 2*1024*1024;
+
 		void Apply(const Config& config, const tchar* tableName = TC("NetworkBackendTcp"));
 	};
 
@@ -83,6 +86,9 @@ namespace uba
 		u16 m_iocpWorkerCount;
 		#endif
 
+		u32 m_recvBufferSize;
+		u32 m_sendBufferSize;
+
 		List<Thread> m_iocpThreads;
 		void* m_iocpHandle = nullptr;
 
