[UBA]
* Added detour for CreateToolhelp32Snapshot and turn it into always fail. We don't want the actions to query information about the system. This call is used to figure out if vctip.exe is running. For future, maybe change this to an option

diff --git a/Engine/Source/Programs/UnrealBuildAccelerator/Core/Private/UbaApplicationRules.cpp b/Engine/Source/Programs/UnrealBuildAccelerator/Core/Private/UbaApplicationRules.cpp
--- a/Engine/Source/Programs/UnrealBuildAccelerator/Core/Private/UbaApplicationRules.cpp
+++ b/Engine/Source/Programs/UnrealBuildAccelerator/Core/Private/UbaApplicationRules.cpp
@@ -131,6 +131,11 @@ namespace uba
 		{
 			return true; // .pch
 		}
+
+		virtual bool CanExist(StringView file) const override
+		{
+			return !file.EndsWith(TCV("vctip.exe")); // This is hacky but we don't want to start vctip.exe
+		}
 	};
 
 	class ApplicationRulesVcLink : public ApplicationRulesVC
@@ -157,9 +162,9 @@ namespace uba
 			return fileName.Contains(TCV(".sup.")); // .sup.lib/exp are throw-away files that we don't want created
 		}
 
-		virtual bool CanExist(const tchar* file) const override
+		virtual bool CanExist(StringView file) const override
 		{
-			return Contains(file, TC("vctip.exe")) == false; // This is hacky but we don't want to start vctip.exe
+			return !file.EndsWith(TCV("vctip.exe")); // This is hacky but we don't want to start vctip.exe
 		}
 
 		virtual u64 FileTypeMaxSize(const StringView& file, bool isSystemOrTempFile) const override
diff --git a/Engine/Source/Programs/UnrealBuildAccelerator/Core/Public/UbaApplicationRules.h b/Engine/Source/Programs/UnrealBuildAccelerator/Core/Public/UbaApplicationRules.h
--- a/Engine/Source/Programs/UnrealBuildAccelerator/Core/Public/UbaApplicationRules.h
+++ b/Engine/Source/Programs/UnrealBuildAccelerator/Core/Public/UbaApplicationRules.h
@@ -82,7 +82,7 @@ namespace uba
 		}
 
 		// If returns false this means that all GetFileAttribute etc will return file-not-found on this file
-		virtual bool CanExist(const tchar* file) const
+		virtual bool CanExist(StringView file) const
 		{
 			return true;
 		}
diff --git a/Engine/Source/Programs/UnrealBuildAccelerator/Detours/Private/Windows/UbaDetoursFunctionsKernelBase.inl b/Engine/Source/Programs/UnrealBuildAccelerator/Detours/Private/Windows/UbaDetoursFunctionsKernelBase.inl
--- a/Engine/Source/Programs/UnrealBuildAccelerator/Detours/Private/Windows/UbaDetoursFunctionsKernelBase.inl
+++ b/Engine/Source/Programs/UnrealBuildAccelerator/Detours/Private/Windows/UbaDetoursFunctionsKernelBase.inl
@@ -645,6 +645,15 @@ void WriteStdFile(LPCVOID lpBuffer, DWORD nNumberOfBytesToWrite, bool isError)
 		g_stdFile.Appendf(L"%.*hs", left, bufferStr + start);
 }
 
+HANDLE Detoured_CreateToolhelp32Snapshot(DWORD dwFlags, DWORD th32ProcessID)
+{
+	DETOURED_CALL(CreateToolhelp32Snapshot);
+	DEBUG_LOG_DETOURED(L"CreateToolhelp32Snapshot", L"(Disabled) %u %u -> INVALID_HANDLE_VALUE", dwFlags, th32ProcessID);
+	SetLastError(ERROR_ACCESS_DENIED);
+	return INVALID_HANDLE_VALUE;
+	//return True_CreateToolhelp32Snapshot(dwFlags, th32ProcessID);
+}
+
 BOOL Detoured_WriteFile(HANDLE hFile, LPCVOID lpBuffer, DWORD nNumberOfBytesToWrite, LPDWORD lpNumberOfBytesWritten, LPOVERLAPPED lpOverlapped)
 {
 	DWORD temp = 0;
@@ -1077,8 +1086,9 @@ BOOL Detoured_GetFileAttributesExW(LPCWSTR lpFileName, GET_FILEEX_INFO_LEVELS fI
 	FixPath(fixedName, lpFileName);
 	DevirtualizePath(fixedName);
 
-	if (!g_rules->CanExist(fixedName.data))
+	if (!g_rules->CanExist(fixedName))
 	{
+		DEBUG_LOG_TRUE(L"GetFileAttributesExW", L"(CanNotExist) (%ls) -> NotFound", lpFileName);
 		SetLastError(ERROR_FILE_NOT_FOUND);
 		return FALSE;
 	}
diff --git a/Engine/Source/Programs/UnrealBuildAccelerator/Detours/Private/Windows/UbaDetoursFunctionsWin.cpp b/Engine/Source/Programs/UnrealBuildAccelerator/Detours/Private/Windows/UbaDetoursFunctionsWin.cpp
--- a/Engine/Source/Programs/UnrealBuildAccelerator/Detours/Private/Windows/UbaDetoursFunctionsWin.cpp
+++ b/Engine/Source/Programs/UnrealBuildAccelerator/Detours/Private/Windows/UbaDetoursFunctionsWin.cpp
@@ -1039,6 +1039,7 @@ void Init(const DetoursPayload& payload, u64 startTime)
 	
 	g_systemRoot.count = True_GetEnvironmentVariableW(L"SystemRoot", g_systemRoot.data, g_systemRoot.capacity);
 	g_systemRoot.MakeLower();
+	UBA_ASSERT(g_systemRoot.count);
 
 	wchar_t systemTemp[256];
 	True_GetEnvironmentVariableW(L"TEMP", systemTemp, 256);
diff --git a/Engine/Source/Programs/UnrealBuildAccelerator/Detours/Private/Windows/UbaDetoursFunctionsWin.h b/Engine/Source/Programs/UnrealBuildAccelerator/Detours/Private/Windows/UbaDetoursFunctionsWin.h
--- a/Engine/Source/Programs/UnrealBuildAccelerator/Detours/Private/Windows/UbaDetoursFunctionsWin.h
+++ b/Engine/Source/Programs/UnrealBuildAccelerator/Detours/Private/Windows/UbaDetoursFunctionsWin.h
@@ -131,6 +131,7 @@
 	DETOURED_FUNCTION(GetExitCodeProcess) \
 	DETOURED_FUNCTION(CreateTimerQueueTimer) \
 	DETOURED_FUNCTION(DeleteTimerQueueTimer) \
+	DETOURED_FUNCTION(CreateToolhelp32Snapshot) \
 	DETOURED_FUNCTIONS_KERNEL32_DEBUG \
 
 #define DETOURED_FUNCTIONS_NTDLL \
