[UBA]
* Promoted IsNormalized(caskey) and AsNormalized(casKey) to UbaHash
* Added so cas validation skips normalized files for now because they will always fail.. will revisit this

diff --git a/Engine/Source/Programs/UnrealBuildAccelerator/Common/Private/UbaStorage.cpp b/Engine/Source/Programs/UnrealBuildAccelerator/Common/Private/UbaStorage.cpp
--- a/Engine/Source/Programs/UnrealBuildAccelerator/Common/Private/UbaStorage.cpp
+++ b/Engine/Source/Programs/UnrealBuildAccelerator/Common/Private/UbaStorage.cpp
@@ -1327,6 +1327,10 @@ namespace uba
 						//m_logger.Info(TC("Validating %s (%s ago)"), name.c_str(), timeStr.data);
 						CasKey casKey = CasKeyFromString(name.c_str());
 
+						// TODO: Add support for normalized files
+						if (IsNormalized(casKey))
+							return;
+
 						auto reportError = MakeGuard([&]()
 							{
 								SCOPED_FUTEX(lock, l);
diff --git a/Engine/Source/Programs/UnrealBuildAccelerator/Common/Public/UbaRootPaths.h b/Engine/Source/Programs/UnrealBuildAccelerator/Common/Public/UbaRootPaths.h
--- a/Engine/Source/Programs/UnrealBuildAccelerator/Common/Public/UbaRootPaths.h
+++ b/Engine/Source/Programs/UnrealBuildAccelerator/Common/Public/UbaRootPaths.h
@@ -153,19 +153,4 @@ namespace uba
 			searchPos = pathStart + root->path.size();
 		}
 	}
-
-	inline bool IsNormalized(const CasKey& key)
-	{
-		UBA_ASSERT(key != CasKeyZero);
-		return (((u8*)&key)[19] & IsNormalizedMask) == IsNormalizedMask;
-	}
-
-	inline CasKey AsNormalized(const CasKey& key, bool normalized)
-	{
-		UBA_ASSERT(key != CasKeyZero);
-		CasKey newKey = key;
-		u8 flagField = ((u8*)&key)[19];
-		((u8*)&newKey)[19] = normalized ? (flagField | IsNormalizedMask) : (flagField & ~IsNormalizedMask);
-		return newKey;
-	}
 }
diff --git a/Engine/Source/Programs/UnrealBuildAccelerator/Core/Public/UbaHash.h b/Engine/Source/Programs/UnrealBuildAccelerator/Core/Public/UbaHash.h
--- a/Engine/Source/Programs/UnrealBuildAccelerator/Core/Public/UbaHash.h
+++ b/Engine/Source/Programs/UnrealBuildAccelerator/Core/Public/UbaHash.h
@@ -155,6 +155,21 @@ namespace uba
 		return newKey;
 	}
 
+	inline bool IsNormalized(const CasKey& key)
+	{
+		UBA_ASSERT(key != CasKeyZero);
+		return (((u8*)&key)[19] & IsNormalizedMask) == IsNormalizedMask;
+	}
+
+	inline CasKey AsNormalized(const CasKey& key, bool normalized)
+	{
+		UBA_ASSERT(key != CasKeyZero);
+		CasKey newKey = key;
+		u8 flagField = ((u8*)&key)[19];
+		((u8*)&newKey)[19] = normalized ? (flagField | IsNormalizedMask) : (flagField & ~IsNormalizedMask);
+		return newKey;
+	}
+
 	struct CasKeyHasher
 	{
 		CasKeyHasher();
