[UBA]
* Added support for network message priority

diff --git a/Engine/Source/Programs/UnrealBuildAccelerator/Common/Private/UbaNetworkClient.cpp b/Engine/Source/Programs/UnrealBuildAccelerator/Common/Private/UbaNetworkClient.cpp
--- a/Engine/Source/Programs/UnrealBuildAccelerator/Common/Private/UbaNetworkClient.cpp
+++ b/Engine/Source/Programs/UnrealBuildAccelerator/Common/Private/UbaNetworkClient.cpp
@@ -52,6 +52,7 @@ namespace uba
 			m_logger.Detail(TC("Adjusted msg size to %u to stay inside limits"), fixedSendSize);
 		m_sendSize = fixedSendSize;
 		m_receiveTimeoutSeconds = info.receiveTimeoutSeconds;
+		m_useMessagePriority = info.useMessagePriority;
 
 		m_connectionsIt = m_connections.end();
 
@@ -616,18 +617,48 @@ namespace uba
 		}
 
 		// Skip connections that has disconnected
-		Connection* connectionPtr = &*m_connectionsIt;
-		Connection* connectionPtrStart = connectionPtr;
-		while (true)
+		Connection* connectionPtr = nullptr;
+		
+		if (m_useMessagePriority)
+		{
+			Connection* firstConnectionPtr = &m_connections.front();
+			if (message.m_priority == HasPriority && firstConnectionPtr->connected)
+			{
+				connectionPtr = firstConnectionPtr;
+			}
+			else
+			{
+				connectionPtr = &*m_connectionsIt;
+				u32 loopCount = u32(m_connections.size()) - 1;
+				while (loopCount--)
+				{
+					++m_connectionsIt;
+					if (m_connectionsIt == m_connections.end())
+					{
+						m_connectionsIt = m_connections.begin();
+						++m_connectionsIt;
+					}
+					if (connectionPtr->connected)
+						break;
+					connectionPtr = &*m_connectionsIt;
+				}
+			}
+		}
+		else
 		{
-			++m_connectionsIt;
-			if (m_connectionsIt == m_connections.end())
-				m_connectionsIt = m_connections.begin();
-			if (connectionPtr->connected)
-				break;
 			connectionPtr = &*m_connectionsIt;
-			if (connectionPtr == connectionPtrStart)
-				break;
+			Connection* connectionPtrStart = connectionPtr;
+			while (true)
+			{
+				++m_connectionsIt;
+				if (m_connectionsIt == m_connections.end())
+					m_connectionsIt = m_connections.begin();
+				if (connectionPtr->connected)
+					break;
+				connectionPtr = &*m_connectionsIt;
+				if (connectionPtr == connectionPtrStart)
+					break;
+			}
 		}
 
 		connectionItLock.Leave();
@@ -783,9 +814,9 @@ namespace uba
 		return m_prefix.c_str();
 	}
 
-	NetworkMessage::NetworkMessage(NetworkClient& client, u8 serviceId, u8 messageType, BinaryWriter& sendWriter)
+	NetworkMessage::NetworkMessage(NetworkClient& client, u8 serviceId, u8 messageType, BinaryWriter& sendWriter, MessagePriority priority)
 	{
-		Init(client, serviceId, messageType, sendWriter);
+		Init(client, serviceId, messageType, sendWriter, priority);
 	}
 
 	NetworkMessage::~NetworkMessage()
@@ -793,10 +824,11 @@ namespace uba
 		UBA_ASSERT(!m_id);
 	}
 
-	void NetworkMessage::Init(NetworkClient& client, u8 serviceId, u8 messageType, BinaryWriter& sendWriter)
+	void NetworkMessage::Init(NetworkClient& client, u8 serviceId, u8 messageType, BinaryWriter& sendWriter, MessagePriority priority)
 	{
 		m_client = &client;
 		m_sendWriter = &sendWriter;
+		m_priority = priority;
 
 		// Header (SendHeaderSize):
 		// 1 byte    - 2 bits for serviceid, 6 bits for messagetype
diff --git a/Engine/Source/Programs/UnrealBuildAccelerator/Common/Public/UbaNetworkClient.h b/Engine/Source/Programs/UnrealBuildAccelerator/Common/Public/UbaNetworkClient.h
--- a/Engine/Source/Programs/UnrealBuildAccelerator/Common/Public/UbaNetworkClient.h
+++ b/Engine/Source/Programs/UnrealBuildAccelerator/Common/Public/UbaNetworkClient.h
@@ -26,6 +26,7 @@ namespace uba
 		u32 sendSize = SendDefaultSize;
 		u32 receiveTimeoutSeconds = DefaultNetworkReceiveTimeoutSeconds;
 		u32 workerCount = ~0u; // ~0u means logical processor count will be used
+		bool useMessagePriority = false;
 		const u8* cryptoKey128 = nullptr;
 		void Apply(Config& config, const tchar* tableName = TC("NetworkClient"));
 	};
@@ -102,6 +103,8 @@ namespace uba
 		Event m_isConnected;
 		Event m_isOrWasConnected;
 		Atomic<u32> m_connectionCount;
+		bool m_useMessagePriority;
+
 		Futex m_onConnectedFunctionsLock;
 		Vector<OnConnectedFunction> m_onConnectedFunctions;
 		ReaderWriterLock m_onDisconnectedFunctionsLock;
diff --git a/Engine/Source/Programs/UnrealBuildAccelerator/Common/Public/UbaNetworkMessage.h b/Engine/Source/Programs/UnrealBuildAccelerator/Common/Public/UbaNetworkMessage.h
--- a/Engine/Source/Programs/UnrealBuildAccelerator/Common/Public/UbaNetworkMessage.h
+++ b/Engine/Source/Programs/UnrealBuildAccelerator/Common/Public/UbaNetworkMessage.h
@@ -7,13 +7,15 @@
 
 namespace uba
 {
+	enum MessagePriority : u8 { NormalPriority, HasPriority };
+
 	struct NetworkMessage
 	{
-		NetworkMessage(NetworkClient& client, u8 serviceId, u8 messageType, BinaryWriter& sendWriter);
+		NetworkMessage(NetworkClient& client, u8 serviceId, u8 messageType, BinaryWriter& sendWriter, MessagePriority priority = NormalPriority);
 		NetworkMessage() {}
 		~NetworkMessage();
 
-		void Init(NetworkClient& client, u8 serviceId, u8 messageType, BinaryWriter& sendWriter);
+		void Init(NetworkClient& client, u8 serviceId, u8 messageType, BinaryWriter& sendWriter, MessagePriority priority = NormalPriority);
 
 		bool Send(BinaryReader& response);
 		bool Send(BinaryReader& response, Timer& outTimer);
@@ -39,6 +41,7 @@ namespace uba
 		u32 m_responseSize = 0;
 		u32 m_responseCapacity = 0;
 		u16 m_id = 0;
+		MessagePriority m_priority;
 		Atomic<u32> m_error;
 		NetworkClient::Connection* m_connection = nullptr;
 		friend NetworkClient;
