[UBA]
* Enabled message priority for cache client
* Fixed some stats tracking for kernel calls when downloading cached files

diff --git a/Engine/Source/Programs/UnrealBuildAccelerator/Common/Private/UbaCacheClient.cpp b/Engine/Source/Programs/UnrealBuildAccelerator/Common/Private/UbaCacheClient.cpp
--- a/Engine/Source/Programs/UnrealBuildAccelerator/Common/Private/UbaCacheClient.cpp
+++ b/Engine/Source/Programs/UnrealBuildAccelerator/Common/Private/UbaCacheClient.cpp
@@ -537,7 +537,7 @@ namespace uba
 			TimerScope ts(cacheStats.fetchEntries);
 			// Fetch entries.. server will provide as many as fits. TODO: Should it be possible to ask for more entries?
 			StackBinaryWriter<32> writer;
-			NetworkMessage msg(m_client, CacheServiceId, CacheMessageType_FetchEntries, writer);
+			NetworkMessage msg(m_client, CacheServiceId, CacheMessageType_FetchEntries, writer, HasPriority);
 			writer.Write7BitEncoded(MakeId(bucket.id));
 			writer.WriteCasKey(cmdKey);
 			if (!msg.Send(reader))
@@ -820,6 +820,8 @@ namespace uba
 								return;
 							}
 
+							KernelStatsScope _(kernelStats);
+
 							UBA_TRACE_FETCH_HINT_SCOPE("FetchFile")
 							TimerScope fts(cacheStats.fetchOutput);
 							if (!FetchFile(bucket, rootPaths, info, cacheStats, storageStats, casKeyOffset))
@@ -1332,7 +1334,7 @@ namespace uba
 		{
 			Entry& e = entries.emplace_back(data, i, messageFetchSize);
 			StackBinaryWriter<32> writer;
-			e.message.Init(m_client, CacheServiceId, messageType, writer);
+			e.message.Init(m_client, CacheServiceId, messageType, writer, HasPriority);
 			writer.Write7BitEncoded(MakeId(bucketId));
 			writer.WriteU32(tableSize + i*messageFetchSize);
 			if (!e.message.SendAsync(e.reader, [](bool error, void* userData) { ((Event*)userData)->Set(); }, &e.done))
@@ -1403,7 +1405,7 @@ namespace uba
 	bool CacheClient::ReportUsedEntry(Vector<ProcessLogLine>& outLogLines, bool ownedLogLines, Bucket& bucket, const CasKey& cmdKey, u32 entryId)
 	{
 		StackBinaryWriter<128> writer;
-		NetworkMessage msg(m_client, CacheServiceId, CacheMessageType_ReportUsedEntry, writer);
+		NetworkMessage msg(m_client, CacheServiceId, CacheMessageType_ReportUsedEntry, writer, HasPriority);
 		writer.Write7BitEncoded(MakeId(bucket.id));
 		writer.WriteCasKey(cmdKey);
 		writer.Write7BitEncoded(entryId);
diff --git a/Engine/Source/Programs/UnrealBuildAccelerator/Host/Private/UbaExports.cpp b/Engine/Source/Programs/UnrealBuildAccelerator/Host/Private/UbaExports.cpp
--- a/Engine/Source/Programs/UnrealBuildAccelerator/Host/Private/UbaExports.cpp
+++ b/Engine/Source/Programs/UnrealBuildAccelerator/Host/Private/UbaExports.cpp
@@ -812,6 +812,7 @@ extern "C"
 		ncci.receiveTimeoutSeconds = DefaultNetworkReceiveTimeoutSeconds;
 		ncci.cryptoKey128 = crypto128;
 		ncci.workerCount = 0;
+		ncci.useMessagePriority = true;
 		ncci.Apply(GetConfig(), TC("CacheNetworkClient"));
 		bool ctorSuccess = false;
 		auto networkClient = new NetworkClientWithBackend(ctorSuccess, ncci, server.backend, TC("UbaCache"));
