[UBA]
* Added so trace interval can be set in options (default to 500)
* Changed so AddEnvironmentVariableNoLock takes stringview instead of naked char*

diff --git a/Engine/Source/Programs/UnrealBuildAccelerator/Common/Private/UbaSession.cpp b/Engine/Source/Programs/UnrealBuildAccelerator/Common/Private/UbaSession.cpp
--- a/Engine/Source/Programs/UnrealBuildAccelerator/Common/Private/UbaSession.cpp
+++ b/Engine/Source/Programs/UnrealBuildAccelerator/Common/Private/UbaSession.cpp
@@ -239,6 +239,7 @@ namespace uba
 		table.GetValueAsU32(writeFilesBottleneck, TC("WriteFilesBottleneck"));
 		table.GetValueAsU32(writeFilesFileMapMaxMb, TC("WriteFilesFileMapMaxMB"));
 		table.GetValueAsU32(writeFilesNoBufferingMinMb, TC("WriteFilesNoBufferingMinMB"));
+		table.GetValueAsU32(traceIntervalMs, TC("TraceIntervalMs"));
 
 	}
 
@@ -525,11 +526,11 @@ namespace uba
 		return res;
 	}
 
-	void Session::AddEnvironmentVariableNoLock(const tchar* key, const tchar* value)
+	void Session::AddEnvironmentVariableNoLock(StringView key, StringView value)
 	{
-		m_environmentVariables.insert(m_environmentVariables.end(), key, key + TStrlen(key));
+		m_environmentVariables.insert(m_environmentVariables.end(), key.data, key.data + key.count);
 		m_environmentVariables.push_back('=');
-		m_environmentVariables.insert(m_environmentVariables.end(), value, value + TStrlen(value));
+		m_environmentVariables.insert(m_environmentVariables.end(), value.data, value.data + value.count);
 		m_environmentVariables.push_back(0);
 	}
 
@@ -1577,6 +1578,7 @@ namespace uba
 		m_allowLinkDependencyCrawler = info.allowLinkDependencyCrawler;
 
 		m_detailedTrace = info.detailedTrace;
+		m_traceIntervalMs = info.traceIntervalMs;
 		m_traceChildProcesses = info.traceChildProcesses;
 		m_traceWrittenFiles = info.traceWrittenFiles;
 		m_logToFile = info.logToFile;
@@ -2708,12 +2710,12 @@ namespace uba
 
 			if (m_runningRemote && varName.Equals(TCV("PATH")))
 			{
-				AddEnvironmentVariableNoLock(TC("PATH"), TC("c:\\noenvironment"));
+				AddEnvironmentVariableNoLock(TCV("PATH"), TCV("c:\\noenvironment"));
 				return;
 			}
 			if (varName.Equals(TCV("TEMP")) || varName.Equals(TCV("TMP")))
 			{
-				AddEnvironmentVariableNoLock(varName.data, m_tempPath.data);
+				AddEnvironmentVariableNoLock(varName, m_tempPath);
 				return;
 			}
 			if (varName.Equals(TCV("_CL_")) || varName.Equals(TCV("CL")))
@@ -2721,7 +2723,7 @@ namespace uba
 				return;
 			}
 
-			AddEnvironmentVariableNoLock(varName.data, varValue);
+			AddEnvironmentVariableNoLock(varName, ToView(varValue));
 		};
 
 		if (m_environmentMemory.empty())
@@ -2737,9 +2739,9 @@ namespace uba
 			while (reader.GetLeft())
 				HandleEnvironmentVar(reader.ReadString().c_str());
 		}
-		AddEnvironmentVariableNoLock(TC("MSBUILDDISABLENODEREUSE"), TC("1")); // msbuild will reuse existing helper nodes but since those are not detoured we can't let that happen
-		AddEnvironmentVariableNoLock(TC("DOTNET_CLI_USE_MSBUILD_SERVER"), TC("0")); // Disable msbuild server
-		AddEnvironmentVariableNoLock(TC("DOTNET_CLI_TELEMETRY_OPTOUT"), TC("1")); // Stop talking to telemetry service
+		AddEnvironmentVariableNoLock(TCV("MSBUILDDISABLENODEREUSE"), TCV("1")); // msbuild will reuse existing helper nodes but since those are not detoured we can't let that happen
+		AddEnvironmentVariableNoLock(TCV("DOTNET_CLI_USE_MSBUILD_SERVER"), TCV("0")); // Disable msbuild server
+		AddEnvironmentVariableNoLock(TCV("DOTNET_CLI_TELEMETRY_OPTOUT"), TCV("1")); // Stop talking to telemetry service
 #else
 		auto HandleEnvironmentVar = [&](const tchar* env)
 		{
@@ -2779,7 +2781,7 @@ namespace uba
 					paths += ":";
 				paths.append(s, e);
 			}
-			AddEnvironmentVariableNoLock("PATH", paths.c_str());
+			AddEnvironmentVariableNoLock(TCV("PATH"), paths);
 		};
 
 		if (m_environmentMemory.empty())
@@ -2794,10 +2796,10 @@ namespace uba
 			while (reader.GetLeft())
 				HandleEnvironmentVar(reader.ReadString().c_str());
 		}
-		AddEnvironmentVariableNoLock("TMPDIR", m_tempPath.data);
+		AddEnvironmentVariableNoLock(TCV("TMPDIR"), m_tempPath);
 #endif
 
-		AddEnvironmentVariableNoLock(TC("UBA_DETOURED"), TC("1"));
+		AddEnvironmentVariableNoLock(TCV("UBA_DETOURED"), TCV("1"));
 
 		m_environmentVariables.push_back(0);
 		return m_environmentVariables.data();
@@ -4183,7 +4185,7 @@ namespace uba
 		while (true)
 		{
 			TraceSessionUpdate();
-			if (m_traceThreadEvent.IsSet(500))
+			if (m_traceThreadEvent.IsSet(m_traceIntervalMs))
 				break;
 		}
 	}
diff --git a/Engine/Source/Programs/UnrealBuildAccelerator/Common/Private/UbaSessionClient.cpp b/Engine/Source/Programs/UnrealBuildAccelerator/Common/Private/UbaSessionClient.cpp
--- a/Engine/Source/Programs/UnrealBuildAccelerator/Common/Private/UbaSessionClient.cpp
+++ b/Engine/Source/Programs/UnrealBuildAccelerator/Common/Private/UbaSessionClient.cpp
@@ -1434,17 +1434,21 @@ namespace uba
 		}
 
 		#if PLATFORM_WINDOWS
-		AddEnvironmentVariableNoLock(TC("TEMP"), m_tempPath.data);
-		AddEnvironmentVariableNoLock(TC("TMP"), m_tempPath.data);
+		AddEnvironmentVariableNoLock(TCV("TEMP"), m_tempPath);
+		AddEnvironmentVariableNoLock(TCV("TMP"), m_tempPath);
 		#else
-		AddEnvironmentVariableNoLock(TC("TMPDIR"), m_tempPath.data);
+		AddEnvironmentVariableNoLock(TCV("TMPDIR"), m_tempPath);
 		#endif
 
 		StringBuffer<> v;
 		for (auto& var : m_localEnvironmentVariables)
-			if (GetEnvironmentVariableW(var, v.data, v.capacity))
-				AddEnvironmentVariableNoLock(var, v.data);
-
+		{
+			if (u32 res = GetEnvironmentVariableW(var, v.data, v.capacity))
+			{
+				v.count = res;
+				AddEnvironmentVariableNoLock(ToView(var), v);
+			}
+		}
 		m_environmentVariables.push_back(0);
 	}
 
diff --git a/Engine/Source/Programs/UnrealBuildAccelerator/Common/Public/UbaSession.h b/Engine/Source/Programs/UnrealBuildAccelerator/Common/Public/UbaSession.h
--- a/Engine/Source/Programs/UnrealBuildAccelerator/Common/Public/UbaSession.h
+++ b/Engine/Source/Programs/UnrealBuildAccelerator/Common/Public/UbaSession.h
@@ -195,7 +195,7 @@ namespace uba
 		bool WriteFileToDisk(ProcessImpl& process, WrittenFile& file);
 		bool WriteMemoryToDisk(FileAccessor& destinationFile, const void* fileMem, u64 fileSize);
 		bool GetFileMemory(const Function<bool(const void*, u64)>& func, const StringKey& fileNameKey, StringView filePath, bool deleteInternalMapping);
-		void AddEnvironmentVariableNoLock(const tchar* key, const tchar* value);
+		void AddEnvironmentVariableNoLock(StringView key, StringView value);
 		bool WriteDirectoryEntriesInternal(DirectoryTable::Directory& dir, const StringKey& dirKey, StringView dirPath, bool isRefresh, u32& outTableOffset);
 		void WriteDirectoryEntriesRecursive(const StringKey& dirKey, StringView dirPath, u32& outTableOffset);
 		bool CopyImports(Vector<BinaryModule>& out, const tchar* library, tchar* applicationDir, tchar* applicationDirEnd, UnorderedSet<TString>& handledImports, const char* const* loaderPaths);
@@ -333,6 +333,7 @@ namespace uba
 		Event m_traceThreadEvent;
 		Thread m_traceThread;
 		StringBuffer<256> m_traceOutputFile;
+		u32 m_traceIntervalMs;
 		TString m_extraInfo;
 		u64 m_maxPageSize = ~u64(0);
 		u64 m_previousTotalCpuTime = 0;
diff --git a/Engine/Source/Programs/UnrealBuildAccelerator/Common/Public/UbaSessionCreateInfo.h b/Engine/Source/Programs/UnrealBuildAccelerator/Common/Public/UbaSessionCreateInfo.h
--- a/Engine/Source/Programs/UnrealBuildAccelerator/Common/Public/UbaSessionCreateInfo.h
+++ b/Engine/Source/Programs/UnrealBuildAccelerator/Common/Public/UbaSessionCreateInfo.h
@@ -49,5 +49,6 @@ namespace uba
 		u32 writeFilesBottleneck = 16;			// When writing files to disk this number will control how many that can write in parallel
 		u32 writeFilesFileMapMaxMb = 100000;	// If file size is smaller than this number, then use memory maps for writing to disk
 		u32 writeFilesNoBufferingMinMb = 16;	// If file size is at least this number, then overlapped io with no buffering will be used
+		u32 traceIntervalMs = 500;				// milliseconds between each trace entry
 	};
 }
