[UBA]
* Fixed hang when pressing ctrl-c in ubt if uba is in the middle of both downloading cache entries and compiling

diff --git a/Engine/Source/Programs/UnrealBuildAccelerator/Common/Private/UbaNetworkServer.cpp b/Engine/Source/Programs/UnrealBuildAccelerator/Common/Private/UbaNetworkServer.cpp
--- a/Engine/Source/Programs/UnrealBuildAccelerator/Common/Private/UbaNetworkServer.cpp
+++ b/Engine/Source/Programs/UnrealBuildAccelerator/Common/Private/UbaNetworkServer.cpp
@@ -805,6 +805,8 @@ namespace uba
 	{
 		{
 			SCOPED_FUTEX(m_availableWorkersLock, lock);
+			if (!m_workersEnabled) // Already disconnecting.
+				return;
 			m_workersEnabled = false;
 			while (PopWorkerRequest* req = m_firstRequest)
 			{
@@ -821,15 +823,17 @@ namespace uba
 
 		{
 			SCOPED_WRITE_LOCK(m_connectionsLock, lock);
+			List<Connection> connections(std::move(m_connections));
+			lock.Leave();
+
 			bool success = true;
-			for (auto& c : m_connections)
+			for (auto& c : connections)
 			{
 				success = c.Stop() && success;
 				m_sendTimer += c.m_sendTimer;
 				m_encryptTimer += c.m_encryptTimer;
 				m_decryptTimer += c.m_decryptTimer;
 			}
-			lock.Leave();
 
 			// If stopping connections fail we need to abort because we will most likely run into a deadlock when deleting the workers.
 			if (!success)
@@ -842,7 +846,7 @@ namespace uba
 		FlushWorkers();
 
 		SCOPED_WRITE_LOCK(m_connectionsLock, lock);
-		m_connections.clear();
+		UBA_ASSERT(m_connections.empty());
 		m_allClientsDisconnected = true;
 		m_workersEnabled = true;
 	}
