[UBA]
* Optimized download of files from server to client when both source and target are compressed and target is memory mapped
* Simplified code sending files from client to server
* Enabled the path that does multiple decompressions of downloaded data if possible.

diff --git a/Engine/Source/Programs/UnrealBuildAccelerator/Common/Private/UbaStorageUtils.cpp b/Engine/Source/Programs/UnrealBuildAccelerator/Common/Private/UbaStorageUtils.cpp
--- a/Engine/Source/Programs/UnrealBuildAccelerator/Common/Private/UbaStorageUtils.cpp
+++ b/Engine/Source/Programs/UnrealBuildAccelerator/Common/Private/UbaStorageUtils.cpp
@@ -253,7 +253,6 @@ namespace uba
 		u64 fileSize = sourceSize;
 
 		u16 storeId = 0;
-		bool isFirst = true;
 		bool sendEnd = false;
 		u64 sendLeft = fileSize;
 		u64 sendPos = 0;
@@ -268,64 +267,60 @@ namespace uba
 				return msg.Send();
 			};
 
-		while (sendLeft)
+		StackBinaryWriter<SendMaxSize> writer;
+
+		if (sendLeft)
 		{
-			StackBinaryWriter<SendMaxSize> writer;
-			NetworkMessage msg(client, StorageServiceId, isFirst ? StorageMessageType_StoreBegin : StorageMessageType_StoreSegment, writer);
-			if (isFirst)
-			{
-				writer.WriteCasKey(casKey);
-				writer.WriteU64(fileSize);
-				writer.WriteU64(sourceSize);
-				writer.WriteString(hint);
-			}
-			else
-			{
-				UBA_ASSERT(storeId != 0);
-				writer.WriteU16(storeId);
-				writer.WriteU64(sendPos);
-			}
+			NetworkMessage msg(client, StorageServiceId, StorageMessageType_StoreBegin, writer);
+			writer.WriteCasKey(casKey);
+			writer.WriteU64(fileSize);
+			writer.WriteU64(sourceSize);
+			writer.WriteString(hint);
 
-			u64 capacityLeft = writer.GetCapacityLeft();
-			u64 toWrite = Min(sendLeft, capacityLeft);
+			u64 toWrite = Min(sendLeft, writer.GetCapacityLeft());
 			writer.WriteBytes(readData, toWrite);
 
+			StackBinaryReader<128> reader;
+			if (!msg.Send(reader))
+				return false;
+
 			readData += toWrite;
 			sendLeft -= toWrite;
 			sendPos += toWrite;
 
-			bool isDone = sendLeft == 0;
-
-			if (isFirst) // First message must always be acknowledged (provide a reader) to make sure there is an entry on server that can be waited on.
-			{
-				StackBinaryReader<128> reader;
-				if (!msg.Send(reader))
-					return false;
-				storeId = reader.ReadU16();
-				sendEnd = reader.ReadBool();
-				if (isDone)
-					break;
+			storeId = reader.ReadU16();
+			sendEnd = reader.ReadBool();
 
-				if (!storeId) // Zero means error
-					return logger.Error(TC("Server failed to start storing file %s (%s)"), CasKeyString(casKey).str, hint);
+			if (sendLeft == 0)
+				return sendEndMessage();
 
-				if (storeId == u16(~0)) // File already exists on server
-				{
-					//logger.Info(TC("Server already has file %s (%s)"), CasKeyString(casKey).str, hint);
-					return sendEndMessage();
-				}
+			if (!storeId) // Zero means error
+				return logger.Error(TC("Server failed to start storing file %s (%s)"), CasKeyString(casKey).str, hint);
 
-				isFirst = false;
-			}
-			else
+			if (storeId == u16(~0)) // File already exists on server
 			{
-				if (!msg.Send())
-					return false;
-				if (isDone)
-					break;
+				//logger.Info(TC("Server already has file %s (%s)"), CasKeyString(casKey).str, hint);
+				return sendEndMessage();
 			}
 		}
 
+		while (sendLeft)
+		{
+			NetworkMessage msg(client, StorageServiceId, StorageMessageType_StoreSegment, writer.Reset());
+			writer.WriteU16(storeId);
+			writer.WriteU64(sendPos);
+
+			u64 toWrite = Min(sendLeft, writer.GetCapacityLeft());
+			writer.WriteBytes(readData, toWrite); // We always need a separate buffer on the send since we do in-place encryption of the buffers
+
+			if (!msg.Send())
+				return false;
+
+			readData += toWrite;
+			sendLeft -= toWrite;
+			sendPos += toWrite;
+		}
+
 		return sendEndMessage();
 	}
 
@@ -605,23 +600,23 @@ namespace uba
 			lastResponseSize += sizeof(CompressedFileHeader);
 			new (lastSource) CompressedFileHeader{ casKey };
 
-			auto writePrev = [&]() { return WriteDestination(lastSource, lastResponseSize); };
-
 			if (fileMappingMem)
 			{
-				if (!writePrev())
-					return false;
-				sg.Execute();
+				auto writePrev = [&]() { bool res = WriteDestination(lastSource, lastResponseSize); sg.Execute(); return res;};
 				if (u64 leftCompressed = fileSize - responseSize)
 				{
+					u64 destOffset2 = destOffset + lastResponseSize;
 					u32 error;
-					if (!SendBatchMessages(logger, client, fetchId, fileMappingMem + destOffset, leftCompressed, leftCompressed, sizeOfFirstMessage, readIndex, responseSize, {}, destination, &error))
+					if (!SendBatchMessages(logger, client, fetchId, fileMappingMem + destOffset2, leftCompressed, leftCompressed, sizeOfFirstMessage, readIndex, responseSize, writePrev, destination, &error))
 						return logger.Error(TC("Failed to send batched messages to server while retrieving cas %s to %s. Error: %u"), CasKeyString(casKey).str, destination, error);
 					UBA_ASSERT(responseSize == leftCompressed);
 				}
+				else if (!writePrev())
+					return false;
 			}
 			else
 			{
+				auto writePrev = [&]() { return WriteDestination(lastSource, lastResponseSize); };
 
 				u64 leftCompressed = fileSize - responseSize;
 				while (leftCompressed)
@@ -775,7 +770,7 @@ namespace uba
 					lastDecompressSize = decompressedSize;
 					leftUncompressed -= decompressedSize;
 
-					constexpr bool decompressMultiple = false; // This does not seem to be a win.. it batches more but didn't save any time
+					constexpr bool decompressMultiple = true; // This does not seem to be a win.. it batches more but didn't save any time
 
 					if (!decompressMultiple)
 						break;
