[UBA]
* Minor logging/documentation fixes

diff --git a/Engine/Source/Programs/UnrealBuildAccelerator/Cli/Private/UbaCli.cpp b/Engine/Source/Programs/UnrealBuildAccelerator/Cli/Private/UbaCli.cpp
--- a/Engine/Source/Programs/UnrealBuildAccelerator/Cli/Private/UbaCli.cpp
+++ b/Engine/Source/Programs/UnrealBuildAccelerator/Cli/Private/UbaCli.cpp
@@ -790,11 +790,14 @@ namespace uba
 			void* decoderMem = malloc(decoredMemSize);
 			auto mg = MakeGuard([decoderMem]() { free(decoderMem); });
 
+			u32 blockCounter = 0;
 			while (reader.GetLeft())
 			{
 				u32 compressedBlockSize = reader.ReadU32();
 				u32 decompressedBlockSize = reader.ReadU32();
 
+				logger.Info(TC("Decompressing block %u. %u bytes to %u bytes"), blockCounter++, compressedBlockSize, decompressedBlockSize);
+
 				OO_SINTa decompLen = OodleLZ_Decompress(reader.GetPositionData(), (OO_SINTa)compressedBlockSize, destMem, (OO_SINTa)decompressedBlockSize, OodleLZ_FuzzSafe_Yes, OodleLZ_CheckCRC_No, OodleLZ_Verbosity_None, NULL, 0, NULL, NULL, decoderMem, decoredMemSize);
 				if (decompLen != decompressedBlockSize)
 					return logger.Error(TC("Failed to decompress %s (CompressedSize: %llu DecompressedSize: %llu ReadPos: %llu CompressedBlock: %u DecompressedBlock: %u)"), testDecompress.c_str(), fileSize, decompressedSize, reader.GetPosition(), compressedBlockSize, decompressedBlockSize);
diff --git a/Engine/Source/Programs/UnrealBuildAccelerator/Common/Private/UbaProcess.cpp b/Engine/Source/Programs/UnrealBuildAccelerator/Common/Private/UbaProcess.cpp
--- a/Engine/Source/Programs/UnrealBuildAccelerator/Common/Private/UbaProcess.cpp
+++ b/Engine/Source/Programs/UnrealBuildAccelerator/Common/Private/UbaProcess.cpp
@@ -477,9 +477,14 @@ namespace uba
 
 	bool ProcessImpl::WaitForRead(PipeReader& outReader, PipeReader& errReader)
 	{
+		// This number could be bigger but
+		// 1. It would delay ctrl-c
+		// 2. If processes crashes and does not send exit message
+		constexpr u32 readEventTimeout = 1000;
+
 		while (true)
 		{
-			if (m_readEvent.IsSet(1000))
+			if (m_readEvent.IsSet(readEventTimeout))
 				break;
 
 			#if !PLATFORM_WINDOWS
