[UBA]
* Added detouring of "syscall" to catch vfork. Only detoured on linux for now to fix rocky linux distro
* Changed fstat for file size before close to lseek instead since it for some reason fails on rocky linux distro

diff --git a/Engine/Source/Programs/UnrealBuildAccelerator/Detours/Private/PosixOS/UbaDetoursFunctionsPosix.cpp b/Engine/Source/Programs/UnrealBuildAccelerator/Detours/Private/PosixOS/UbaDetoursFunctionsPosix.cpp
--- a/Engine/Source/Programs/UnrealBuildAccelerator/Detours/Private/PosixOS/UbaDetoursFunctionsPosix.cpp
+++ b/Engine/Source/Programs/UnrealBuildAccelerator/Detours/Private/PosixOS/UbaDetoursFunctionsPosix.cpp
@@ -153,6 +153,7 @@ using namespace uba;
 	DETOURED_FUNCTION(fstatat64) \
 	DETOURED_FUNCTION(fpathconf) \
 	DETOURED_FUNCTION(pathconf) \
+	DETOURED_FUNCTION(syscall) \
 
 #else
 #define DETOURED_FUNCTION_LINUX
@@ -602,9 +603,7 @@ void Shared_close(int fd, const TrueClose& trueClose)
 
 	if (fo->closeId)
 	{
-		struct stat attr;
-		TRUE_WRAPPER(fstat)(fd, &attr); // is this safe? not sure about fflush and user space FILE* streams
-		fi.size = attr.st_size;
+		fi.size = lseek(fd, 0, SEEK_END) // is this safe? not sure about fflush and user space FILE* streams
 		fi.created = true; // TODO: Should be set earlier
 		closeGuard.Execute();
 		//for (auto& kv : g_fileHandles)
@@ -2497,6 +2496,33 @@ UBA_EXPORT pid_t UBA_WRAPPER(vfork)(void)
 	return pid;
 }
 
+#if PLATFORM_LINUX
+UBA_EXPORT long UBA_WRAPPER(syscall)(long number, ...)
+{
+	va_list args;
+	va_start(args, number);
+
+	UBA_INIT_DETOUR(syscall, number, va_arg(args, void *), va_arg(args, void *), va_arg(args, void *), va_arg(args, void *), va_arg(args, void *));
+	DEBUG_LOG_TRUE("syscall", "%ld", number);
+
+	long ret = TRUE_WRAPPER(syscall)(number, va_arg(args, void *), va_arg(args, void *), va_arg(args, void *), va_arg(args, void *), va_arg(args, void *));
+	va_end(args);
+
+	if (number == SYS_clone)
+	{
+		if (ret == 0)
+		{
+		#if PLATFORM_LINUX
+			prctl(PR_SET_PDEATHSIG, SIGHUP, 0, 0, 0); // We want the process to die if the parent die
+		#endif
+			t_inVfork = 1;
+		}
+	}
+
+	return ret;
+}
+#endif
+
 // TODO: Allow multiple popen.. hacky!
 FILE* g_activePopen; 
 StringBuffer<64*1024>* g_activePopenResult;
diff --git a/Engine/Source/Programs/UnrealBuildAccelerator/TestApp/Private/PosixOS/UbaTestAppPosix.cpp b/Engine/Source/Programs/UnrealBuildAccelerator/TestApp/Private/PosixOS/UbaTestAppPosix.cpp
--- a/Engine/Source/Programs/UnrealBuildAccelerator/TestApp/Private/PosixOS/UbaTestAppPosix.cpp
+++ b/Engine/Source/Programs/UnrealBuildAccelerator/TestApp/Private/PosixOS/UbaTestAppPosix.cpp
@@ -4,6 +4,7 @@
 #include <fcntl.h>
 #include <unistd.h>
 #include <sys/stat.h>
+#include <sys/syscall.h>
 #include <stdio.h>
 #include <dirent.h>
 #include <dlfcn.h>
@@ -199,24 +200,69 @@ int main(int argc, char* argv[])
 			if (readlink("/proc/self/exe", execPath, sizeof(execPath)) == -1)
 				return LogError("readlink failed while getting executable path");
 
-			pid_t childPid = 0;
-			const char* args[] { execPath, "-child", nullptr };
-			if (posix_spawn(&childPid, execPath, nullptr, nullptr, (char**)args, environ) != 0)
-				return LogError("posix_spawn failed");
-			int status = 0;
-			do
+			auto Waitpid = [](pid_t childPid)
+				{
+					int status = 0;
+					do
+					{
+						if (waitpid(childPid, &status, WUNTRACED | WCONTINUED) == -1)
+							return LogError("waitpid on child process failed (pid %i)", childPid);
+						if (WIFSIGNALED(status))
+							return LogError("Child process killed by signal %d", WTERMSIG(status));
+						if (WIFSTOPPED(status))
+							return LogError("Child process stopped by signal %d", WSTOPSIG(status));
+						if (WIFCONTINUED(status))
+							return LogError("Child process continued");
+					} while (!WIFEXITED(status));
+					if (WEXITSTATUS(status) != 0)
+						return LogError("Child process failed");
+					return 0;
+				};
+
+			{
+				const char* args[] { execPath, "-child", nullptr };
+				pid_t childPid = 0;
+				if (posix_spawn(&childPid, execPath, nullptr, nullptr, (char**)args, environ) != 0)
+					return LogError("posix_spawn failed");
+				if (Waitpid(childPid) == -1)
+					return -1;
+			}
+			#if PLATFORM_LINUX
 			{
-				if (waitpid(childPid, &status, WUNTRACED | WCONTINUED) == -1)
-					return LogError("waitpid on child process failed (pid %i)", childPid);
-				if (WIFSIGNALED(status))
-					return LogError("Child process killed by signal %d", WTERMSIG(status));
-				if (WIFSTOPPED(status))
-					return LogError("Child process stopped by signal %d", WSTOPSIG(status));
-				if (WIFCONTINUED(status))
-					return LogError("Child process continued");
-			} while (!WIFEXITED(status));
-			if (WEXITSTATUS(status) != 0)
-				return LogError("Child process failed");
+				pid_t childPid = 0;
+				childPid = fork();
+				if (childPid == -1)
+					return LogError("SYS_clone failed");
+				if (childPid == 0)
+				{
+					const char* args[] = { execPath, "-child", nullptr };
+					execve((char*)args[0], (char**)args, environ);  // NOTE: `environ` is implicitly declared	
+					_exit(1);
+				}
+				else
+				{
+					if (Waitpid(childPid) == -1)
+						return -1;
+				}
+			}
+			{
+				pid_t childPid = 0;
+				childPid = syscall(SYS_clone, SIGCHLD, NULL);
+				if (childPid == -1)
+					return LogError("SYS_clone failed");
+				if (childPid == 0)
+				{
+					const char* args[] = { execPath, "-child", nullptr };
+					execve((char*)args[0], (char**)args, environ);  // NOTE: `environ` is implicitly declared	
+					_exit(1);
+				}
+				else
+				{
+					if (Waitpid(childPid) == -1)
+						return -1;
+				}
+			}
+			#endif
 		}
 
 		return 0;
