[UBA]
* Fixed crash when building using pvs on remote helper running wine

diff --git a/Engine/Source/Programs/UnrealBuildAccelerator/Detours/Private/Windows/UbaDetoursFunctionsWin.cpp b/Engine/Source/Programs/UnrealBuildAccelerator/Detours/Private/Windows/UbaDetoursFunctionsWin.cpp
--- a/Engine/Source/Programs/UnrealBuildAccelerator/Detours/Private/Windows/UbaDetoursFunctionsWin.cpp
+++ b/Engine/Source/Programs/UnrealBuildAccelerator/Detours/Private/Windows/UbaDetoursFunctionsWin.cpp
@@ -1005,7 +1005,7 @@ void Init(const DetoursPayload& payload, u64 startTime)
 			auto& objName0 = *reinterpret_cast<UNICODE_STRING*>(buffer0);
 			if (!g_stdHandle[0])
 				objName0.Buffer = const_cast<tchar*>(L"Unset");
-			else if (NT_SUCCESS(True_NtQueryObject(g_stdHandle[0], (OBJECT_INFORMATION_CLASS)1, buffer0, sizeof(buffer0), &neededSize)))
+			else if (NT_SUCCESS(True_NtQueryObject(g_stdHandle[0], (OBJECT_INFORMATION_CLASS)1, buffer0, sizeof(buffer0), &neededSize)) && objName0.Buffer)
 				g_conEnabled[0] = !Equals(objName0.Buffer, L"\\Device\\Null");
 			else
 				objName0.Buffer = const_cast<tchar*>(L"???");
@@ -1014,7 +1014,7 @@ void Init(const DetoursPayload& payload, u64 startTime)
 			auto& objName1 = *reinterpret_cast<UNICODE_STRING*>(buffer1);
 			if (!g_stdHandle[1])
 				objName1.Buffer = const_cast<tchar*>(L"Unset");
-			else if (NT_SUCCESS(True_NtQueryObject(g_stdHandle[1], (OBJECT_INFORMATION_CLASS)1, buffer1, sizeof(buffer1), &neededSize)))
+			else if (NT_SUCCESS(True_NtQueryObject(g_stdHandle[1], (OBJECT_INFORMATION_CLASS)1, buffer1, sizeof(buffer1), &neededSize)) && objName1.Buffer)
 				g_conEnabled[1] = !Equals(objName1.Buffer, L"\\Device\\Null");
 			else
 				objName1.Buffer = const_cast<tchar*>(L"???");
