[UBA]
* Fixed visualizer process stats for when remote process result is in progress of being uploaded to cache

diff --git a/Engine/Source/Programs/UnrealBuildAccelerator/Visualizer/Private/UbaVisualizer.cpp b/Engine/Source/Programs/UnrealBuildAccelerator/Visualizer/Private/UbaVisualizer.cpp
--- a/Engine/Source/Programs/UnrealBuildAccelerator/Visualizer/Private/UbaVisualizer.cpp
+++ b/Engine/Source/Programs/UnrealBuildAccelerator/Visualizer/Private/UbaVisualizer.cpp
@@ -797,9 +797,16 @@ namespace uba
 		TraceView::CacheWrite& write = findIt->second;
 		logger.Info(L"");
 		logger.Info(L"  -------- Cache write stats ----------");
-		logger.Info(L"  Duration                    %9s", TimeToText(write.end - write.start).str);
-		logger.Info(L"  Success                     %9s", write.success ? L"true" : L"false");
-		logger.Info(L"  Bytes sent                  %9s", BytesToText(write.bytesSent).str);
+		if (write.end)
+		{
+			logger.Info(L"  Duration                    %9s", TimeToText(write.end - write.start).str);
+			logger.Info(L"  Success                     %9s", write.success ? L"true" : L"false");
+			logger.Info(L"  Bytes sent                  %9s", BytesToText(write.bytesSent).str);
+		}
+		else
+		{
+			logger.Info(L"  In progress...");
+		}
 	}
 
 	void Visualizer::UpdateFont(Font& font, int height, bool createUnderline)
