[UBA]
* Added missing semicolon for uba detours linux

diff --git a/Engine/Source/Programs/UnrealBuildAccelerator/Detours/Private/PosixOS/UbaDetoursFunctionsPosix.cpp b/Engine/Source/Programs/UnrealBuildAccelerator/Detours/Private/PosixOS/UbaDetoursFunctionsPosix.cpp
--- a/Engine/Source/Programs/UnrealBuildAccelerator/Detours/Private/PosixOS/UbaDetoursFunctionsPosix.cpp
+++ b/Engine/Source/Programs/UnrealBuildAccelerator/Detours/Private/PosixOS/UbaDetoursFunctionsPosix.cpp
@@ -603,7 +603,7 @@ void Shared_close(int fd, const TrueClose& trueClose)
 
 	if (fo->closeId)
 	{
-		fi.size = lseek(fd, 0, SEEK_END) // is this safe? not sure about fflush and user space FILE* streams
+		fi.size = lseek(fd, 0, SEEK_END); // is this safe? not sure about fflush and user space FILE* streams
 		fi.created = true; // TODO: Should be set earlier
 		closeGuard.Execute();
 		//for (auto& kv : g_fileHandles)
