[UBA]
* Fixed so when scheduler is cancelled/stopped we don't start processes after failed to download from cache
* Fixed crash in NetworkClient happening when a tcp backend had multiple messages in internal recv buffer while connection got disconnected. This could cause messages to get Done call from the disconnect and then go out of scope while the body of one of those messages arrived just after
* Fixed crash when pressing ctrl-c while in the middle of cache client FetchCompactTable

diff --git a/Engine/Source/Programs/UnrealBuildAccelerator/Common/Private/UbaCacheClient.cpp b/Engine/Source/Programs/UnrealBuildAccelerator/Common/Private/UbaCacheClient.cpp
--- a/Engine/Source/Programs/UnrealBuildAccelerator/Common/Private/UbaCacheClient.cpp
+++ b/Engine/Source/Programs/UnrealBuildAccelerator/Common/Private/UbaCacheClient.cpp
@@ -1041,6 +1041,11 @@ namespace uba
 		return true;
 	}
 
+	void CacheClient::Disable()
+	{
+		m_connected = false;
+	}
+
 	bool CacheClient::SendPathTable(Bucket& bucket, u32 requiredPathTableSize)
 	{
 		SCOPED_FUTEX(bucket.sendPathTableNetworkLock, lock);
@@ -1329,6 +1334,8 @@ namespace uba
 			Event done;
 		};
 
+		bool success = true;
+
 		List<Entry> entries;
 		for (u32 i=0; i!=messageCount; ++i)
 		{
@@ -1337,8 +1344,11 @@ namespace uba
 			e.message.Init(m_client, CacheServiceId, messageType, writer, HasPriority);
 			writer.Write7BitEncoded(MakeId(bucketId));
 			writer.WriteU32(tableSize + i*messageFetchSize);
-			if (!e.message.SendAsync(e.reader, [](bool error, void* userData) { ((Event*)userData)->Set(); }, &e.done))
-				return false;
+			if (e.message.SendAsync(e.reader, [](bool error, void* userData) { ((Event*)userData)->Set(); }, &e.done))
+				continue;
+			entries.pop_back();
+			success = false;
+			break;
 		}
 
 		u32 timeOutTimeMs = 5*60*1000;
@@ -1346,12 +1356,16 @@ namespace uba
 		u64 written = 0;
 		for (Entry& e : entries)
 		{
-			if (!e.done.IsSet(timeOutTimeMs))
-				return false;
+			if (!e.done.IsSet(timeOutTimeMs)) // This should not be possible
+				m_logger.Error(TC("FetchCompactTable timed out after 5 minutes getting async message response. This timeout will cause a crash."));
 			if (!e.message.ProcessAsyncResults(e.reader))
-				return false;
+				success = false;
 			written += e.reader.GetLeft();
 		}
+
+		if (!success)
+			return false;
+
 		table.EndCommit(data, written);
 		return true;
 	}
diff --git a/Engine/Source/Programs/UnrealBuildAccelerator/Common/Private/UbaNetworkClient.cpp b/Engine/Source/Programs/UnrealBuildAccelerator/Common/Private/UbaNetworkClient.cpp
--- a/Engine/Source/Programs/UnrealBuildAccelerator/Common/Private/UbaNetworkClient.cpp
+++ b/Engine/Source/Programs/UnrealBuildAccelerator/Common/Private/UbaNetworkClient.cpp
@@ -359,8 +359,20 @@ namespace uba
 	{
 		auto& msg = *(NetworkMessage*)bodyContext;
 		if (recvError)
-			msg.m_error = 2;
+		{
+			// It could be that message already had done called and went out of scope when we get this since there can be tcp buffer with multiple messages in it
+			auto& connection = *(Connection*)context;
+			auto& client = connection.owner;
+			bool found = false;
+			SCOPED_WRITE_LOCK(client.m_activeMessagesLock, lock);
+			for (auto m : client.m_activeMessages)
+				found = m == &msg;
+			lock.Leave();
+			if (!found)
+				return client.m_isDisconnecting ? false : client.m_logger.Error(TC("NetworkClient got receive response body for message that does not exist"));
 
+			msg.m_error = 2;
+		}
 		#if UBA_TRACK_NETWORK_TIMES
 		if (msg.m_connection)
 			msg.m_connection->lastBodyRecvTime = GetTime();
diff --git a/Engine/Source/Programs/UnrealBuildAccelerator/Common/Private/UbaScheduler.cpp b/Engine/Source/Programs/UnrealBuildAccelerator/Common/Private/UbaScheduler.cpp
--- a/Engine/Source/Programs/UnrealBuildAccelerator/Common/Private/UbaScheduler.cpp
+++ b/Engine/Source/Programs/UnrealBuildAccelerator/Common/Private/UbaScheduler.cpp
@@ -133,6 +133,7 @@ namespace uba
 	void Scheduler::Stop()
 	{
 		m_loop = false;
+		m_cancelled = true;
 		m_updateThreadLoop.Set();
 		m_thread.Wait();
 		m_cacheDownloadThreads.clear();
@@ -143,6 +144,7 @@ namespace uba
 
 	void Scheduler::Cancel()
 	{
+		m_cancelled = true;
 		m_enableProcessReuse = false;
 		SkipAllQueued();
 		m_session.CancelAllProcesses();
@@ -358,10 +360,22 @@ namespace uba
 	void Scheduler::HandleCacheMissed(u32 processIndex)
 	{
 		SCOPED_FUTEX(m_processEntriesLock, lock);
-		UBA_ASSERT(m_processEntries[processIndex].status == ProcessStatus_DownloadingFromCache);
-		m_processEntries[processIndex].status = ProcessStatus_QueuedForRun;
-		m_processEntriesStart = Min(m_processEntriesStart, processIndex);
+		auto& entry = m_processEntries[processIndex];
+		UBA_ASSERT(entry.status == ProcessStatus_DownloadingFromCache);
+
 		--m_activeCacheQueries;
+
+		if (m_cancelled)
+		{
+			entry.status = ProcessStatus_Skipped;
+			auto pi = entry.info;
+			lock.Leave();
+			SkipProcess(*pi);
+			return;
+		}
+
+		entry.status = ProcessStatus_QueuedForRun;
+		m_processEntriesStart = Min(m_processEntriesStart, processIndex);
 		lock.Leave();
 
 		UpdateQueueCounter(1);
diff --git a/Engine/Source/Programs/UnrealBuildAccelerator/Common/Private/UbaStorageUtils.cpp b/Engine/Source/Programs/UnrealBuildAccelerator/Common/Private/UbaStorageUtils.cpp
--- a/Engine/Source/Programs/UnrealBuildAccelerator/Common/Private/UbaStorageUtils.cpp
+++ b/Engine/Source/Programs/UnrealBuildAccelerator/Common/Private/UbaStorageUtils.cpp
@@ -182,8 +182,10 @@ namespace uba
 		Entry* entries = (Entry*)entriesMem;
 
 		if (sizeof(Entry)*sendCount > sizeof(entriesMem))
+		{
 			entries = (Entry*)malloc(sizeof(Entry)*sendCount);
-		auto entriesGuard = MakeGuard([&]() { if (entries != (Entry*)entriesMem) free(entries); });
+			UBA_ASSERT(u64(entries) % alignof(Entry) == 0);
+		}
 
 		bool success = true;
 		u32 error = 0;
@@ -237,6 +239,9 @@ namespace uba
 		for (u32 i=0; i!=inFlightCount; ++i)
 			entries[i].~Entry();
 
+		if (entries != (Entry*)entriesMem)
+			free(entries);
+
 		readIndex += u32(sendCount);
 
 		if (outError)
diff --git a/Engine/Source/Programs/UnrealBuildAccelerator/Common/Public/UbaCacheClient.h b/Engine/Source/Programs/UnrealBuildAccelerator/Common/Public/UbaCacheClient.h
--- a/Engine/Source/Programs/UnrealBuildAccelerator/Common/Public/UbaCacheClient.h
+++ b/Engine/Source/Programs/UnrealBuildAccelerator/Common/Public/UbaCacheClient.h
@@ -70,6 +70,8 @@ namespace uba
 		inline StorageImpl& GetStorage() { return m_storage; }
 		inline Session& GetSession() { return m_session; }
 
+		void Disable();
+
 	private:
 		struct Bucket;
 		u64 MakeId(u32 bucketId);
diff --git a/Engine/Source/Programs/UnrealBuildAccelerator/Common/Public/UbaScheduler.h b/Engine/Source/Programs/UnrealBuildAccelerator/Common/Public/UbaScheduler.h
--- a/Engine/Source/Programs/UnrealBuildAccelerator/Common/Public/UbaScheduler.h
+++ b/Engine/Source/Programs/UnrealBuildAccelerator/Common/Public/UbaScheduler.h
@@ -144,6 +144,7 @@ namespace uba
 		bool m_forceNative;
 		bool m_useThreadToExitRemoteProcess;
 		bool m_allowDisableRemoteExecution = false;
+		bool m_cancelled = false;
 		ConfigTable* m_processConfigs = nullptr;
 
 		float m_activeLocalProcessWeight = 0.0f;
diff --git a/Engine/Source/Programs/UnrealBuildAccelerator/Host/Private/UbaExports.cpp b/Engine/Source/Programs/UnrealBuildAccelerator/Host/Private/UbaExports.cpp
--- a/Engine/Source/Programs/UnrealBuildAccelerator/Host/Private/UbaExports.cpp
+++ b/Engine/Source/Programs/UnrealBuildAccelerator/Host/Private/UbaExports.cpp
@@ -205,6 +205,7 @@ extern "C"
 		{
 			NetworkBackendTcpCreateInfo info(writer);
 			info.Apply(GetConfig());
+			//info.statusUpdateSeconds = 2;
 			networkBackend = new NetworkBackendTcp(info);
 		}
 
@@ -934,6 +935,7 @@ extern "C"
 	{
 		using namespace uba;
 		auto& networkClient = (NetworkClientWithBackend&)cacheClient->GetClient();
+		cacheClient->Disable();
 		networkClient.Disconnect();
 	}
 
