[UBA]
* Fixed timer precision problem when calculating cache client test entry time

diff --git a/Engine/Source/Programs/UnrealBuildAccelerator/Common/Private/UbaCacheClient.cpp b/Engine/Source/Programs/UnrealBuildAccelerator/Common/Private/UbaCacheClient.cpp
--- a/Engine/Source/Programs/UnrealBuildAccelerator/Common/Private/UbaCacheClient.cpp
+++ b/Engine/Source/Programs/UnrealBuildAccelerator/Common/Private/UbaCacheClient.cpp
@@ -517,7 +517,6 @@ namespace uba
 		bool success = false;
 		auto tg = MakeGuard([&]()
 			{
-				cacheStats.testEntry.time -= (cacheStats.fetchCasTable.time + cacheStats.normalizeFile.time);
 				BinaryWriter writer(memory, 0, SendMaxSize);
 				cacheStats.Write(writer);
 				storageStats.Write(writer);
@@ -525,6 +524,8 @@ namespace uba
 				m_session.GetTrace().CacheEndFetch(fetchId, success, memory, writer.GetPosition());
 			});
 
+		auto subtractTimer = MakeGuard([&](){ cacheStats.testEntry.time -= (cacheStats.fetchCasTable.time + cacheStats.normalizeFile.time); });
+
 		BinaryReader reader(memory, 0, SendMaxSize);
 
 		SCOPED_FUTEX(m_bucketsLock, bucketsLock);
@@ -783,6 +784,8 @@ namespace uba
 						continue;
 				}
 
+				subtractTimer.Execute();
+
 				#if UBA_TRACE_FETCH_CACHE
 				addStoreKeyHint();
 				#endif
