[UBA]
* Fixed bug in visualizer showing byte per second instead of bit for second

diff --git a/Engine/Source/Programs/UnrealBuildAccelerator/Visualizer/Private/UbaVisualizer.cpp b/Engine/Source/Programs/UnrealBuildAccelerator/Visualizer/Private/UbaVisualizer.cpp
--- a/Engine/Source/Programs/UnrealBuildAccelerator/Visualizer/Private/UbaVisualizer.cpp
+++ b/Engine/Source/Programs/UnrealBuildAccelerator/Visualizer/Private/UbaVisualizer.cpp
@@ -2818,8 +2818,8 @@ namespace uba
 						recvPerS = u64(float(updateRecv - session.prevRecv) / duration);
 					}
 
-					drawText(L"Recv: %ls (%sit/s)", BytesToText(updateRecv), BytesToText(recvPerS));
-					drawText(L"Send: %ls (%sit/s)", BytesToText(updateSend), BytesToText(sendPerS));
+					drawText(L"Recv: %ls (%sit/s)", BytesToText(updateRecv), BytesToText(recvPerS*8));
+					drawText(L"Send: %ls (%sit/s)", BytesToText(updateSend), BytesToText(sendPerS*8));
 				}
 			}
 
