[UBA]
* Fixed shutdown crash when pressing ctrl-c while downloading files from cache server. Shutdown has to wait for waitentries to be empty
* Added more logging to try to get more info on an error seen on the farm once

diff --git a/Engine/Source/Programs/UnrealBuildAccelerator/Common/Private/UbaStorageServer.cpp b/Engine/Source/Programs/UnrealBuildAccelerator/Common/Private/UbaStorageServer.cpp
--- a/Engine/Source/Programs/UnrealBuildAccelerator/Common/Private/UbaStorageServer.cpp
+++ b/Engine/Source/Programs/UnrealBuildAccelerator/Common/Private/UbaStorageServer.cpp
@@ -75,6 +75,17 @@ namespace uba
 	{
 		while (m_activeUnmap)
 			m_server.DoAdditionalWork();
+
+		// Need to make sure all wait entries are done before exiting.
+		while (true)
+		{
+			SCOPED_FUTEX_READ(m_waitEntriesLock, lock);
+			if (m_waitEntries.empty())
+				break;
+			lock.Leave();
+			m_server.DoAdditionalWork();
+		}
+
 	}
 
 	bool StorageServer::GetZone(StringBufferBase& out)
@@ -548,6 +559,7 @@ namespace uba
 				FileHandle readFileHandle = InvalidFileHandle;
 				auto rfg = MakeGuard([&]()
 					{
+						todoRemoveMe = 13;
 						if (!CloseFile(nullptr, readFileHandle))
 							m_logger.Error(TC("Failed to close file %s (%s)"), casFile.data, LastErrorToText().data);
 					});
@@ -561,6 +573,7 @@ namespace uba
 					{
 						if (ownsMapping)
 						{
+							todoRemoveMe = 11;
 							if (mappedView.memory)
 								if (!UnmapViewOfFile(m_logger, mappedView.memory, fileSize, TC("FetchBegin")))
 									m_logger.Error(TC("Failed to unmap memory %s at 0x%llx with size %llu (%s)"), hint.data, mappedView.memory, fileSize, LastErrorToText().data);
@@ -569,7 +582,10 @@ namespace uba
 									m_logger.Error(TC("Failed to close file mapping for %s with handle %llu (%s)"), hint.data, mappedView.handle.ToU64(), LastErrorToText().data);
 						}
 						else
+						{
+							todoRemoveMe = 12;
 							m_casDataBuffer.UnmapView(mappedView, TC("FetchBegin"));
+						}
 					});
 
 				todoRemoveMe = 5;
