[UBA]
* Fixed annoying one-off offset in visualizer graphs (cpu, recv, send, etc). Now graphs should match the bars better

diff --git a/Engine/Source/Programs/UnrealBuildAccelerator/Visualizer/Private/UbaVisualizer.cpp b/Engine/Source/Programs/UnrealBuildAccelerator/Visualizer/Private/UbaVisualizer.cpp
--- a/Engine/Source/Programs/UnrealBuildAccelerator/Visualizer/Private/UbaVisualizer.cpp
+++ b/Engine/Source/Programs/UnrealBuildAccelerator/Visualizer/Private/UbaVisualizer.cpp
@@ -1389,6 +1389,7 @@ namespace uba
 								decltype(maxValue) prevValue = 0;
 								u64 prevTime = 0;
 								int prevX = 0;
+								int prevPrevX = 0;
 								int prevY = 0;
 
 								u64 i = 0;
@@ -1438,20 +1439,20 @@ namespace uba
 									{
 										if (isFirstUpdate)
 										{
-											line.push_back({x, y}); //MoveToEx(hdc, x, y, NULL);
 											isFirstDraw = false;
 										}
 										else
 										{
 											if (isFirstDraw)
-												line.push_back({prevX, prevY}); // MoveToEx(hdc, prevX, prevY, NULL);
-											line.push_back({x, y}); // LineTo(hdc, x, y);
+												line.push_back({prevPrevX, prevY});
+											line.push_back({prevX, y});
 											isFirstDraw = false;
 										}
 									}
-									if (x > clientRect.right)
+									if (prevX > clientRect.right)
 										break;
 									isFirstUpdate = false;
+									prevPrevX = prevX;
 									prevX = x;
 									prevY = y;
 									prevValue = value;
@@ -3065,10 +3066,10 @@ namespace uba
 			bool showGraph = m_config.showNetworkStats || m_config.showCpuMemStats || m_config.showDriveStats;
 			if (showGraph && !session.updates.empty())
 			{
-				if (pos.y >= posY && pos.y < posY + GraphHeight)
-				{
-					int posX = int(m_scrollPosX) + progressRect.left;
+				int posX = int(m_scrollPosX) + progressRect.left;
 
+				if (pos.y >= posY && pos.y < posY + GraphHeight && pos.x >= posX)
+				{
 					bool loop = true;
 					u32 reconnectIndex = 0;
 					while (loop)
@@ -3087,16 +3088,16 @@ namespace uba
 							loop = false;
 						}
 
-						int prevX = 100000;
+						int prevPrevX = 100000;
+						int prevX = 0;
 						for (; i!=e; ++i)
 						{
 							u64 updateTime = session.updates[i];
 							auto updateSend = session.networkSend[i];
 							auto updateRecv = session.networkRecv[i];
-							int x = posX + int(float(TimeToS(updateTime)) * scaleX);
 
-							int hitOffset = (prevX - x)/2;
-							if (pos.x + hitOffset >= prevX && pos.x + hitOffset <= x)
+							int hitOffset = (prevPrevX - prevX)/2;
+							if (pos.x + hitOffset <= prevX)
 							{
 								u64 prevTime = 0;
 								u64 prevSend = 0;
@@ -3133,7 +3134,8 @@ namespace uba
 								return;
 							}
 
-							prevX = x;
+							prevPrevX = prevX;
+							prevX = posX + int(float(TimeToS(updateTime)) * scaleX);
 						}
 						++reconnectIndex;
 					}
