[UBA]
* Added fast path to downloading and decompress files from server

diff --git a/Engine/Source/Programs/UnrealBuildAccelerator/Common/Private/UbaStorageUtils.cpp b/Engine/Source/Programs/UnrealBuildAccelerator/Common/Private/UbaStorageUtils.cpp
--- a/Engine/Source/Programs/UnrealBuildAccelerator/Common/Private/UbaStorageUtils.cpp
+++ b/Engine/Source/Programs/UnrealBuildAccelerator/Common/Private/UbaStorageUtils.cpp
@@ -145,7 +145,7 @@ namespace uba
 		return ToCasKey(hasher, storeCompressed);
 	}
 
-	bool SendBatchMessages(Logger& logger, NetworkClient& client, u16 fetchId, u8* slot, u64 capacity, u64 left, u32 messageMaxSize, u32& readIndex, u32& responseSize, const Function<bool()>& runInWaitFunc, const tchar* hint, u32* outError)
+	bool SendBatchMessages(Logger& logger, NetworkClient& client, u16 fetchId, u8* slot, u64 capacity, u64 left, u32 messageMaxSize, u32& readIndex, u32& responseSize, const Function<bool()>& runInWaitFunc, const tchar* hint, u32* outError, const Function<bool(u64)>& progressFunc)
 	{
 		responseSize = 0;
 
@@ -234,6 +234,16 @@ namespace uba
 			}
 			else
 				responseSize += u32(entry.reader.GetLeft());
+
+			if (success && progressFunc)
+			{
+				if (!progressFunc(responseSize))
+				{
+					success = false;
+					if (!error)
+						error = 100;
+				}
+			}
 		}
 
 		for (u32 i=0; i!=inFlightCount; ++i)
@@ -545,7 +555,6 @@ namespace uba
 				actualSize = *(u64*)readBuffer;
 		}
 
-		bytesReceived = fileSize;
 		sizeOnDisk = writeCompressed ? (sizeof(CompressedFileHeader) + fileSize) : actualSize;
 
 		FileAccessor destinationFile(logger, destination);
@@ -643,6 +652,83 @@ namespace uba
 					return false;
 			}
 		}
+		else if (actualSize && (fileMappingMem || destinationMem))
+		{
+			u64 fileSize2 = fileSize - sizeof(u64);
+			readBuffer += sizeof(u64); // Size is stored first
+			responseSize -= sizeof(u64);
+
+			u8* destMem = fileMappingMem;
+
+			if (!destMem)
+				destMem = (u8*)destinationMem->Allocate(actualSize, 1, TC(""));
+
+			u64 initialRead = 0;
+			u32 compressedSize = 0;
+			u32 decompressedSize = 0;
+			u64 nextOffset = 0;
+			void* decoderMem = nullptr;
+			u64 decoderMemSize = 0;
+
+			auto decompressAndWrite = [&](u64 pos)
+				{
+					u64 offset = pos + initialRead;
+					while (true)
+					{
+						if (!compressedSize && offset >= nextOffset + 8)
+						{
+							compressedSize =  ((u32*)readBuffer)[0];
+							decompressedSize = ((u32*)readBuffer)[1];
+							readBuffer += 8;
+							nextOffset += compressedSize + 8;
+						}
+
+						if (offset < nextOffset)
+							return true;
+
+						TimerScope ts2(m_stats.decompressRecv);
+						OO_SINTa decompLen = OodleLZ_Decompress(readBuffer, int(compressedSize), destMem, int(decompressedSize), OodleLZ_FuzzSafe_Yes, OodleLZ_CheckCRC_No, OodleLZ_Verbosity_None, NULL, 0, NULL, NULL, decoderMem, decoderMemSize);
+						if (decompLen != decompressedSize)
+							return logger.Error(TC("Expected %u but got %i when decompressing %u bytes for file %s"), decompressedSize, int(decompLen), compressedSize, destination);
+
+						readBuffer += compressedSize;
+						destMem += decompressedSize;
+						compressedSize = 0;
+
+						if (offset < nextOffset + 8)
+							break;
+					}
+					return true;
+				};
+
+			if (u64 leftCompressed = fileSize2 - responseSize)
+			{
+				initialRead = responseSize;
+
+				u8* compressedData = readBuffer;
+
+				u64 sizeThatFitsInSlot = BufferSlotSize - (readBuffer - slot);
+				bool useAllocation = fileSize2 > sizeThatFitsInSlot;
+				if (useAllocation)
+				{
+					compressedData = (u8*)malloc(fileSize2);
+					memcpy(compressedData, readBuffer, responseSize);
+					
+					sg.Execute();
+					//decoderMem = slot;
+					//decoderMemSize = BufferSlotSize;
+				}
+				auto mbg = MakeGuard([&]() { if (useAllocation) free(compressedData); });
+
+				readBuffer = compressedData;
+
+				u32 error;
+				if (!SendBatchMessages(logger, client, fetchId, compressedData + responseSize, leftCompressed, leftCompressed, sizeOfFirstMessage, readIndex, responseSize, {}, destination, &error, decompressAndWrite))
+					return logger.Error(TC("Failed to send batched messages to server while retrieving cas %s to %s. Error: %u"), CasKeyString(casKey).str, destination, error);
+			}
+			else if (!decompressAndWrite(responseSize))
+				return false;
+		}
 		else if (actualSize)
 		{
 			bool sendSegmentMessage = responseSize == 0;
@@ -839,6 +925,8 @@ namespace uba
 			if (!destinationFile.Close(&lastWritten))
 				return false;
 
+		bytesReceived = fileSize;
+
 		m_stats.recvCasBytesRaw += actualSize;
 		m_stats.recvCasBytesComp += fileSize;
 
