[UBA]
* Added network backend connection priority

diff --git a/Engine/Source/Programs/UnrealBuildAccelerator/Common/Private/UbaNetworkBackendMemory.cpp b/Engine/Source/Programs/UnrealBuildAccelerator/Common/Private/UbaNetworkBackendMemory.cpp
--- a/Engine/Source/Programs/UnrealBuildAccelerator/Common/Private/UbaNetworkBackendMemory.cpp
+++ b/Engine/Source/Programs/UnrealBuildAccelerator/Common/Private/UbaNetworkBackendMemory.cpp
@@ -146,6 +146,10 @@ namespace uba
 		// Not implemented
 	}
 
+	void NetworkBackendMemory::SetPriority(void* connection)
+	{
+	}
+
 	bool NetworkBackendMemory::StartListen(Logger& logger, u16 port, const tchar* ip, const ListenConnectedFunc& connectedFunc)
 	{
 		SCOPED_WRITE_LOCK(m_connectedFuncLock, l);
diff --git a/Engine/Source/Programs/UnrealBuildAccelerator/Common/Private/UbaNetworkBackendTcp.cpp b/Engine/Source/Programs/UnrealBuildAccelerator/Common/Private/UbaNetworkBackendTcp.cpp
--- a/Engine/Source/Programs/UnrealBuildAccelerator/Common/Private/UbaNetworkBackendTcp.cpp
+++ b/Engine/Source/Programs/UnrealBuildAccelerator/Common/Private/UbaNetworkBackendTcp.cpp
@@ -12,9 +12,11 @@
 #include <iphlpapi.h>
 #include <ipifcons.h>
 #include <Mstcpip.h>
+#include <qos2.h>
 #pragma comment (lib, "Netapi32.lib")
 #pragma comment (lib, "Ws2_32.lib")
 #pragma comment(lib, "IPHLPAPI.lib") // For GetAdaptersInfo
+#pragma comment(lib, "qwave.lib")
 #else
 #include <netdb.h>
 #include <netinet/tcp.h>
@@ -61,6 +63,7 @@ namespace uba
 	bool SetLinger(Logger& logger, SOCKET socket, u32 lingerSeconds);
 	bool SetRecvBuf(Logger& logger, SOCKET socket, u32 windowSize);
 	bool SetSendBuf(Logger& logger, SOCKET socket, u32 windowSize);
+	bool SetSocketPriority(Logger& logger, SOCKET socket);
 	bool DisableNagle(Logger& logger, SOCKET socket);
 	void LogTcpInfo(Logger& logger, SOCKET socket);
 
@@ -330,6 +333,13 @@ namespace uba
 		conn.allowLess = allow;
 	}
 
+	void NetworkBackendTcp::SetPriority(void* connection)
+	{
+		auto& conn = *(Connection*)connection;
+		SCOPED_CRITICAL_SECTION(conn.shutdownLock, lock);
+		SetSocketPriority(m_logger, conn.socket);
+	}
+
 	bool NetworkBackendTcp::StartListen(Logger& logger, u16 port, const tchar* ip, const ListenConnectedFunc& connectedFunc)
 	{
 		if (!EnsureInitialized(logger))
@@ -1728,6 +1738,34 @@ namespace uba
 		return true;
 	}
 
+	bool SetSocketPriority(Logger& logger, SOCKET socket)
+	{
+#if PLATFORM_WINDOWS
+		static HANDLE s_qosHandle = [](Logger& logger)
+			{
+				HANDLE h = INVALID_HANDLE_VALUE;
+				QOS_VERSION ver = {1, 0};
+				if (!QOSCreateHandle(&ver, &qosHandle))
+					logger.Error(TC("QOSCreateHandle failed"));
+				return h;
+			}(logger);
+		if (s_qosHandle == INVALID_HANDLE_VALUE)
+			return false;
+		QOS_FLOWID flowId = 0;
+		if (!QOSAddSocketToFlow(s_qosHandle, socket, NULL, QOSTrafficTypeExcellentEffort, QOS_NON_ADAPTIVE_FLOW, &flowId))
+			return logger.Error(TC("QOSAddSocketToFlow failed for socket %i (%s)"), socket, LastErrorToText().data);
+		return true;
+#else
+		int prio = 6;
+		if (setsockopt(socket, SOL_SOCKET, SO_PRIORITY, &prio, sizeof(prio)) == -1)
+			return logger.Error(TC("setsockopt SO_PRIORITY failed (%s)"), LastErrorToText(WSAGetLastError()).data);
+		//uint8_t tos = 46 << 2; // EF = 0b101110xx
+		//if (setsockopt(socket, IPPROTO_IP, IP_TOS, &tos, sizeof(tos)) == -1)
+		//	return logger.Error(TC("setsockopt IP_TOS failed (%s)"), LastErrorToText(WSAGetLastError()).data);
+		return false;
+#endif
+	}
+
 	bool SetKeepAlive(Logger& logger, SOCKET socket) // This will make sure that WSAPoll exits when the network cable is pulled
 	{
 		
diff --git a/Engine/Source/Programs/UnrealBuildAccelerator/Common/Public/UbaNetworkBackend.h b/Engine/Source/Programs/UnrealBuildAccelerator/Common/Public/UbaNetworkBackend.h
--- a/Engine/Source/Programs/UnrealBuildAccelerator/Common/Public/UbaNetworkBackend.h
+++ b/Engine/Source/Programs/UnrealBuildAccelerator/Common/Public/UbaNetworkBackend.h
@@ -54,6 +54,9 @@ namespace uba
 		virtual void SetDisconnectCallback(void* connection, void* context, DisconnectCallback* callback) = 0;
 		virtual void SetAllowLessThanBodySize(void* connection, bool allow) = 0;
 
+		// Set priority of connection
+		virtual void SetPriority(void* connection) = 0;
+
 		// Start listen on port/ip.
 		using ListenConnectedFunc = Function<bool(void* connection, const sockaddr& remoteSocketAddr)>;
 		virtual bool StartListen(Logger& logger, u16 port, const tchar* ip, const ListenConnectedFunc& connectedFunc) = 0;
diff --git a/Engine/Source/Programs/UnrealBuildAccelerator/Common/Public/UbaNetworkBackendMemory.h b/Engine/Source/Programs/UnrealBuildAccelerator/Common/Public/UbaNetworkBackendMemory.h
--- a/Engine/Source/Programs/UnrealBuildAccelerator/Common/Public/UbaNetworkBackendMemory.h
+++ b/Engine/Source/Programs/UnrealBuildAccelerator/Common/Public/UbaNetworkBackendMemory.h
@@ -20,6 +20,7 @@ namespace uba
 		virtual void SetRecvTimeout(void* connection, u32 timeoutMs, void* context, RecvTimeoutCallback* callback) override;
 		virtual void SetDisconnectCallback(void* connection, void* context, DisconnectCallback* callback) override;
 		virtual void SetAllowLessThanBodySize(void* connection, bool allow) override;
+		virtual void SetPriority(void* connection) override;
 
 		virtual bool StartListen(Logger& logger, u16 port, const tchar* ip, const ListenConnectedFunc& connectedFunc) override;
 		virtual void StopListen() override;
diff --git a/Engine/Source/Programs/UnrealBuildAccelerator/Common/Public/UbaNetworkBackendTcp.h b/Engine/Source/Programs/UnrealBuildAccelerator/Common/Public/UbaNetworkBackendTcp.h
--- a/Engine/Source/Programs/UnrealBuildAccelerator/Common/Public/UbaNetworkBackendTcp.h
+++ b/Engine/Source/Programs/UnrealBuildAccelerator/Common/Public/UbaNetworkBackendTcp.h
@@ -41,6 +41,7 @@ namespace uba
 		virtual void SetRecvTimeout(void* connection, u32 timeoutMs, void* context, RecvTimeoutCallback* callback) override;
 		virtual void SetDisconnectCallback(void* connection, void* context, DisconnectCallback* callback) override;
 		virtual void SetAllowLessThanBodySize(void* connection, bool allow) override;
+		virtual void SetPriority(void* connection) override;
 
 		virtual bool StartListen(Logger& logger, u16 port, const tchar* ip, const ListenConnectedFunc& connectedFunc) override;
 		virtual void StopListen() override;
