[UBA]
* New UbaWine.dll.so with more api to query environment

diff --git a/Engine/Source/Programs/UnrealBuildAccelerator/Wine/Linux.c b/Engine/Source/Programs/UnrealBuildAccelerator/Wine/Linux.c
--- a/Engine/Source/Programs/UnrealBuildAccelerator/Wine/Linux.c
+++ b/Engine/Source/Programs/UnrealBuildAccelerator/Wine/Linux.c
@@ -3,8 +3,21 @@
 #include <netinet/tcp.h>
 #include <netinet/in.h>
 #include <errno.h>
+#include <stdio.h>
+#include <stdlib.h>
+#include "../Common/Public/UbaLinuxNetworkWrappers.h"
 
-int unix_get_tcp_info(int fd, void *buf, int *len)
+int UnixGetsockopt(int fd, int level, int optname, char* optval, int* optlen)
+{
+	return getsockopt(fd, level, optname, optval, optlen);
+}
+
+int UnixSetsockopt(int fd, int level, int optname, char* optval, int optlen)
+{
+	return setsockopt(fd, level, optname, optval, optlen);
+}
+
+int UnixGetTcpInfo(int fd, void* buf, int* len)
 {
 	socklen_t l = *len;
 	if (getsockopt(fd, IPPROTO_TCP, TCP_INFO, buf, &l) == -1)
diff --git a/Engine/Source/Programs/UnrealBuildAccelerator/Wine/Windows.c b/Engine/Source/Programs/UnrealBuildAccelerator/Wine/Windows.c
--- a/Engine/Source/Programs/UnrealBuildAccelerator/Wine/Windows.c
+++ b/Engine/Source/Programs/UnrealBuildAccelerator/Wine/Windows.c
@@ -3,17 +3,56 @@
 #include <winsock2.h>
 #include <stdint.h>
 #include <unistd.h>
+#include <stdbool.h>
+
+// NOTE, README
+// Must use "WINAPI" on function and add Windows.def
+
 
 NTSYSAPI NTSTATUS CDECL wine_server_handle_to_fd( HANDLE handle, unsigned int access, int *unix_fd, unsigned int *options );
 
-extern int unix_get_tcp_info(int fd, void *buf, int *len);
+extern int UnixGetsockopt(int fd, int level, int optname, char* optval, int* optlen);
+extern int UnixSetsockopt(int fd, int level, int optname, char* optval, int optlen);
+
+extern int UnixGetTcpInfo(int fd, void* buf, int* len);
+extern bool UnixGetTcpAutoTuning(int* outReadMin, int* outReadDefault, int* outReadMax, int* outWriteMin, int* outWriteDefault, int* outWriteMax);
 
-INT WINAPI GetLinuxTcpInfo(SOCKET s, void *buf, INT *len)
+INT WINAPI GetLinuxTcpInfo(SOCKET s, void* buf, INT* len)
 {
 	int fd = -1;
 	if (wine_server_handle_to_fd((HANDLE)s, 0, &fd, NULL) != 0)
 		return -1;
-	int res = unix_get_tcp_info(fd, buf, len);
+	int res = UnixGetTcpInfo(fd, buf, len);
 	close(fd);
 	return res;
 }
+
+int UNIX_IPPROTO_TCP = 6;
+int TCP_CONGESTION = 13;
+
+
+bool WINAPI GetCongestionAlgorithm(SOCKET s, char* out, int outCapacity)
+{
+	int fd = -1;
+	if (wine_server_handle_to_fd((HANDLE)s, 0, &fd, NULL) != 0)
+		return -1;
+	int res = UnixGetsockopt(fd, UNIX_IPPROTO_TCP, TCP_CONGESTION, out, &outCapacity);
+	close(fd);
+	return res == 0;
+}
+
+bool WINAPI SetCongestionAlgorithm(SOCKET s, char* value, int valuelen)
+{
+	int fd = -1;
+	if (wine_server_handle_to_fd((HANDLE)s, 0, &fd, NULL) != 0)
+		return -1;
+	int res = UnixSetsockopt(fd, UNIX_IPPROTO_TCP, TCP_CONGESTION, value, valuelen);
+	close(fd);
+	return res == 0;
+}
+
+
+bool WINAPI GetTcpAutoTuning(int* outReadMin, int* outReadDefault, int* outReadMax, int* outWriteMin, int* outWriteDefault, int* outWriteMax)
+{
+	return UnixGetTcpAutoTuning(outReadMin, outReadDefault, outReadMax, outWriteMin, outWriteDefault, outWriteMax);
+}
\ No newline at end of file
diff --git a/Engine/Source/Programs/UnrealBuildAccelerator/Wine/Windows.def b/Engine/Source/Programs/UnrealBuildAccelerator/Wine/Windows.def
--- a/Engine/Source/Programs/UnrealBuildAccelerator/Wine/Windows.def
+++ b/Engine/Source/Programs/UnrealBuildAccelerator/Wine/Windows.def
@@ -1,3 +1,6 @@
 LIBRARY libtcpinfo.dll
 EXPORTS
-    GetLinuxTcpInfo @1    ; stdcall decoration unnecessary for x64
+	GetLinuxTcpInfo
+	GetCongestionAlgorithm
+	SetCongestionAlgorithm
+	GetTcpAutoTuning
\ No newline at end of file
