[UBA]
* Changed to <Time/Bytes/Count>ToText is not using sprintf
* Moved TimeToText from UbaTimer (core) to UbaLogger (common)
* Added option to use base2 for BytesToText

diff --git a/Engine/Source/Programs/UnrealBuildAccelerator/Common/Private/UbaLogger.cpp b/Engine/Source/Programs/UnrealBuildAccelerator/Common/Private/UbaLogger.cpp
--- a/Engine/Source/Programs/UnrealBuildAccelerator/Common/Private/UbaLogger.cpp
+++ b/Engine/Source/Programs/UnrealBuildAccelerator/Common/Private/UbaLogger.cpp
@@ -367,32 +367,94 @@ namespace uba
 #endif
 	}
 
-	BytesToText::BytesToText(u64 bytes)
-	{
-		if (bytes < 1000)
-			TSprintf_s(str, 32, TC("%ub"), u32(bytes));
-		else if (bytes < 1000 * 1000)
-			TSprintf_s(str, 32, TC("%.1fkb"), double(bytes) / 1000ull);
-		else if (bytes < 1000ull * 1000 * 1000)
-			TSprintf_s(str, 32, TC("%.1fmb"), double(bytes) / (1000ull * 1000));
-		else if (bytes < 1000ull * 1000 * 1000 * 1000)
-			TSprintf_s(str, 32, TC("%.1fgb"), double(bytes) / (1000ull * 1000 * 1000));
+	u32 UToA(tchar* out, u32 val)
+	{
+		tchar tmp[16];
+		tchar* last = tmp + sizeof_array(tmp) - 1;
+		tchar* p = last;
+		*p = 0;
+		do
+		{
+			*--p = '0' + (val % 10);
+			val /= 10;
+		}
+		while (val);
+		u32 len = u32(last - p);
+		memcpy(out, p, (len+1)*sizeof(tchar));
+		return len;
+	}
+
+	void Format(tchar* out, u32 value, const tchar* suffix)
+	{
+		out += UToA(out, value);
+		TStrcpy_s(out, 4, suffix);
+	}
+
+	void Format1dp(tchar* out, double value, const tchar* suffix)
+	{
+		u32 scaled = (u32)(value * 10 + 0.5); // round to 1 decimal
+		out += UToA(out, scaled / 10);
+		*out++ = '.';
+		*out++ = '0' + scaled % 10;
+		TStrcpy_s(out, 4, suffix);
+	}
+
+	BytesToText::BytesToText(u64 bytes, bool base2)
+	{
+		u64 v = base2 ? 1024 : 1000;
+
+		if (bytes < v)
+			Format(str, u32(bytes), TC("b"));
+		else if (bytes < v * v)
+			Format1dp(str, double(bytes)/double(v), TC("kb"));
+		else if (bytes < v * v * v)
+			Format1dp(str, double(bytes)/double(v * v), TC("mb"));
+		else if (bytes < v * v * v * v)
+			Format1dp(str, double(bytes)/double(v * v * v), TC("gb"));
 		else
-			TSprintf_s(str, 32, TC("%.1ftb"), double(bytes) / (1000ull * 1000 * 1000 * 1000));
+			Format1dp(str, double(bytes)/double(v * v * v * v), TC("tb"));
 	}
 
 	CountToText::CountToText(u64 count)
 	{
 		if (count < 1000)
-			TSprintf_s(str, 32, TC("%u"), u32(count));
+			UToA(str, u32(count));
 		else if (count < 1000 * 1000)
-			TSprintf_s(str, 32, TC("%.1fk"), double(count) / 1000ull);
+			Format1dp(str, double(count) / (1000ull), TC("k"));
 		else if (count < 1000ull * 1000 * 1000)
-			TSprintf_s(str, 32, TC("%.1fm"), double(count) / (1000ull * 1000));
+			Format1dp(str, double(count) / (1000ull * 1000), TC("m"));
 		else if (count < 1000ull * 1000 * 1000 * 1000)
-			TSprintf_s(str, 32, TC("%.1fg"), double(count) / (1000ull * 1000 * 1000));
+			Format1dp(str, double(count) / (1000ull * 1000 * 1000), TC("g"));
+		else
+			Format1dp(str, double(count) / (1000ull * 1000 * 1000 * 1000), TC("t"));
+	}
+
+	TimeToText::TimeToText(u64 time, bool allowMinutes) : TimeToText(time, allowMinutes, GetFrequency()) {}
+	TimeToText::TimeToText(u64 time, bool allowMinutes, u64 frequency)
+	{
+		u64 ms = TimeToMs(time, frequency);
+		if (ms == 0 && time != 0)
+			TStrcpy_s(str, 32, TC("<1ms"));
+		else if (ms < 1000)
+			Format(str, u32(ms), TC("ms"));
+		else if (ms < 60 * 1000 || !allowMinutes)
+			Format1dp(str, double(ms) / 1000, TC("s"));
 		else
-			TSprintf_s(str, 32, TC("%.1ft"), double(count) / (1000ull * 1000 * 1000 * 1000));
+		{
+			u32 totalSec = u32(float(ms) / 1000);
+			u32 totalMin = totalSec / 60;
+			u32 min = totalMin % 60;
+			u32 sec = totalSec % 60;
+			u32 hour = totalMin / 60;
+			u32 days = hour / 24;
+			hour -= days*24;
+			if (days)
+				TSprintf_s(str, 32, TC("%ud%uh%um"), (unsigned int)days, (unsigned int)hour, (unsigned int)min);
+			else if (hour)
+				TSprintf_s(str, 32, TC("%uh%um%us"), (unsigned int)hour, (unsigned int)min, (unsigned int)sec);
+			else
+				TSprintf_s(str, 32, TC("%um%us"), (unsigned int)min, (unsigned int)sec);
+		}
 	}
 
 #if UBA_DEBUG_LOGGER
diff --git a/Engine/Source/Programs/UnrealBuildAccelerator/Common/Public/UbaLogger.h b/Engine/Source/Programs/UnrealBuildAccelerator/Common/Public/UbaLogger.h
--- a/Engine/Source/Programs/UnrealBuildAccelerator/Common/Public/UbaLogger.h
+++ b/Engine/Source/Programs/UnrealBuildAccelerator/Common/Public/UbaLogger.h
@@ -11,17 +11,23 @@
 
 namespace uba
 {
+	////////////////////////////////////////////////////////////////////////////////////////////////////
+
 	class StringBufferBase;
 	struct BinaryReader;
 	struct StringView;
 	using TraverseThreadErrorFunc = Function<void(const StringView& error)>;
 
+	////////////////////////////////////////////////////////////////////////////////////////////////////
+
 	struct LogEntry
 	{
 		LogEntryType type;
 		const tchar* string;
 	};
 
+	////////////////////////////////////////////////////////////////////////////////////////////////////
+
 	class Logger
 	{
 	public:
@@ -43,6 +49,8 @@ namespace uba
 		bool ToFalse() { return false; }
 	};
 
+	////////////////////////////////////////////////////////////////////////////////////////////////////
+
 	class LoggerWithWriter : public Logger
 	{
 	public:
@@ -57,6 +65,8 @@ namespace uba
 		u32 m_prefixLen;
 	};
 
+	////////////////////////////////////////////////////////////////////////////////////////////////////
+
 	struct MutableLogger : public LoggerWithWriter
 	{
 		MutableLogger(LogWriter& writer, const tchar* prefix) : LoggerWithWriter(writer, prefix) {}
@@ -65,6 +75,8 @@ namespace uba
 		Atomic<u32> isMuted;
 	};
 
+	////////////////////////////////////////////////////////////////////////////////////////////////////
+
 	class FilteredLogWriter : public LogWriter
 	{
 	public:
@@ -77,13 +89,17 @@ namespace uba
 		LogEntryType m_level;
 	};
 
+	////////////////////////////////////////////////////////////////////////////////////////////////////
+
 	struct BytesToText
 	{
-		BytesToText(u64 bytes);
+		BytesToText(u64 bytes, bool base2 = false);
 		operator const tchar* () const { return str; };
 		tchar str[32];
 	};
 
+	////////////////////////////////////////////////////////////////////////////////////////////////////
+
 	struct CountToText
 	{
 		CountToText(u64 count);
@@ -91,6 +107,18 @@ namespace uba
 		tchar str[32];
 	};
 
+	////////////////////////////////////////////////////////////////////////////////////////////////////
+
+	struct TimeToText
+	{
+		TimeToText(float seconds, bool allowMinutes = false);
+		TimeToText(u64 time, bool allowMinutes = false);
+		TimeToText(u64 time, bool allowMinutes, u64 frequency);
+		operator const tchar*() const { return str; };
+		tchar str[32];
+	};
+
+	////////////////////////////////////////////////////////////////////////////////////////////////////
 
 	#if UBA_DEBUG_LOGGER
 	Logger* StartDebugLogger(Logger& outerLogger, const tchar* fileName);
@@ -104,6 +132,7 @@ namespace uba
 	void TraverseAllCallstacks(const Function<void(const CallstackInfo&)>& func, const TraverseThreadErrorFunc& errorFunc);
 	void PrintAllCallstacks(Logger& logger);
 
+	////////////////////////////////////////////////////////////////////////////////////////////////////
 
 	struct LogStallScope
 	{
@@ -121,4 +150,6 @@ namespace uba
 	#else
 	#define LOG_STALL_SCOPE(logger, timeSeconds, messageFormat)
 	#endif
+
+	////////////////////////////////////////////////////////////////////////////////////////////////////
 }
diff --git a/Engine/Source/Programs/UnrealBuildAccelerator/Core/Private/UbaTimer.cpp b/Engine/Source/Programs/UnrealBuildAccelerator/Core/Private/UbaTimer.cpp
--- a/Engine/Source/Programs/UnrealBuildAccelerator/Core/Private/UbaTimer.cpp
+++ b/Engine/Source/Programs/UnrealBuildAccelerator/Core/Private/UbaTimer.cpp
@@ -68,32 +68,4 @@ namespace uba
 		return u64(tv.tv_sec) * 1'000'000ull + u64(tv.tv_usec);
 		#endif
 	}
-
-	TimeToText::TimeToText(u64 time, bool allowMinutes) : TimeToText(time, allowMinutes, GetFrequency()) {}
-	TimeToText::TimeToText(u64 time, bool allowMinutes, u64 frequency)
-	{
-		u64 ms = TimeToMs(time, frequency);
-		if (ms == 0 && time != 0)
-			TStrcpy_s(str, 32, TC("<1ms"));
-		else if (ms < 1000)
-			TSprintf_s(str, 32, TC("%llums"), ms);
-		else if (ms < 60 * 1000 || !allowMinutes)
-			TSprintf_s(str, 32, TC("%.1fs"), float(ms) / 1000);
-		else
-		{
-			u32 totalSec = u32(float(ms) / 1000);
-			u32 totalMin = totalSec / 60;
-			u32 min = totalMin % 60;
-			u32 sec = totalSec % 60;
-			u32 hour = totalMin / 60;
-			u32 days = hour / 24;
-			hour -= days*24;
-			if (days)
-				TSprintf_s(str, 32, TC("%ud%uh%um"), (unsigned int)days, (unsigned int)hour, (unsigned int)min);
-			else if (hour)
-				TSprintf_s(str, 32, TC("%uh%um%us"), (unsigned int)hour, (unsigned int)min, (unsigned int)sec);
-			else
-				TSprintf_s(str, 32, TC("%um%us"), (unsigned int)min, (unsigned int)sec);
-		}
-	}
 }
\ No newline at end of file
diff --git a/Engine/Source/Programs/UnrealBuildAccelerator/Core/Public/UbaTimer.h b/Engine/Source/Programs/UnrealBuildAccelerator/Core/Public/UbaTimer.h
--- a/Engine/Source/Programs/UnrealBuildAccelerator/Core/Public/UbaTimer.h
+++ b/Engine/Source/Programs/UnrealBuildAccelerator/Core/Public/UbaTimer.h
@@ -6,6 +6,8 @@
 
 namespace uba
 {
+	////////////////////////////////////////////////////////////////////////////////////////////////////
+
 	u64 GetTime();							// Counter for time. Note this has no specific start time for 0
 	u64 GetFrequency();						// Number of counts per second. GetTime() / GetFrequency() = seconds
 	u64 GetSystemTimeUs();					// Get time in microseconds since 1 Jan 1970 00:00 UTC
@@ -24,7 +26,7 @@ namespace uba
 	inline u64 UsToTime(u64 us)						{ return us * GetFrequency() / 1'000'000; }
 	inline u64 MsToTime(u64 ms)						{ return ms * GetFrequency() / 1'000; }
 
-
+	////////////////////////////////////////////////////////////////////////////////////////////////////
 
 	struct Timer
 	{
@@ -39,11 +41,15 @@ namespace uba
 		bool operator==(const Timer& o) const { return time == o.time && count == o.count; }
 	};
 
+	////////////////////////////////////////////////////////////////////////////////////////////////////
+
 	struct ExtendedTimer : Timer
 	{
 		AtomicU64 longest;
 	};
 
+	////////////////////////////////////////////////////////////////////////////////////////////////////
+
 	struct TimerScope
 	{
 		TimerScope(Timer& t) : timer(t), start(GetTime()) { }
@@ -56,6 +62,8 @@ namespace uba
 		void operator=(const TimerScope&) = delete;
 	};
 
+	////////////////////////////////////////////////////////////////////////////////////////////////////
+
 	struct ExtendedTimerScope 
 	{
 		ExtendedTimerScope(ExtendedTimer& t) : timer(t), start(GetTime()) { ++timer.count; }
@@ -70,12 +78,5 @@ namespace uba
 		u64 start;
 	};
 
-	struct TimeToText
-	{
-		TimeToText(float seconds, bool allowMinutes = false);
-		TimeToText(u64 time, bool allowMinutes = false);
-		TimeToText(u64 time, bool allowMinutes, u64 frequency);
-		operator const tchar*() const { return str; };
-		tchar str[32];
-	};
+	////////////////////////////////////////////////////////////////////////////////////////////////////
 }
