[UBA]
* Added AllowWaitOnMem, AllowKillOnMem, MemWaitLoadPercent and MemKillLoadPercent as config options

diff --git a/Engine/Source/Programs/UnrealBuildAccelerator/Common/Private/UbaSessionServer.cpp b/Engine/Source/Programs/UnrealBuildAccelerator/Common/Private/UbaSessionServer.cpp
--- a/Engine/Source/Programs/UnrealBuildAccelerator/Common/Private/UbaSessionServer.cpp
+++ b/Engine/Source/Programs/UnrealBuildAccelerator/Common/Private/UbaSessionServer.cpp
@@ -118,6 +118,15 @@ namespace uba
 			table->GetValueAsBool(remoteTraceEnabled, TC("RemoteTraceEnabled"));
 			table->GetValueAsBool(nameToHashTableEnabled, TC("NameToHashTableEnabled"));
 			table->GetValueAsBool(traceIOEnabled, TC("TraceIOEnabled"));
+
+			table->GetValueAsBool(allowWaitOnMem, TC("AllowWaitOnMem"));
+			table->GetValueAsBool(allowKillOnMem, TC("AllowKillOnMem"));
+
+			u32 temp = 0;
+			if (table->GetValueAsU32(temp, TC("MemWaitLoadPercent")))
+				memWaitLoadPercent = u8(temp);
+			if (table->GetValueAsU32(temp, TC("MemKillLoadPercent")))
+				memKillLoadPercent = u8(temp);
 		}
 	}
 
