[UBA]
* Added cache client summary to trace stream
* Bumped trace stream version

diff --git a/Engine/Source/Programs/UnrealBuildAccelerator/Common/Private/UbaCacheClient.cpp b/Engine/Source/Programs/UnrealBuildAccelerator/Common/Private/UbaCacheClient.cpp
--- a/Engine/Source/Programs/UnrealBuildAccelerator/Common/Private/UbaCacheClient.cpp
+++ b/Engine/Source/Programs/UnrealBuildAccelerator/Common/Private/UbaCacheClient.cpp
@@ -1049,6 +1049,16 @@ namespace uba
 		m_connected = false;
 	}
 
+	void CacheClient::TraceSummary()
+	{
+		StackBinaryWriter<256*1024> writer;
+		m_session.GetTrace().WriteSummaryText(writer, [&](Logger& logger)
+			{
+				m_client.PrintSummary(logger);
+			});
+		m_session.GetTrace().CacheSummary(writer.GetData(), writer.GetPosition());
+	}
+
 	bool CacheClient::SendPathTable(Bucket& bucket, u32 requiredPathTableSize)
 	{
 		SCOPED_FUTEX(bucket.sendPathTableNetworkLock, lock);
diff --git a/Engine/Source/Programs/UnrealBuildAccelerator/Common/Private/UbaScheduler.cpp b/Engine/Source/Programs/UnrealBuildAccelerator/Common/Private/UbaScheduler.cpp
--- a/Engine/Source/Programs/UnrealBuildAccelerator/Common/Private/UbaScheduler.cpp
+++ b/Engine/Source/Programs/UnrealBuildAccelerator/Common/Private/UbaScheduler.cpp
@@ -177,6 +177,9 @@ namespace uba
 		m_processEntries.clear();
 		m_processEntriesStart = 0;
 		m_session.SetOuterScheduler(nullptr);
+
+		for (auto cacheClient : m_cacheClients)
+			cacheClient->TraceSummary();
 	}
 
 	void Scheduler::SetMaxLocalProcessors(u32 maxLocalProcessors)
@@ -530,7 +533,7 @@ namespace uba
 
 					CacheResult cacheResult;
 					bool isHit = false;
-					for (auto& cacheClient : m_cacheClients)
+					for (auto cacheClient : m_cacheClients)
 					{
 						if (!cacheClient->FetchFromCache(cacheResult, si.rootsHandle, si.cacheBucketId, si) || !cacheResult.hit)
 							continue;
diff --git a/Engine/Source/Programs/UnrealBuildAccelerator/Common/Private/UbaTrace.cpp b/Engine/Source/Programs/UnrealBuildAccelerator/Common/Private/UbaTrace.cpp
--- a/Engine/Source/Programs/UnrealBuildAccelerator/Common/Private/UbaTrace.cpp
+++ b/Engine/Source/Programs/UnrealBuildAccelerator/Common/Private/UbaTrace.cpp
@@ -537,6 +537,12 @@ namespace uba
 		writer.Write7BitEncoded(bytesSent);
 	}
 
+	void Trace::CacheSummary(const u8* data, u64 dataSize)
+	{
+		BEGIN_TRACE_ENTRY(CacheSummary);
+		writer.WriteBytes(data, dataSize);
+	}
+
 	TraceChannel::TraceChannel(Logger& logger) : m_logger(logger)
 	{
 	}
diff --git a/Engine/Source/Programs/UnrealBuildAccelerator/Common/Public/UbaCacheClient.h b/Engine/Source/Programs/UnrealBuildAccelerator/Common/Public/UbaCacheClient.h
--- a/Engine/Source/Programs/UnrealBuildAccelerator/Common/Public/UbaCacheClient.h
+++ b/Engine/Source/Programs/UnrealBuildAccelerator/Common/Public/UbaCacheClient.h
@@ -72,6 +72,8 @@ namespace uba
 
 		void Disable();
 
+		void TraceSummary();
+
 	private:
 		struct Bucket;
 		u64 MakeId(u32 bucketId);
diff --git a/Engine/Source/Programs/UnrealBuildAccelerator/Common/Public/UbaTrace.h b/Engine/Source/Programs/UnrealBuildAccelerator/Common/Public/UbaTrace.h
--- a/Engine/Source/Programs/UnrealBuildAccelerator/Common/Public/UbaTrace.h
+++ b/Engine/Source/Programs/UnrealBuildAccelerator/Common/Public/UbaTrace.h
@@ -61,6 +61,7 @@ namespace uba
 		UBA_TRACE_TYPE(ProcessBreadcrumbs) \
 		UBA_TRACE_TYPE(WorkHint) \
 		UBA_TRACE_TYPE(DriveUpdate) \
+		UBA_TRACE_TYPE(CacheSummary) \
 
 	enum TraceType : u8
 	{
@@ -69,7 +70,7 @@ namespace uba
 		#undef UBA_TRACE_TYPE
 	};
 
-	static constexpr u32 TraceVersion = 43;
+	static constexpr u32 TraceVersion = 44;
 	static constexpr u32 TraceReadCompatibilityVersion = 6;
 
 	class Trace : public WorkTracker
@@ -112,6 +113,7 @@ namespace uba
 		void CacheEndFetch(u32 fetchId, bool success, const u8* data, u64 dataSize);
 		void CacheBeginWrite(u32 processId);
 		void CacheEndWrite(u32 processId, bool success, u64 bytesSent);
+		void CacheSummary(const u8* data, u64 dataSize);
 
 		// Writes the current state of the trace to the specified output file without stopping anything
 		bool Write(const tchar* writeFileName, bool writeSummary);
diff --git a/Engine/Source/Programs/UnrealBuildAccelerator/Visualizer/Private/UbaTraceReader.cpp b/Engine/Source/Programs/UnrealBuildAccelerator/Visualizer/Private/UbaTraceReader.cpp
--- a/Engine/Source/Programs/UnrealBuildAccelerator/Visualizer/Private/UbaTraceReader.cpp
+++ b/Engine/Source/Programs/UnrealBuildAccelerator/Visualizer/Private/UbaTraceReader.cpp
@@ -35,6 +35,7 @@ namespace uba
 		version = 0;
 		sessions.clear();
 		workTracks.clear();
+		cacheSummaries.clear();
 		for (auto s : strings)
 			free((void*)s);
 		strings.clear();
@@ -1352,6 +1353,16 @@ namespace uba
 				record.entries.emplace_back(time, startTime, text); 
 			break;
 		}
+		case TraceType_CacheSummary:
+		{
+			auto& cacheSummary = out.cacheSummaries.emplace_back();
+
+			u32 lineCount = reader.ReadU32();
+			cacheSummary.lines.reserve(lineCount);
+			for (u32 i=0; i!=lineCount; ++i)
+				cacheSummary.lines.push_back(reader.ReadString());
+			break;
+		}
 		default:
 		{
 			return m_logger.Error(TC("Unknown trace type found in stream. UbaVisualizer too old?"));
diff --git a/Engine/Source/Programs/UnrealBuildAccelerator/Visualizer/Private/UbaTraceReader.h b/Engine/Source/Programs/UnrealBuildAccelerator/Visualizer/Private/UbaTraceReader.h
--- a/Engine/Source/Programs/UnrealBuildAccelerator/Visualizer/Private/UbaTraceReader.h
+++ b/Engine/Source/Programs/UnrealBuildAccelerator/Visualizer/Private/UbaTraceReader.h
@@ -167,7 +167,13 @@ namespace uba
 			u16 count;
 		};
 
+		struct CacheSummary
+		{
+			Vector<TString> lines;
+		};
+
 		Vector<Session> sessions;
+		Vector<CacheSummary> cacheSummaries;
 		Vector<WorkTrack> workTracks;
 		Vector<tchar*> strings;
 		Vector<ActiveProcessCount> activeProcessCounts;
diff --git a/Engine/Source/Programs/UnrealBuildAccelerator/Visualizer/Private/UbaVisualizer.cpp b/Engine/Source/Programs/UnrealBuildAccelerator/Visualizer/Private/UbaVisualizer.cpp
--- a/Engine/Source/Programs/UnrealBuildAccelerator/Visualizer/Private/UbaVisualizer.cpp
+++ b/Engine/Source/Programs/UnrealBuildAccelerator/Visualizer/Private/UbaVisualizer.cpp
@@ -2295,6 +2295,13 @@ namespace uba
 				logger.Log(LogEntryType_Info, line);
 			logger.AddTextOffset(0);
 
+			if (m_sessionSelectedIndex == 0)
+			{
+				for (auto& cacheSummary : m_traceView.cacheSummaries)
+					for (auto& line : cacheSummary.lines)
+						logger.Log(LogEntryType_Info, line);
+			}
+
 			for (auto& pair : session.drives)
 			{
 				auto& d = pair.second;
