[UBA]
* Added message priority support on network server but commented out because it doesn't seem to matter in my test cases

diff --git a/Engine/Source/Programs/UnrealBuildAccelerator/Common/Private/UbaNetworkServer.cpp b/Engine/Source/Programs/UnrealBuildAccelerator/Common/Private/UbaNetworkServer.cpp
--- a/Engine/Source/Programs/UnrealBuildAccelerator/Common/Private/UbaNetworkServer.cpp
+++ b/Engine/Source/Programs/UnrealBuildAccelerator/Common/Private/UbaNetworkServer.cpp
@@ -377,7 +377,8 @@ namespace uba
 
 			conn.m_client = &client;
 
-			if (client.connectionCount.fetch_add(1) == 0)
+			u32 oldCount = client.connectionCount.fetch_add(1);
+			if (oldCount == 0)
 			{
 				if (server.m_onConnectionFunction)
 					server.m_onConnectionFunction(clientUid, clientId);
@@ -388,6 +389,21 @@ namespace uba
 			{
 				if (server.m_logConnections)
 					server.m_logger.Detail(TC("Client %u (%s) additional connection %s connected"), clientId, GuidToString(clientUid).str, GuidToString(connectionUid).str);
+
+				#if 0 // Haven't really seen any difference by setting priority on socket. Will need more testing
+				if (m_useMessagePriority && oldCount == 1)
+				{
+					SCOPED_READ_LOCK(server.m_connectionsLock, lock);
+					for (auto& conIt : server.m_connections)
+					{
+						if (conIt.m_client != &client)
+							continue;
+						conIt.m_backend.SetPriority(conIt.m_backendConnection);
+						server.m_logger.Info(TC("Priority set on first connection for client"));
+						break;
+					}
+				}
+				#endif
 			}
 
 			return true;
@@ -764,6 +780,7 @@ namespace uba
 		m_receiveTimeoutMs = info.receiveTimeoutSeconds * 1000;
 		m_logConnections = info.logConnections;
 		m_useKeepAlive = info.useKeepAlive;
+		m_useMessagePriority = info.useMessagePriority;
 
 		#if PLATFORM_MAC
 		m_useKeepAlive = true; // Always run keep alive on mac since the built-in one has a probe interval of 1 minute or something like that... so timeout is always 10 minutes
diff --git a/Engine/Source/Programs/UnrealBuildAccelerator/Common/Public/UbaNetworkServer.h b/Engine/Source/Programs/UnrealBuildAccelerator/Common/Public/UbaNetworkServer.h
--- a/Engine/Source/Programs/UnrealBuildAccelerator/Common/Public/UbaNetworkServer.h
+++ b/Engine/Source/Programs/UnrealBuildAccelerator/Common/Public/UbaNetworkServer.h
@@ -39,6 +39,7 @@ namespace uba
 		u32 receiveTimeoutSeconds = DefaultNetworkReceiveTimeoutSeconds;
 		bool logConnections = true;
 		bool useKeepAlive = false;
+		bool useMessagePriority = false;
 	};
 
 	struct MessageInfo
@@ -175,6 +176,7 @@ namespace uba
 		bool m_workersEnabled = true;
 		bool m_logConnections = true;
 		bool m_useKeepAlive = false;
+		bool m_useMessagePriority = false;
 		bool m_allClientsDisconnected = false;
 
 		Futex m_addConnectionsLock;
