[UBA]
* Fixed shutdown crash when pressing ctrl-c while creating connections to cache server
* Added commented out code for setting socket priority. Will need more testing before enabled (using Qos on windows)

diff --git a/Engine/Source/Programs/UnrealBuildAccelerator/Common/Private/UbaNetworkClient.cpp b/Engine/Source/Programs/UnrealBuildAccelerator/Common/Private/UbaNetworkClient.cpp
--- a/Engine/Source/Programs/UnrealBuildAccelerator/Common/Private/UbaNetworkClient.cpp
+++ b/Engine/Source/Programs/UnrealBuildAccelerator/Common/Private/UbaNetworkClient.cpp
@@ -144,7 +144,7 @@ namespace uba
 				Guid cryptoIV;
 				
 				CryptoKey key = client.m_cryptoKey;
-				if (key !=InvalidCryptoKey)
+				if (key != InvalidCryptoKey)
 					cryptoIV = *(Guid*)(headerData+17);
 
 				if (serverUid == Guid())
@@ -271,6 +271,14 @@ namespace uba
 		connection->backend = &backend;
 		connection->cryptoIV = cryptoIV;
 
+		#if 0 // Haven't really seen any difference by setting priority on socket. Will need more testing
+		if (m_useMessagePriority && m_connections.size() > 3)
+		{
+			auto& firstConnection = m_connections.front();
+			firstConnection.backend->SetPriority(firstConnection.backendConnection);
+		}
+		#endif
+
 		SCOPED_FUTEX(m_connectionsItLock, l); // Take this lock to make sure callbacks are set before connection is used
 		m_connectionsIt = --m_connections.end();
 
@@ -395,17 +403,31 @@ namespace uba
 			});
 
 		{
-			SCOPED_READ_LOCK(m_connectionsLock, lock);
+			SCOPED_WRITE_LOCK(m_connectionsLock, connectionsLock);
 			if (m_isDisconnecting)
+			{
+				// We can't return straight away here if m_connections is not empty since it might actively disconnect on another thread
+				while (!m_connections.empty())
+				{
+					connectionsLock.Leave();
+					Sleep(10);
+					connectionsLock.Enter();
+				}
 				return;
+			}
+
 			m_isDisconnecting = true;
 			for (auto& c : m_connections)
-			{
-				OnDisconnected(c, 0);
-				c.disconnectedEvent.IsSet(~0u);
-			}
+				if (c.connected)
+					c.backend->Shutdown(c.backendConnection);
 		}
 
+		for (auto& c : m_connections)
+			c.disconnectedEvent.IsSet(~0u);
+
+		if (m_connectionCount != 0)
+			m_logger.Error(TC("Should be zero connections here. %llu still active"), m_connectionCount.load());
+
 		{
 			SCOPED_WRITE_LOCK(m_connectionsLock, lock2);
 			for (auto& c : m_connections)
