[UBA]
* Added code that check tcp environment and print out tips if settings are not good. can be disabled with -notips

diff --git a/Engine/Source/Programs/UnrealBuildAccelerator/Agent/Private/UbaAgent.cpp b/Engine/Source/Programs/UnrealBuildAccelerator/Agent/Private/UbaAgent.cpp
--- a/Engine/Source/Programs/UnrealBuildAccelerator/Agent/Private/UbaAgent.cpp
+++ b/Engine/Source/Programs/UnrealBuildAccelerator/Agent/Private/UbaAgent.cpp
@@ -108,6 +108,7 @@ namespace uba
 		logger.Info(TC("  -noproxy                Does not allow this agent to be a storage proxy for other agents"));
 		logger.Info(TC("  -proxyuselocalstorage   Storage proxy will use local storage to see if files exist"));
 		logger.Info(TC("  -nocloud                Will not try to connect to cloud meta data server (this can take time during first startup)"));
+		logger.Info(TC("  -notips                 Disable printing of tips at startup"));
 		logger.Info(TC("  -killrandom             Kills random process and exit session"));
 		logger.Info(TC("  -memwait=<percent>      The amount of memory needed to spawn a process. Set this to 100 to disable. Defaults to 80%%"));
 		logger.Info(TC("  -memkill=<percent>      The amount of memory needed before processes starts to be killed. Set this to 100 to disable. Defaults to 90%%"));
@@ -367,6 +368,7 @@ namespace uba
 		bool allowProxy = true;
 		bool proxyUseLocalStorage = false;
 		bool couldBeCloud = true;
+		bool printTips = true;
 		bool useStorage = true;
 		bool resetStore = false;
 		bool quiet = false;
@@ -709,6 +711,10 @@ namespace uba
 			{
 				couldBeCloud = false;
 			}
+			else if (name.Equals(TCV("-notips")))
+			{
+				printTips = false;
+			}
 			else if (name.Equals(TCV("-command")))
 			{
 				if (!parseValue(command))
@@ -840,6 +846,9 @@ namespace uba
 			if (IsRunningArm())
 				logger.Warning(TC("  Running x64 binary on arm64 system. Use arm binaries instead"));
 
+		if (!NetworkBackendTcp::CheckEnvironment(logger, printTips))
+			return false;
+
 		logger.Info(TC(""));
 
 
diff --git a/Engine/Source/Programs/UnrealBuildAccelerator/CacheService/Private/UbaCacheService.cpp b/Engine/Source/Programs/UnrealBuildAccelerator/CacheService/Private/UbaCacheService.cpp
--- a/Engine/Source/Programs/UnrealBuildAccelerator/CacheService/Private/UbaCacheService.cpp
+++ b/Engine/Source/Programs/UnrealBuildAccelerator/CacheService/Private/UbaCacheService.cpp
@@ -64,6 +64,7 @@ namespace uba
 		logger.Info(TC("  -nomaintenance          Skip all maintenance"));
 		logger.Info(TC("  -crash                  Force a crash (for testing)"));
 		logger.Info(TC("  -nosignalhandler        Will not hook up signal handler"));
+		logger.Info(TC("  -notips                 Disable printing of tips at startup"));
 		logger.Info(TC("  -maxworkers=<number>    Max number of workers used by cache server. Defaults to \"%u\""), DefaultWorkerCount);
 		logger.Info(TC("  -reportinterval=<sec>   How often the service should report status. Defaults to \"%s\""), TimeToText(MsToTime(DefaultReportIntervalSeconds*1000)).str);
 		logger.Info(TC("  -trace                  Create a trace file that is written out when cache service exits"));
@@ -132,6 +133,7 @@ namespace uba
 		bool shouldCrash = false;
 		bool requiresCrypto = false;
 		bool writeTrace = false;
+		bool printTips = true;
 		u32 expirationTimeSeconds = DefaultExpiration;
 		u32 maxWorkerCount = DefaultWorkerCount;
 		u32 reportIntervalSeconds = DefaultReportIntervalSeconds;
@@ -263,6 +265,10 @@ namespace uba
 			{
 				maintenanceEnabled = false;
 			}
+			else if (name.Equals(TCV("-notips")))
+			{
+				printTips = false;
+			}
 			else if (name.Equals(TCV("-nosave")))
 			{
 				allowSave = false;
@@ -313,10 +319,13 @@ namespace uba
 		{
 			u64 hugePageCount = GetHugePageCount();
 			u64 recommendedHugePageCount = (maintenanceReserveSizeMb*GetLogicalProcessorCount())/2;
-			if (hugePageCount < recommendedHugePageCount)
-				logger.Info(TC("  Improve maintenance performance by enabling %llu huge pages on system (%llu enabled)"), recommendedHugePageCount, hugePageCount);
+			if (hugePageCount < recommendedHugePageCount && printTips)
+				logger.Info(TC("  TIP: Improve maintenance performance by enabling %llu huge pages on system (%llu enabled)"), recommendedHugePageCount, hugePageCount);
 		}
 
+		if (!NetworkBackendTcp::CheckEnvironment(logger, printTips))
+			return false;
+
 		StringBuffer<512> currentDir;
 		GetCurrentDirectoryW(currentDir);
 
@@ -490,6 +499,7 @@ namespace uba
 		nsci.workerCount = maxWorkerCount;
 		nsci.logConnections = false; // Let the cache server report instead
 		nsci.receiveTimeoutSeconds = 2*60*60; // Two hours timeout
+		nsci.useMessagePriority = true;
 		bool ctorSuccess = true;
 		NetworkServer networkServer(ctorSuccess, nsci);
 		if (!ctorSuccess)
diff --git a/Engine/Source/Programs/UnrealBuildAccelerator/Common/Private/UbaNetworkBackendTcp.cpp b/Engine/Source/Programs/UnrealBuildAccelerator/Common/Private/UbaNetworkBackendTcp.cpp
--- a/Engine/Source/Programs/UnrealBuildAccelerator/Common/Private/UbaNetworkBackendTcp.cpp
+++ b/Engine/Source/Programs/UnrealBuildAccelerator/Common/Private/UbaNetworkBackendTcp.cpp
@@ -18,6 +18,7 @@
 #pragma comment(lib, "IPHLPAPI.lib") // For GetAdaptersInfo
 #pragma comment(lib, "qwave.lib")
 #else
+#include "UbaLinuxNetworkWrappers.h"
 #include <netdb.h>
 #include <netinet/tcp.h>
 #include <arpa/inet.h>
@@ -164,37 +165,69 @@ namespace uba
 
 	////////////////////////////////////////////////////////////////////////////////////////////////////
 
-	bool NetworkBackendTcp::EnsureInitialized(Logger& logger)
+	#if PLATFORM_WINDOWS
+	HMODULE g_wineDll;
+	#endif
+
+	bool (*GetCongestionAlgorithm)(SOCKET s, char* out, int outCapacity);
+	bool (*SetCongestionAlgorithm)(SOCKET s, char* value, int valuelen);
+
+	////////////////////////////////////////////////////////////////////////////////////////////////////
+
+	bool InitNetwork(Logger& logger)
 	{
-		#if PLATFORM_WINDOWS
+#if PLATFORM_WINDOWS
 		WSADATA wsaData;
-		if (!m_wsaInitDone)
-		{
-			if (int res = WSAStartup(MAKEWORD(2, 2), &wsaData))
-				return logger.Error(TC("WSAStartup failed (%d)"), res);
+		if (int res = WSAStartup(MAKEWORD(2, 2), &wsaData))
+			return logger.Error(TC("WSAStartup failed (%d)"), res);
 
-			#if UBA_USE_IOCP
-			if (m_iocpWorkerCount)
+		if (IsRunningWine())
+		{
+			g_wineDll = LoadLibrary(TC("UbaWine.dll.so"));
+			if (g_wineDll)
 			{
-				m_iocpHandle = CreateIoCompletionPort(INVALID_HANDLE_VALUE, NULL, 0, 0);
-
-				for (u32 i=0; i!=m_iocpWorkerCount; ++i)
-					m_iocpThreads.emplace_back().Start([this]() { ThreadIocp(); return 0; }, TC("UbaIocp"));
+				GetCongestionAlgorithm = (decltype(GetCongestionAlgorithm))GetProcAddress(g_wineDll, "GetCongestionAlgorithm");
+				SetCongestionAlgorithm = (decltype(SetCongestionAlgorithm))GetProcAddress(g_wineDll, "SetCongestionAlgorithm");
 			}
-			#endif
-
-			m_wsaInitDone = true;
 		}
-		#else
-		static bool initOnce = []()
-			{
-				struct sigaction sa = { { SIG_IGN } };
-				sigaction(SIGPIPE, &sa, NULL); // Needed for broken pipe that can happen if helpers crash
-				return true;
-			}();
+#else
+		struct sigaction sa = { { SIG_IGN } };
+		sigaction(SIGPIPE, &sa, NULL); // Needed for broken pipe that can happen if helpers crash
+		GetCongestionAlgorithm = [](SOCKET s, char* out, int outCapacity) { socklen_t len = outCapacity; return getsockopt(s, IPPROTO_TCP, TCP_CONGESTION, out, &len) == 0; };
+		SetCongestionAlgorithm = [](SOCKET s, char* v, int l) { return setsockopt(s, IPPROTO_TCP, TCP_CONGESTION, v, l) == 0; };
+#endif
+		return true;
+	}
+
+	bool InitNetworkOnce(Logger& logger)
+	{
+		static bool initOnce = [](Logger& logger) { return InitNetwork(logger); }(logger);
+		return initOnce;
+	}
+
+	bool NetworkBackendTcp::EnsureInitialized(Logger& logger)
+	{
+		if (!InitNetworkOnce(logger))
+			return false;
+
+		#if PLATFORM_WINDOWS
+		SCOPED_FUTEX(m_initLock, lock);
+		if (m_initDone)
+			return m_initSuccess;
+		m_initDone = true;
 
+		#if UBA_USE_IOCP
+		if (m_iocpWorkerCount)
+		{
+			m_iocpHandle = CreateIoCompletionPort(INVALID_HANDLE_VALUE, NULL, 0, 0);
+
+			for (u32 i=0; i!=m_iocpWorkerCount; ++i)
+				m_iocpThreads.emplace_back().Start([this]() { ThreadIocp(); return 0; }, TC("UbaIocp"));
+		}
 		#endif
 
+		m_initSuccess = true;
+		#endif
 		return true;
 	}
 
@@ -250,11 +283,6 @@ namespace uba
 
 		m_tcpStatusLoop.Set();
 		m_tcpStatusThread.Wait();
-
-		#if PLATFORM_WINDOWS
-		if (m_wsaInitDone)
-			WSACleanup();
-		#endif
 	}
 
 	void NetworkBackendTcp::Shutdown(void* connection)
@@ -1354,11 +1382,87 @@ namespace uba
 				SOCKET s = conn.socket;
 				if (s == INVALID_SOCKET)
 					continue;
-				LogTcpInfo(m_logger, s);
+				//LogTcpInfo(m_logger, s);
+
+				#if 1
+				socklen_t optlen = sizeof(int);
+				int sendBufSize = 0;
+				int recvBufSize = 0;
+				int prio = 0;
+				char alg[16] = { 0 };
+
+				getsockopt(s, SOL_SOCKET, SO_SNDBUF, (char*)&sendBufSize, &optlen);
+				getsockopt(s, SOL_SOCKET, SO_RCVBUF, (char*)&recvBufSize, &optlen);
+
+				if (GetCongestionAlgorithm)
+					GetCongestionAlgorithm(s, alg, sizeof(alg));
+				#if PLATFORM_LINUX
+				getsockopt(s, SOL_SOCKET, SO_PRIORITY, (char*)&prio, &optlen);
+				#endif
+
+				m_logger.Info(TC("%i - Send: %s, Recv: %s, Prio %i, Alg: %hs"), (int)s, BytesToText(sendBufSize).str, BytesToText(recvBufSize).str, prio, alg);
+				#endif
+
 			}
 		}
 	}
 
+	bool NetworkBackendTcp::CheckEnvironment(Logger& logger, bool printTips)
+	{
+		if (!InitNetworkOnce(logger))
+			return false;
+
+		if (!printTips)
+			return true;
+
+		bool (*GetTcpAutoTuning)(int* outReadMin, int* outReadDefault, int* outReadMax, int* outWriteMin, int* outWriteDefault, int* outWriteMax) = nullptr;
+
+		#if PLATFORM_WINDOWS
+		if (!GetTcpAutoTuning && IsRunningWine())
+			GetTcpAutoTuning = (decltype(GetTcpAutoTuning))GetProcAddress(g_wineDll, "GetTcpAutoTuning");
+		#else
+		GetTcpAutoTuning = UnixGetTcpAutoTuning;
+		#endif
+
+		if (GetCongestionAlgorithm)
+		{
+			SOCKET tempSocket = ::socket(AF_INET, SOCK_STREAM, IPPROTO_TCP);
+			char alg[16] = { 0 };
+			GetCongestionAlgorithm(tempSocket, alg, sizeof(alg));
+			::closesocket(tempSocket);
+
+			if (strncmp(alg, "bbr", 3) != 0)
+				logger.Info(TC("  TIP: Enable BBR congestion algorithm and fq (Fair Queue) queuing discipline (Current: %hs)"), alg);
+		}
+
+		if (GetTcpAutoTuning)
+		{
+			int readMin, readDefault, readMax, writeMin, writeDefault, writeMax;
+			if (GetTcpAutoTuning(&readMin, &readDefault, &readMax, &writeMin, &writeDefault, &writeMax))
+			{
+				if (readMax < 1024*1024)
+					logger.Info(TC("  TIP: Increase tcp window auto tuning read max to 4-16 MiB (Current: %s %s %s)"), BytesToText(readMin, true).str, BytesToText(readDefault, true).str, BytesToText(readMax, true).str);
+				if (writeMax < 1024*1024)
+					logger.Info(TC("  TIP: Increase tcp window auto tuning write max to 4-16 MiB (Current: %s %s %s)"), BytesToText(writeMin, true).str, BytesToText(writeDefault, true).str, BytesToText(writeMax, true).str);
+			}
+		}
+
+		// sudo modprobe tcp_bbr
+		// sudo sysctl -w net.core.default_qdisc=fq
+		// sudo sysctl -w net.ipv4.tcp_congestion_control=bbr
+		// sudo sysctl -w net.ipv4.tcp_rmem="4096 131072 16777216"
+		// sudo sysctl -w net.ipv4.tcp_wmem="4096 65536 33554432\"
+
+		// Persist
+		// 
+		// echo "tcp_bbr" | sudo tee /etc/modules-load.d/bbr.conf
+		// echo "net.core.default_qdisc=fq" | sudo tee -a /etc/sysctl.conf
+		// echo "net.ipv4.tcp_congestion_control=bbr" | sudo tee -a /etc/sysctl.conf
+		// sudo sysctl -p
+
+		return true;
+	}
+
 	void NetworkBackendTcp::PrintTcpStatistics(Logger& logger, u32 statusUpdateSeconds, void* prevStatsPtr)
 	{
 		#if PLATFORM_WINDOWS
@@ -1501,22 +1605,12 @@ namespace uba
 			LoggerWithWriter logger(g_nullLogWriter);
 			CloseSocket(logger, m_socket, TC("HttpDtor"));
 		}
-
-		#if PLATFORM_WINDOWS
-		if (m_wsaInitDone)
-			WSACleanup();
-		#endif
 	}
 
 	bool HttpConnection::Connect(Logger& logger, const char* host)
 	{
-		#if PLATFORM_WINDOWS
-		WSADATA wsaData;
-		if (!m_wsaInitDone)
-			if (int res = WSAStartup(MAKEWORD(2, 2), &wsaData))
-				return logger.Error(TC("WSAStartup failed (%d)"), res);
-		m_wsaInitDone = true;
-		#endif
+		if (!InitNetworkOnce(logger))
+			return false;
 
 		hostent* hostent = gethostbyname(host);
 		if (hostent == NULL)
@@ -1726,14 +1820,22 @@ namespace uba
 
 	bool SetRecvBuf(Logger& logger, SOCKET socket, u32 windowSize)
 	{
-		if (setsockopt(socket, SOL_SOCKET, SO_RCVBUF, (const char*)&windowSize, sizeof(windowSize)) == SOCKET_ERROR)
+		int optname = SO_RCVBUF;
+		#if PLATFORM_LINUX
+		optname = SO_RCVBUFFORCE;
+		#endif
+		if (setsockopt(socket, SOL_SOCKET, optname, (const char*)&windowSize, sizeof(windowSize)) == SOCKET_ERROR)
 			return logger.Error(TC("setsockopt TCP_NODELAY error: (error: %s)"), LastErrorToText(WSAGetLastError()).data);
 		return true;
 	}
 
 	bool SetSendBuf(Logger& logger, SOCKET socket, u32 windowSize)
 	{
-		if (setsockopt(socket, SOL_SOCKET, SO_SNDBUF, (const char*)&windowSize, sizeof(windowSize)) == SOCKET_ERROR)
+		int optname = SO_SNDBUF;
+		#if PLATFORM_LINUX
+		optname = SO_SNDBUFFORCE;
+		#endif
+		if (setsockopt(socket, SOL_SOCKET, optname, (const char*)&windowSize, sizeof(windowSize)) == SOCKET_ERROR)
 			return logger.Error(TC("setsockopt TCP_NODELAY error: (error: %s)"), LastErrorToText(WSAGetLastError()).data);
 		return true;
 	}
@@ -1745,7 +1847,7 @@ namespace uba
 			{
 				HANDLE h = INVALID_HANDLE_VALUE;
 				QOS_VERSION ver = {1, 0};
-				if (!QOSCreateHandle(&ver, &qosHandle))
+				if (!QOSCreateHandle(&ver, &h))
 					logger.Error(TC("QOSCreateHandle failed"));
 				return h;
 			}(logger);
@@ -1908,16 +2010,6 @@ namespace uba
 	void LogTcpInfo(Logger& logger, SOCKET socket)
 	{
 	#if PLATFORM_WINDOWS
-
-		#if 0
-		socklen_t optlen = sizeof(int);
-		int sendBufSize = 0;
-		getsockopt(socket, SOL_SOCKET, SO_SNDBUF, (char *)&sendBufSize, &optlen);
-		int recvBufSize = 0;
-		getsockopt(socket, SOL_SOCKET, SO_RCVBUF, (char *)&recvBufSize, &optlen);
-		logger.Info(TC("%i - Send: %s, Recv: %s"), (int)socket, BytesToText(sendBufSize).str, BytesToText(recvBufSize).str);
-		#endif
-
 		#if 0
 		if (IsRunningWine())
 		{
diff --git a/Engine/Source/Programs/UnrealBuildAccelerator/Common/Public/UbaLinuxNetworkWrappers.h b/Engine/Source/Programs/UnrealBuildAccelerator/Common/Public/UbaLinuxNetworkWrappers.h
new file mode 100644
--- /dev/null
+++ b/Engine/Source/Programs/UnrealBuildAccelerator/Common/Public/UbaLinuxNetworkWrappers.h
@@ -0,0 +1,46 @@
+// Copyright Epic Games, Inc. All Rights Reserved.
+
+#pragma once
+
+#include <netinet/tcp.h>
+#include <netinet/in.h>
+#include <errno.h>
+#include <stdbool.h>
+
+bool CatFile(char* outBuffer, int bufferCapacity, const char* path)
+{
+	FILE *f = fopen(path, "r");
+	if (!f)
+		return false;
+	char buffer[256];
+	if (!fgets(outBuffer, bufferCapacity, f))
+		return false;
+	fclose(f);
+	return true;
+}
+
+bool UnixGetTcpAutoTuning(int* outReadMin, int* outReadDefault, int* outReadMax, int* outWriteMin, int* outWriteDefault, int* outWriteMax)
+{
+	char buffer[256];
+	if (!CatFile(buffer, sizeof(buffer), "/proc/sys/net/core/rmem_max"))
+		return false;
+	int readMax = atoi(buffer);
+
+	if (!CatFile(buffer, sizeof(buffer), "/proc/sys/net/core/wmem_max"))
+		return false;
+	int writeMax = atoi(buffer);
+
+	if (!CatFile(buffer, sizeof(buffer), "/proc/sys/net/ipv4/tcp_rmem"))
+		return false;
+	sscanf(buffer, "%i %i %i", outReadMin, outReadDefault, outReadMax);
+	if (readMax < *outReadMax)
+		*outReadMax = readMax;
+
+	if (!CatFile(buffer, sizeof(buffer), "/proc/sys/net/ipv4/tcp_wmem"))
+		return false;
+	sscanf(buffer, "%i %i %i", outWriteMin, outWriteDefault, outWriteMax);
+	if (writeMax < *outWriteMax)
+		*outWriteMax = writeMax;
+
+	return true;
+}
diff --git a/Engine/Source/Programs/UnrealBuildAccelerator/Common/Public/UbaNetworkBackendTcp.h b/Engine/Source/Programs/UnrealBuildAccelerator/Common/Public/UbaNetworkBackendTcp.h
--- a/Engine/Source/Programs/UnrealBuildAccelerator/Common/Public/UbaNetworkBackendTcp.h
+++ b/Engine/Source/Programs/UnrealBuildAccelerator/Common/Public/UbaNetworkBackendTcp.h
@@ -52,6 +52,8 @@ namespace uba
 		virtual void GetTotalSendAndRecv(u64& outSend, u64& outRecv) override;
 		virtual void Validate(Logger& logger, const Vector<void*>& connections, bool full) override;
 
+		static bool CheckEnvironment(Logger& logger, bool printTips);
+
 	private:
 		bool EnsureInitialized(Logger& logger);
 
@@ -82,7 +84,9 @@ namespace uba
 		bool m_disableNagle;
 
 		#if PLATFORM_WINDOWS
-		bool m_wsaInitDone = false;
+		Futex m_initLock;
+		bool m_initDone = false;
+		bool m_initSuccess = false;
 		bool m_useOverlappedSend;
 		u16 m_iocpWorkerCount;
 		#endif
