[UBA]
* Fixed macos compile errors

diff --git a/Engine/Source/Programs/UnrealBuildAccelerator/Common/Private/UbaNetworkBackendTcp.cpp b/Engine/Source/Programs/UnrealBuildAccelerator/Common/Private/UbaNetworkBackendTcp.cpp
--- a/Engine/Source/Programs/UnrealBuildAccelerator/Common/Private/UbaNetworkBackendTcp.cpp
+++ b/Engine/Source/Programs/UnrealBuildAccelerator/Common/Private/UbaNetworkBackendTcp.cpp
@@ -190,7 +190,7 @@ namespace uba
 				SetCongestionAlgorithm = (decltype(SetCongestionAlgorithm))GetProcAddress(g_wineDll, "SetCongestionAlgorithm");
 			}
 		}
-#else
+#elif PLATFORM_LINUX
 		struct sigaction sa = { { SIG_IGN } };
 		sigaction(SIGPIPE, &sa, NULL); // Needed for broken pipe that can happen if helpers crash
 		GetCongestionAlgorithm = [](SOCKET s, char* out, int outCapacity) { socklen_t len = outCapacity; return getsockopt(s, IPPROTO_TCP, TCP_CONGESTION, out, &len) == 0; };
@@ -1856,16 +1856,15 @@ namespace uba
 		QOS_FLOWID flowId = 0;
 		if (!QOSAddSocketToFlow(s_qosHandle, socket, NULL, QOSTrafficTypeExcellentEffort, QOS_NON_ADAPTIVE_FLOW, &flowId))
 			return logger.Error(TC("QOSAddSocketToFlow failed for socket %i (%s)"), socket, LastErrorToText().data);
-		return true;
-#else
+#elif PLATFORM_LINUX
 		int prio = 6;
 		if (setsockopt(socket, SOL_SOCKET, SO_PRIORITY, &prio, sizeof(prio)) == -1)
 			return logger.Error(TC("setsockopt SO_PRIORITY failed (%s)"), LastErrorToText(WSAGetLastError()).data);
 		//uint8_t tos = 46 << 2; // EF = 0b101110xx
 		//if (setsockopt(socket, IPPROTO_IP, IP_TOS, &tos, sizeof(tos)) == -1)
 		//	return logger.Error(TC("setsockopt IP_TOS failed (%s)"), LastErrorToText(WSAGetLastError()).data);
-		return false;
 #endif
+		return true;
 	}
 
 	bool SetKeepAlive(Logger& logger, SOCKET socket) // This will make sure that WSAPoll exits when the network cable is pulled
