[UBA]
* Added so sendfile run in parallel to utilize all sockets

diff --git a/Engine/Source/Programs/UnrealBuildAccelerator/Common/Private/UbaCacheClient.cpp b/Engine/Source/Programs/UnrealBuildAccelerator/Common/Private/UbaCacheClient.cpp
--- a/Engine/Source/Programs/UnrealBuildAccelerator/Common/Private/UbaCacheClient.cpp
+++ b/Engine/Source/Programs/UnrealBuildAccelerator/Common/Private/UbaCacheClient.cpp
@@ -1195,7 +1195,7 @@ namespace uba
 					fileSize = file.GetSize();
 				}
 
-				if (!SendFile(logger, m_client, casKey, fileData, fileSize, path.data))
+				if (!SendFile(logger, m_client, &m_storage.m_workManager, casKey, fileData, fileSize, path.data))
 					return false;
 
 				outBytesSent += fileSize;
diff --git a/Engine/Source/Programs/UnrealBuildAccelerator/Common/Private/UbaStorageClient.cpp b/Engine/Source/Programs/UnrealBuildAccelerator/Common/Private/UbaStorageClient.cpp
--- a/Engine/Source/Programs/UnrealBuildAccelerator/Common/Private/UbaStorageClient.cpp
+++ b/Engine/Source/Programs/UnrealBuildAccelerator/Common/Private/UbaStorageClient.cpp
@@ -999,7 +999,7 @@ namespace uba
 			{
 				auto& stats = Stats();
 				TimerScope ts(stats.sendCas);
-				if (!SendFile(m_logger, m_client, casKey, fileMem, fileSize, hint))
+				if (!SendFile(m_logger, m_client, &m_workManager, casKey, fileMem, fileSize, hint))
 					return false;
 				stats.sendCasBytesRaw += fileSize;
 				stats.sendCasBytesComp += fileSize;
diff --git a/Engine/Source/Programs/UnrealBuildAccelerator/Common/Private/UbaStorageUtils.cpp b/Engine/Source/Programs/UnrealBuildAccelerator/Common/Private/UbaStorageUtils.cpp
--- a/Engine/Source/Programs/UnrealBuildAccelerator/Common/Private/UbaStorageUtils.cpp
+++ b/Engine/Source/Programs/UnrealBuildAccelerator/Common/Private/UbaStorageUtils.cpp
@@ -260,16 +260,32 @@ namespace uba
 		return success;
 	}
 
-	bool SendFile(Logger& logger, NetworkClient& client, const CasKey& casKey, const u8* sourceMem, u64 sourceSize, const tchar* hint)
+	UBA_NOINLINE bool SendFileFirstMessage(NetworkClient& client, const CasKey& casKey, const u8* sourceData, u64 sourceSize, const tchar* hint, StackBinaryReader<128>& reader, u64& outToWrite)
+	{
+		StackBinaryWriter<SendMaxSize> writer;
+		NetworkMessage msg(client, StorageServiceId, StorageMessageType_StoreBegin, writer);
+		writer.WriteCasKey(casKey);
+		writer.WriteU64(sourceSize);
+		writer.WriteU64(sourceSize);
+		writer.WriteString(hint);
+
+		u64 toWrite = Min(sourceSize, writer.GetCapacityLeft());
+		writer.WriteBytes(sourceData, toWrite);
+
+		outToWrite = toWrite;
+
+		return msg.Send(reader);
+	}
+
+	bool SendFile(Logger& logger, NetworkClient& client, WorkManager* workManager, const CasKey& casKey, const u8* sourceMem, u64 sourceSize, const tchar* hint)
 	{
 		UBA_ASSERT(casKey != CasKeyZero);
 
 		const u8* readData = sourceMem;
-		u64 fileSize = sourceSize;
 
 		u16 storeId = 0;
 		bool sendEnd = false;
-		u64 sendLeft = fileSize;
+		u64 sendLeft = sourceSize;
 		u64 sendPos = 0;
 
 		auto sendEndMessage = [&]()
@@ -282,21 +298,11 @@ namespace uba
 				return msg.Send();
 			};
 
-		StackBinaryWriter<SendMaxSize> writer;
-
 		if (sendLeft)
 		{
-			NetworkMessage msg(client, StorageServiceId, StorageMessageType_StoreBegin, writer);
-			writer.WriteCasKey(casKey);
-			writer.WriteU64(fileSize);
-			writer.WriteU64(sourceSize);
-			writer.WriteString(hint);
-
-			u64 toWrite = Min(sendLeft, writer.GetCapacityLeft());
-			writer.WriteBytes(readData, toWrite);
-
 			StackBinaryReader<128> reader;
-			if (!msg.Send(reader))
+			u64 toWrite = 0;
+			if (!SendFileFirstMessage(client, casKey, readData, sendLeft, hint, reader, toWrite))
 				return false;
 
 			readData += toWrite;
@@ -319,21 +325,77 @@ namespace uba
 			}
 		}
 
-		while (sendLeft)
+		u64 messageHeader = client.GetMessageHeaderSize();
+		u32 messageCapacity = u32(client.GetMessageMaxSize() - (messageHeader + sizeof(u16) + sizeof(u64)));
+		u32 sendCount = u32((sendLeft + messageCapacity - 1)/ messageCapacity);
+
+		struct Container
 		{
-			NetworkMessage msg(client, StorageServiceId, StorageMessageType_StoreSegment, writer.Reset());
-			writer.WriteU16(storeId);
-			writer.WriteU64(sendPos);
+			struct iterator
+			{
+				iterator() : index(0) {}
+				iterator(u32 i) : index(i) {}
+				iterator operator++(int) { return iterator(index++); }
+				bool operator==(const iterator& o) const { return index == o.index; }
+				u32 index;
+			};
+			Container(u32 c) : count(c) {}
+			iterator begin() const { return iterator(0); }
+			iterator end() const { return iterator(count); }
+			u32 size() { return count; }
+			u32 count;
+		};
 
-			u64 toWrite = Min(sendLeft, writer.GetCapacityLeft());
-			writer.WriteBytes(readData, toWrite); // We always need a separate buffer on the send since we do in-place encryption of the buffers
+		u32 connectionCount = client.GetConnectionCount();
+		if (workManager && connectionCount > 1) // Use parallel send.. to make sure all sockets send buffers are fully utilized
+		{
+			Atomic<bool> error;
+			Container container(sendCount);
 
-			if (!msg.Send())
+			workManager->ParallelFor(connectionCount - 1, container, [&](const WorkContext&, auto& it)
+				{
+					u64 sendPos2 = sendPos + it.index * messageCapacity;
+					StackBinaryWriter<SendMaxSize> writer;
+					NetworkMessage msg(client, StorageServiceId, StorageMessageType_StoreSegment, writer);
+					writer.WriteU16(storeId);
+					writer.WriteU64(sendPos2);
+
+					UBA_ASSERT(messageCapacity == writer.GetCapacityLeft());
+
+					const u8* writeData = readData + it.index * messageCapacity;
+					u64 toWrite = messageCapacity;
+					if (it.index == sendCount - 1)
+						toWrite = sendLeft % messageCapacity;
+
+					writer.WriteBytes(writeData, toWrite); // We always need a separate buffer on the send since we do in-place encryption of the buffers
+
+					if (!msg.Send())
+						error = true;
+
+				}, TCV("SendFile"));
+
+			if (error)
 				return false;
+		}
+		else
+		{
+			while (sendLeft)
+			{
+				StackBinaryWriter<SendMaxSize> writer;
+				NetworkMessage msg(client, StorageServiceId, StorageMessageType_StoreSegment, writer);
+				writer.WriteU16(storeId);
+				writer.WriteU64(sendPos);
 
-			readData += toWrite;
-			sendLeft -= toWrite;
-			sendPos += toWrite;
+				u64 toWrite = Min(sendLeft, writer.GetCapacityLeft());
+				writer.WriteBytes(readData, toWrite); // We always need a separate buffer on the send since we do in-place encryption of the buffers
+
+				if (!msg.Send())
+					return false;
+
+				readData += toWrite;
+				sendLeft -= toWrite;
+				sendPos += toWrite;
+			}
 		}
 
 		return sendEndMessage();
diff --git a/Engine/Source/Programs/UnrealBuildAccelerator/Common/Public/UbaStorageUtils.h b/Engine/Source/Programs/UnrealBuildAccelerator/Common/Public/UbaStorageUtils.h
--- a/Engine/Source/Programs/UnrealBuildAccelerator/Common/Public/UbaStorageUtils.h
+++ b/Engine/Source/Programs/UnrealBuildAccelerator/Common/Public/UbaStorageUtils.h
@@ -24,7 +24,7 @@ namespace uba
 
 	CasKey CalculateCasKey(u8* fileMem, u64 fileSize, bool storeCompressed, WorkManager* workManager, const tchar* hint);
 	bool SendBatchMessages(Logger& logger, NetworkClient& client, u16 fetchId, u8* slot, u64 capacity, u64 left, u32 messageMaxSize, u32& readIndex, u32& responseSize, const Function<bool()>& runInWaitFunc = {}, const tchar* hint = TC(""), u32* outError = nullptr, const Function<bool(u64)>& progressFunc = {});
-	bool SendFile(Logger& logger, NetworkClient& client, const CasKey& casKey, const u8* sourceMem, u64 sourceSize, const tchar* hint);
+	bool SendFile(Logger& logger, NetworkClient& client, WorkManager* workManager, const CasKey& casKey, const u8* sourceMem, u64 sourceSize, const tchar* hint);
 
 	////////////////////////////////////////////////////////////////////////////////////////////////////
 
