[UBA]
* Optimized sending path and cas table.. still opportunities to do better.. but for now.
* Optimized sending files to cache.. it will now use the work manager and send on multiple threads when possible
* Added lots of send-to-cache stats to trace file

diff --git a/Engine/Source/Programs/UnrealBuildAccelerator/Common/Private/UbaCacheClient.cpp b/Engine/Source/Programs/UnrealBuildAccelerator/Common/Private/UbaCacheClient.cpp
--- a/Engine/Source/Programs/UnrealBuildAccelerator/Common/Private/UbaCacheClient.cpp
+++ b/Engine/Source/Programs/UnrealBuildAccelerator/Common/Private/UbaCacheClient.cpp
@@ -82,10 +82,12 @@ namespace uba
 		CompactPathTable sendPathTable;
 		Futex sendPathTableNetworkLock;
 		u32 pathTableSizeSent = 0;
+		Atomic<u32> pathTableSizeMaxToSend;
 
 		CompactCasKeyTable sendCasKeyTable;
 		Futex sendCasKeyTableNetworkLock;
 		u32 casKeyTableSizeSent = 0;
+		Atomic<u32> casKeyTableSizeMaxToSend;
 	};
 
 	CacheClient::CacheClient(const CacheClientCreateInfo& info)
@@ -193,11 +195,19 @@ namespace uba
 			return false;
 		}
 
+		CacheSendStats stats;
+
 		bool finished = false;
-		u64 bytesSent = 0;
 		if (processId)
 			m_session.GetTrace().CacheBeginWrite(processId);
-		auto tg = MakeGuard([&]() { if (processId) m_session.GetTrace().CacheEndWrite(processId, finished, bytesSent); });
+		auto tg = MakeGuard([&]()
+			{
+				if (!processId)
+					return;
+				StackBinaryWriter<128> writer;
+				stats.Write(writer);
+				m_session.GetTrace().CacheEndWrite(processId, finished, writer.GetData(), writer.GetPosition());
+			});
 
 		BinaryReader inputsReader(inputs, 0, inputsSize);
 		BinaryReader outputsReader(outputs, 0, outputsSize);
@@ -221,6 +231,8 @@ namespace uba
 
 		UBA_TRACE_WRITE_HINT("TraverseInputsOutputs");
 
+		TimerScope buildTs(stats.build);
+
 		// Traverse all inputs and outputs. to create cache entry that we can send to server
 		while (true)
 		{
@@ -421,23 +433,25 @@ namespace uba
 		if (!success)
 			return false;
 
+		buildTs.Leave();
+
 		if (outputsStringToCasKey.empty())
 			m_logger.Warning(TC("NO OUTPUTS FROM process %s"), info.GetDescription()); 
 
 		UBA_TRACE_WRITE_HINT("SendPathTable");
 
 		// Make sure server has enough of the path table to be able to resolve offsets from cache entry
-		if (!SendPathTable(bucket, requiredPathTableSize))
+		if (!SendPathTable(bucket, stats, requiredPathTableSize))
 			return false;
 
 		UBA_TRACE_WRITE_HINT("SendCasTable");
 
 		// Make sure server has enough of the cas table to be able to resolve offsets from cache entry
-		if (!SendCasTable(bucket, requiredCasTableSize))
+		if (!SendCasTable(bucket, stats, requiredCasTableSize))
 			return false;
 
 		// actual cache entry now when we know server has the needed tables
-		if (!SendCacheEntry(tws, bucket, rootPaths, cmdKey, inputsStringToCasKey, outputsStringToCasKey, logLines, logLinesSize, bytesSent))
+		if (!SendCacheEntry(tws, bucket, stats, rootPaths, cmdKey, inputsStringToCasKey, outputsStringToCasKey, logLines, logLinesSize))
 			return false;
 
 
@@ -493,7 +507,7 @@ namespace uba
 		TrackWorkScope tws;
 		#endif
 
-		CacheStats cacheStats;
+		CacheFetchStats cacheStats;
 		StorageStats storageStats;
 		KernelStats kernelStats;
 		auto kg = MakeGuard([&]() { KernelStats::GetGlobal().Add(kernelStats); m_storage.AddStats(storageStats); });
@@ -847,7 +861,7 @@ namespace uba
 		return false;
 	}
 
-	bool CacheClient::FetchFile(Bucket& bucket, const RootPaths& rootPaths, const ProcessStartInfo& info, CacheStats& cacheStats, StorageStats& storageStats, u32 casKeyOffset)
+	bool CacheClient::FetchFile(Bucket& bucket, const RootPaths& rootPaths, const ProcessStartInfo& info, CacheFetchStats& cacheStats, StorageStats& storageStats, u32 casKeyOffset)
 	{
 		StringBuffer<MaxPath> path;
 		CasKey casKey;
@@ -1059,15 +1073,32 @@ namespace uba
 		m_session.GetTrace().CacheSummary(writer.GetData(), writer.GetPosition());
 	}
 
-	bool CacheClient::SendPathTable(Bucket& bucket, u32 requiredPathTableSize)
+	template<typename T>
+	void AtomicMax(Atomic<T>& target, T value)
 	{
+		T current = target.load(std::memory_order_relaxed);
+		while (current < value && !target.compare_exchange_weak(current, value, std::memory_order_relaxed, std::memory_order_relaxed))
+		{
+			// current is updated by compare_exchange_weak
+		}
+	}
+
+	bool CacheClient::SendPathTable(Bucket& bucket, CacheSendStats& stats, u32 requiredPathTableSize)
+	{
+		AtomicMax(bucket.pathTableSizeMaxToSend, requiredPathTableSize);
+
+		TimerScope waitTs(stats.sendTableWait);
 		SCOPED_FUTEX(bucket.sendPathTableNetworkLock, lock);
 		if (requiredPathTableSize <= bucket.pathTableSizeSent)
 			return true;
+		waitTs.Leave();
+
+		requiredPathTableSize = bucket.pathTableSizeMaxToSend;
 
 		u32 left = requiredPathTableSize - bucket.pathTableSizeSent;
 		while (left)
 		{
+			TimerScope ts(stats.sendPathTable);
 			StackBinaryWriter<SendMaxSize> writer;
 			NetworkMessage msg(m_client, CacheServiceId, CacheMessageType_StorePathTable, writer);
 			writer.Write7BitEncoded(MakeId(bucket.id));
@@ -1084,15 +1115,22 @@ namespace uba
 		return true;
 	}
 
-	bool CacheClient::SendCasTable(Bucket& bucket, u32 requiredCasTableSize)
+	bool CacheClient::SendCasTable(Bucket& bucket, CacheSendStats& stats, u32 requiredCasTableSize)
 	{
+		AtomicMax(bucket.casKeyTableSizeMaxToSend, requiredCasTableSize);
+
+		TimerScope waitTs(stats.sendTableWait);
 		SCOPED_FUTEX(bucket.sendCasKeyTableNetworkLock, lock);
 		if (requiredCasTableSize <= bucket.casKeyTableSizeSent)
 			return true;
+		waitTs.Leave();
+
+		requiredCasTableSize = bucket.casKeyTableSizeMaxToSend; // Send as much as possible
 
 		u32 left = requiredCasTableSize - bucket.casKeyTableSizeSent;
 		while (left)
 		{
+			TimerScope ts(stats.sendCasTable);
 			StackBinaryWriter<SendMaxSize> writer;
 			NetworkMessage msg(m_client, CacheServiceId, CacheMessageType_StoreCasTable, writer);
 			writer.Write7BitEncoded(MakeId(bucket.id));
@@ -1109,13 +1147,13 @@ namespace uba
 		return true;
 	}
 
-	bool CacheClient::SendCacheEntry(TrackWorkScope& tws, Bucket& bucket, const RootPaths& rootPaths, const CasKey& cmdKey, const Map<u32, u32>& inputsStringToCasKey, const Map<u32, u32>& outputsStringToCasKey, const u8* logLines, u64 logLinesSize, u64& outBytesSent)
+	bool CacheClient::SendCacheEntry(TrackWorkScope& tws, Bucket& bucket, CacheSendStats& stats, const RootPaths& rootPaths, const CasKey& cmdKey, const Map<u32, u32>& inputsStringToCasKey, const Map<u32, u32>& outputsStringToCasKey, const u8* logLines, u64 logLinesSize)
 	{
 		StackBinaryReader<1024> reader;
 
 		UBA_TRACE_WRITE_HINT("SendCacheEntryMessage");
 
-		if (!SendCacheEntryMessage(reader, bucket, cmdKey, inputsStringToCasKey, outputsStringToCasKey, logLines, logLinesSize))
+		if (!SendCacheEntryMessage(reader, bucket, stats, cmdKey, inputsStringToCasKey, outputsStringToCasKey, logLines, logLinesSize))
 			return false;
 
 		// Server has all content for caskeys.. upload is done
@@ -1125,6 +1163,8 @@ namespace uba
 		bool success = false;
 		auto doneGuard = MakeGuard([&]()
 			{
+				TimerScope ts(stats.sendEntry);
+
 				// Send done.. confirm to server
 				UBA_TRACE_WRITE_HINT("SendCacheDone");
 				StackBinaryWriter<64> writer;
@@ -1154,10 +1194,16 @@ namespace uba
 				casKey = AsExecutable(casKey, false); // Remove exec bit when storing in casdb
 			#endif
 
+			TimerScope createCasTs(stats.createCas);
 			StorageImpl::CasEntry* casEntry;
-			if (m_storage.HasCasFile(casKey, &casEntry))
+			bool hasCasFile = m_storage.HasCasFile(casKey, &casEntry);
+
+			if (hasCasFile)
 			{
+				createCasTs.Leave();
+
 				UBA_TRACE_WRITE_HINT("SendFile");
+				TimerScope ts(stats.sendFile);
 
 				UBA_ASSERT(!IsNormalized(casKey));
 				StringBuffer<> casKeyFileName;
@@ -1198,11 +1244,14 @@ namespace uba
 				if (!SendFile(logger, m_client, &m_storage.m_workManager, casKey, fileData, fileSize, path.data))
 					return false;
 
-				outBytesSent += fileSize;
+				stats.sendBytesComp += fileSize;
 			}
 			else // If we don't have the cas key it should be one of the normalized files.... otherwise there is a bug
 			{
+				createCasTs.Cancel();
+
 				UBA_TRACE_WRITE_HINT("SendNormalizedFile");
+				TimerScope ts(stats.sendNormalizedFile);
 
 				if (!IsNormalized(casKey))
 					return m_logger.Error(TC("Can't find output file %s to send to cache server"), path.data);
@@ -1251,7 +1300,7 @@ namespace uba
 				if (!sender.SendFileCompressed(casKey, path.data, dataToSend, sizeToSend, path.data))
 					return m_logger.Warning(TC("Failed to send cas content for file %s"), path.data);
 
-				outBytesSent += sender.m_bytesSent;
+				stats.sendBytesComp += sender.m_bytesSent;
 			}
 
 		}
@@ -1260,8 +1309,10 @@ namespace uba
 		return doneGuard.Execute();
 	}
 
-	bool CacheClient::SendCacheEntryMessage(BinaryReader& out, Bucket& bucket, const CasKey& cmdKey, const Map<u32, u32>& inputsStringToCasKey, const Map<u32, u32>& outputsStringToCasKey, const u8* logLines, u64 logLinesSize)
+	bool CacheClient::SendCacheEntryMessage(BinaryReader& out, Bucket& bucket, CacheSendStats& stats, const CasKey& cmdKey, const Map<u32, u32>& inputsStringToCasKey, const Map<u32, u32>& outputsStringToCasKey, const u8* logLines, u64 logLinesSize)
 	{
+		TimerScope ts(stats.sendEntry);
+
 		StackBinaryWriter<SendMaxSize> writer;
 
 		NetworkMessage msg(m_client, CacheServiceId, CacheMessageType_StoreEntry, writer);
@@ -1291,14 +1342,14 @@ namespace uba
 		return false;
 	}
 
-	bool CacheClient::FetchCasTable(TrackWorkScope& tws, Bucket& bucket, CacheStats& stats, u32 requiredCasTableOffset)
+	bool CacheClient::FetchCasTable(TrackWorkScope& tws, Bucket& bucket, CacheFetchStats& stats, u32 requiredCasTableOffset)
 	{
 		if (HasEnoughData(bucket, requiredCasTableOffset))
 			return true;
 		return FetchCasTable2(tws, bucket, stats, requiredCasTableOffset);
 	}
 
-	bool CacheClient::FetchCasTable2(TrackWorkScope& tws, Bucket& bucket, CacheStats& stats, u32 requiredCasTableOffset)
+	bool CacheClient::FetchCasTable2(TrackWorkScope& tws, Bucket& bucket, CacheFetchStats& stats, u32 requiredCasTableOffset)
 	{
 		UBA_TRACE_FETCH_HINT_SCOPE("FetchCasTable")
 
diff --git a/Engine/Source/Programs/UnrealBuildAccelerator/Common/Private/UbaStats.cpp b/Engine/Source/Programs/UnrealBuildAccelerator/Common/Private/UbaStats.cpp
--- a/Engine/Source/Programs/UnrealBuildAccelerator/Common/Private/UbaStats.cpp
+++ b/Engine/Source/Programs/UnrealBuildAccelerator/Common/Private/UbaStats.cpp
@@ -416,33 +416,64 @@ namespace uba
 		logger.Info(TC(""));
 	}
 
-	void CacheStats::Write(BinaryWriter& writer)
+	void CacheFetchStats::Write(BinaryWriter& writer)
 	{
 		#define UBA_CACHE_STAT(type, var, ver) uba::Write(writer, var);
-		UBA_CACHE_STATS
+		UBA_CACHE_FETCH_STATS
 		#undef UBA_CACHE_STAT
 	}
 
-	void CacheStats::Read(BinaryReader& reader, u32 version)
+	void CacheFetchStats::Read(BinaryReader& reader, u32 version)
 	{
 		#define UBA_CACHE_STAT(type, var, ver) if (ver <= version) uba::Read(reader, var, version);
-		UBA_CACHE_STATS
+		UBA_CACHE_FETCH_STATS
 		#undef UBA_CACHE_STAT
 	}
 
-	void CacheStats::Print(Logger& logger, u64 frequency)
+	void CacheFetchStats::Print(Logger& logger, u64 frequency)
 	{
 		#define UBA_CACHE_STAT(type, var, ver) LogStat(logger, #var, var, frequency);
-		UBA_CACHE_STATS
+		UBA_CACHE_FETCH_STATS
 		#undef UBA_CACHE_STAT
 		if (fetchBytesComp)
 			logger.Info(TC("   Bytes   Raw/Comp %9s %9s"), BytesToText(fetchBytesRaw).str, BytesToText(fetchBytesComp).str);
 	}
 
-	bool CacheStats::IsEmpty()
+	bool CacheFetchStats::IsEmpty()
 	{
 		#define UBA_CACHE_STAT(type, var, ver) if (var != type()) return false;
-		UBA_CACHE_STATS
+		UBA_CACHE_FETCH_STATS
+		#undef UBA_CACHE_STAT
+		return true;
+	}
+
+	void CacheSendStats::Write(BinaryWriter& writer)
+	{
+		#define UBA_CACHE_STAT(type, var, ver) uba::Write(writer, var);
+		UBA_CACHE_SEND_STATS
+		#undef UBA_CACHE_STAT
+	}
+
+	void CacheSendStats::Read(BinaryReader& reader, u32 version)
+	{
+		#define UBA_CACHE_STAT(type, var, ver) if (ver <= version) uba::Read(reader, var, version);
+		UBA_CACHE_SEND_STATS
+		#undef UBA_CACHE_STAT
+	}
+
+	void CacheSendStats::Print(Logger& logger, u64 frequency)
+	{
+		#define UBA_CACHE_STAT(type, var, ver) LogStat(logger, #var, var, frequency);
+		UBA_CACHE_SEND_STATS
+		#undef UBA_CACHE_STAT
+		if (sendBytesComp)
+			logger.Info(TC("   Bytes   Raw/Comp %9s %9s"), BytesToText(sendBytesRaw).str, BytesToText(sendBytesComp).str);
+	}
+
+	bool CacheSendStats::IsEmpty()
+	{
+		#define UBA_CACHE_STAT(type, var, ver) if (var != type()) return false;
+		UBA_CACHE_SEND_STATS
 		#undef UBA_CACHE_STAT
 		return true;
 	}
diff --git a/Engine/Source/Programs/UnrealBuildAccelerator/Common/Private/UbaTrace.cpp b/Engine/Source/Programs/UnrealBuildAccelerator/Common/Private/UbaTrace.cpp
--- a/Engine/Source/Programs/UnrealBuildAccelerator/Common/Private/UbaTrace.cpp
+++ b/Engine/Source/Programs/UnrealBuildAccelerator/Common/Private/UbaTrace.cpp
@@ -529,12 +529,12 @@ namespace uba
 		writer.Write7BitEncoded(processId);
 	}
 
-	void Trace::CacheEndWrite(u32 processId, bool success, u64 bytesSent)
+	void Trace::CacheEndWrite(u32 processId, bool success, const u8* data, u64 dataSize)
 	{
 		BEGIN_TRACE_ENTRY(CacheEndWrite);
 		writer.Write7BitEncoded(processId);
 		writer.WriteBool(success);
-		writer.Write7BitEncoded(bytesSent);
+		writer.WriteBytes(data, dataSize);
 	}
 
 	void Trace::CacheSummary(const u8* data, u64 dataSize)
diff --git a/Engine/Source/Programs/UnrealBuildAccelerator/Common/Public/UbaCacheClient.h b/Engine/Source/Programs/UnrealBuildAccelerator/Common/Public/UbaCacheClient.h
--- a/Engine/Source/Programs/UnrealBuildAccelerator/Common/Public/UbaCacheClient.h
+++ b/Engine/Source/Programs/UnrealBuildAccelerator/Common/Public/UbaCacheClient.h
@@ -15,7 +15,8 @@ namespace uba
 	class RootPaths;
 	class Session;
 	class StorageImpl;
-	struct CacheStats;
+	struct CacheFetchStats;
+	struct CacheSendStats;
 	struct CasKey;
 	struct ProcessStartInfo;
 	struct StorageStats;
@@ -79,13 +80,13 @@ namespace uba
 		u64 MakeId(u32 bucketId);
 
 		// No inlining because they use lots of stack space
-		UBA_NOINLINE bool SendPathTable(Bucket& bucket, u32 requiredPathTableSize);
-		UBA_NOINLINE bool SendCasTable(Bucket& bucket, u32 requiredCasTableSize);
-		UBA_NOINLINE bool SendCacheEntry(TrackWorkScope& tws, Bucket& bucket, const RootPaths& rootPaths, const CasKey& cmdKey, const Map<u32, u32>& inputsStringToCasKey, const Map<u32, u32>& outputsStringToCasKey, const u8* logLines, u64 logLinesSize, u64& outBytesSent);
-		UBA_NOINLINE bool SendCacheEntryMessage(BinaryReader& out, Bucket& bucket, const CasKey& cmdKey, const Map<u32, u32>& inputsStringToCasKey, const Map<u32, u32>& outputsStringToCasKey, const u8* logLines, u64 logLinesSize);
-		bool FetchCasTable(TrackWorkScope& tws, Bucket& bucket, CacheStats& stats, u32 requiredCasTableOffset);
-		UBA_NOINLINE bool FetchCasTable2(TrackWorkScope& tws, Bucket& bucket, CacheStats& stats, u32 requiredCasTableOffset);
-		UBA_NOINLINE bool FetchFile(Bucket& bucket, const RootPaths& rootPaths, const ProcessStartInfo& info, CacheStats& cacheStats, StorageStats& storageStats, u32 casKeyOffset);
+		UBA_NOINLINE bool SendPathTable(Bucket& bucket, CacheSendStats& stats, u32 requiredPathTableSize);
+		UBA_NOINLINE bool SendCasTable(Bucket& bucket, CacheSendStats& stats, u32 requiredCasTableSize);
+		UBA_NOINLINE bool SendCacheEntry(TrackWorkScope& tws, Bucket& bucket, CacheSendStats& stats, const RootPaths& rootPaths, const CasKey& cmdKey, const Map<u32, u32>& inputsStringToCasKey, const Map<u32, u32>& outputsStringToCasKey, const u8* logLines, u64 logLinesSize);
+		UBA_NOINLINE bool SendCacheEntryMessage(BinaryReader& out, Bucket& bucket, CacheSendStats& stats, const CasKey& cmdKey, const Map<u32, u32>& inputsStringToCasKey, const Map<u32, u32>& outputsStringToCasKey, const u8* logLines, u64 logLinesSize);
+		bool FetchCasTable(TrackWorkScope& tws, Bucket& bucket, CacheFetchStats& stats, u32 requiredCasTableOffset);
+		UBA_NOINLINE bool FetchCasTable2(TrackWorkScope& tws, Bucket& bucket, CacheFetchStats& stats, u32 requiredCasTableOffset);
+		UBA_NOINLINE bool FetchFile(Bucket& bucket, const RootPaths& rootPaths, const ProcessStartInfo& info, CacheFetchStats& cacheStats, StorageStats& storageStats, u32 casKeyOffset);
 		template<typename TableType> bool FetchCompactTable(u32 bucketId, TableType& table, u32 requiredTableOffset, u8 messageType);
 		bool HasEnoughData(Bucket& bucket, u32 requiredCasTableOffset);
 		bool HasEnoughCasData(Bucket& bucket, u32 requiredCasTableOffset, u32& outPathOffset);
diff --git a/Engine/Source/Programs/UnrealBuildAccelerator/Common/Public/UbaStats.h b/Engine/Source/Programs/UnrealBuildAccelerator/Common/Public/UbaStats.h
--- a/Engine/Source/Programs/UnrealBuildAccelerator/Common/Public/UbaStats.h
+++ b/Engine/Source/Programs/UnrealBuildAccelerator/Common/Public/UbaStats.h
@@ -132,7 +132,7 @@ namespace uba
 		void Print(Logger& logger, u64 frequency = GetFrequency());
 	};
 
-	#define UBA_CACHE_STATS \
+	#define UBA_CACHE_FETCH_STATS \
 		UBA_CACHE_STAT(Timer, fetchEntries, 0) \
 		UBA_CACHE_STAT(Timer, fetchCasTable, 0) \
 		UBA_CACHE_STAT(Timer, normalizeFile, 28) \
@@ -141,10 +141,34 @@ namespace uba
 		UBA_CACHE_STAT(AtomicU64, fetchBytesRaw, 26) \
 		UBA_CACHE_STAT(AtomicU64, fetchBytesComp, 26) \
 
-    struct CacheStats
+    struct CacheFetchStats
 	{
 		#define UBA_CACHE_STAT(type, var, ver) type var;
-		UBA_CACHE_STATS
+		UBA_CACHE_FETCH_STATS
+		#undef UBA_CACHE_STAT
+
+		void Write(BinaryWriter& writer);
+		void Read(BinaryReader& reader, u32 version);
+		void Print(Logger& logger, u64 frequency = GetFrequency());
+		bool IsEmpty();
+	};
+
+	#define UBA_CACHE_SEND_STATS \
+		UBA_CACHE_STAT(Timer, build, 0) \
+		UBA_CACHE_STAT(Timer, sendPathTable, 0) \
+		UBA_CACHE_STAT(Timer, sendCasTable, 0) \
+		UBA_CACHE_STAT(Timer, sendTableWait, 0) \
+		UBA_CACHE_STAT(Timer, sendEntry, 0) \
+		UBA_CACHE_STAT(Timer, createCas, 0) \
+		UBA_CACHE_STAT(Timer, sendFile, 0) \
+		UBA_CACHE_STAT(Timer, sendNormalizedFile, 0) \
+		UBA_CACHE_STAT(AtomicU64, sendBytesRaw, 0) \
+		UBA_CACHE_STAT(AtomicU64, sendBytesComp, 0) \
+
+    struct CacheSendStats
+	{
+		#define UBA_CACHE_STAT(type, var, ver) type var;
+		UBA_CACHE_SEND_STATS
 		#undef UBA_CACHE_STAT
 
 		void Write(BinaryWriter& writer);
diff --git a/Engine/Source/Programs/UnrealBuildAccelerator/Common/Public/UbaTrace.h b/Engine/Source/Programs/UnrealBuildAccelerator/Common/Public/UbaTrace.h
--- a/Engine/Source/Programs/UnrealBuildAccelerator/Common/Public/UbaTrace.h
+++ b/Engine/Source/Programs/UnrealBuildAccelerator/Common/Public/UbaTrace.h
@@ -70,7 +70,7 @@ namespace uba
 		#undef UBA_TRACE_TYPE
 	};
 
-	static constexpr u32 TraceVersion = 44;
+	static constexpr u32 TraceVersion = 45;
 	static constexpr u32 TraceReadCompatibilityVersion = 6;
 
 	class Trace : public WorkTracker
@@ -112,7 +112,7 @@ namespace uba
 		void CacheBeginFetch(u32 fetchId, const tchar* description);
 		void CacheEndFetch(u32 fetchId, bool success, const u8* data, u64 dataSize);
 		void CacheBeginWrite(u32 processId);
-		void CacheEndWrite(u32 processId, bool success, u64 bytesSent);
+		void CacheEndWrite(u32 processId, bool success, const u8* data, u64 dataSize);
 		void CacheSummary(const u8* data, u64 dataSize);
 
 		// Writes the current state of the trace to the specified output file without stopping anything
diff --git a/Engine/Source/Programs/UnrealBuildAccelerator/Visualizer/Private/UbaTraceReader.cpp b/Engine/Source/Programs/UnrealBuildAccelerator/Visualizer/Private/UbaTraceReader.cpp
--- a/Engine/Source/Programs/UnrealBuildAccelerator/Visualizer/Private/UbaTraceReader.cpp
+++ b/Engine/Source/Programs/UnrealBuildAccelerator/Visualizer/Private/UbaTraceReader.cpp
@@ -1283,7 +1283,7 @@ namespace uba
 			if (!success)
 				process.returnedReason = TC("M");
 
-			CacheStats cacheStats;
+			CacheFetchStats cacheStats;
 			KernelStats kernelStats;
 			StorageStats storageStats;
 			const u8* dataStart = reader.GetPositionData();
@@ -1309,7 +1309,18 @@ namespace uba
 			u32 processId = u32(reader.Read7BitEncoded());
 			TraceView::CacheWrite& write = out.cacheWrites[processId];
 			write.success = reader.ReadBool();
-			write.bytesSent = reader.Read7BitEncoded();
+
+			CacheSendStats stats;
+			const u8* dataStart = reader.GetPositionData();
+			if (out.version < 45)
+				stats.sendBytesComp = reader.Read7BitEncoded();
+			else
+				stats.Read(reader, out.version);
+
+			const u8* dataEnd = reader.GetPositionData();
+			write.stats.resize(dataEnd - dataStart);
+			memcpy(write.stats.data(), dataStart, dataEnd - dataStart);
+
 			write.end = time;
 			break;
 		}
diff --git a/Engine/Source/Programs/UnrealBuildAccelerator/Visualizer/Private/UbaTraceReader.h b/Engine/Source/Programs/UnrealBuildAccelerator/Visualizer/Private/UbaTraceReader.h
--- a/Engine/Source/Programs/UnrealBuildAccelerator/Visualizer/Private/UbaTraceReader.h
+++ b/Engine/Source/Programs/UnrealBuildAccelerator/Visualizer/Private/UbaTraceReader.h
@@ -153,7 +153,7 @@ namespace uba
 		{
 			u64 start = 0;
 			u64 end = 0;
-			u64 bytesSent = 0;
+			Vector<u8> stats;
 			bool success = false;
 		};
 
diff --git a/Engine/Source/Programs/UnrealBuildAccelerator/Visualizer/Private/UbaVisualizer.cpp b/Engine/Source/Programs/UnrealBuildAccelerator/Visualizer/Private/UbaVisualizer.cpp
--- a/Engine/Source/Programs/UnrealBuildAccelerator/Visualizer/Private/UbaVisualizer.cpp
+++ b/Engine/Source/Programs/UnrealBuildAccelerator/Visualizer/Private/UbaVisualizer.cpp
@@ -801,7 +801,11 @@ namespace uba
 		{
 			logger.Info(L"  Duration                    %9s", TimeToText(write.end - write.start).str);
 			logger.Info(L"  Success                     %9s", write.success ? L"true" : L"false");
-			logger.Info(L"  Bytes sent                  %9s", BytesToText(write.bytesSent).str);
+
+			BinaryReader reader(write.stats.data(), 0, write.stats.size());
+			CacheSendStats stats;
+			stats.Read(reader, m_traceView.version);
+			stats.Print(logger, m_traceView.frequency);
 		}
 		else
 		{
@@ -2139,7 +2143,7 @@ namespace uba
 				SessionStats sessionStats;
 				StorageStats storageStats;
 				KernelStats kernelStats;
-				CacheStats cacheStats;
+				CacheFetchStats cacheStats;
 
 				if (process.cacheFetch)
 				{
