[UBA]
* Made code more stack size friendly when writing to cache

diff --git a/Engine/Source/Programs/UnrealBuildAccelerator/Common/Private/UbaStorageUtils.cpp b/Engine/Source/Programs/UnrealBuildAccelerator/Common/Private/UbaStorageUtils.cpp
--- a/Engine/Source/Programs/UnrealBuildAccelerator/Common/Private/UbaStorageUtils.cpp
+++ b/Engine/Source/Programs/UnrealBuildAccelerator/Common/Private/UbaStorageUtils.cpp
@@ -260,7 +260,7 @@ namespace uba
 		return success;
 	}
 
-	UBA_NOINLINE bool SendFileFirstMessage(NetworkClient& client, const CasKey& casKey, const u8* sourceData, u64 sourceSize, const tchar* hint, StackBinaryReader<128>& reader, u64& outToWrite)
+	UBA_NOINLINE bool SendFileBeginMessage(NetworkClient& client, const CasKey& casKey, const u8* sourceData, u64 sourceSize, const tchar* hint, StackBinaryReader<128>& reader, u64& outToWrite)
 	{
 		StackBinaryWriter<SendMaxSize> writer;
 		NetworkMessage msg(client, StorageServiceId, StorageMessageType_StoreBegin, writer);
@@ -277,6 +277,16 @@ namespace uba
 		return msg.Send(reader);
 	}
 
+	UBA_NOINLINE bool SendFileSegmentMessage(NetworkClient& client, u16 storeId, u64 sendPos, const u8* data, u64 dataSize)
+	{
+		StackBinaryWriter<SendMaxSize> writer;
+		NetworkMessage msg(client, StorageServiceId, StorageMessageType_StoreSegment, writer);
+		writer.WriteU16(storeId);
+		writer.WriteU64(sendPos);
+		writer.WriteBytes(data, dataSize);
+		return msg.Send();
+	}
+
 	bool SendFile(Logger& logger, NetworkClient& client, WorkManager* workManager, const CasKey& casKey, const u8* sourceMem, u64 sourceSize, const tchar* hint)
 	{
 		UBA_ASSERT(casKey != CasKeyZero);
@@ -292,7 +302,7 @@ namespace uba
 			{
 				if (!sendEnd)
 					return true;
-				StackBinaryWriter<128> writer;
+				StackBinaryWriter<32> writer;
 				NetworkMessage msg(client, StorageServiceId, StorageMessageType_StoreEnd, writer);
 				writer.WriteCasKey(casKey);
 				return msg.Send();
@@ -302,7 +312,7 @@ namespace uba
 		{
 			StackBinaryReader<128> reader;
 			u64 toWrite = 0;
-			if (!SendFileFirstMessage(client, casKey, readData, sendLeft, hint, reader, toWrite))
+			if (!SendFileBeginMessage(client, casKey, readData, sendLeft, hint, reader, toWrite))
 				return false;
 
 			readData += toWrite;
@@ -326,52 +336,40 @@ namespace uba
 		}
 
 		u64 messageHeader = client.GetMessageHeaderSize();
-		u32 messageCapacity = u32(client.GetMessageMaxSize() - (messageHeader + sizeof(u16) + sizeof(u64)));
-		u32 sendCount = u32((sendLeft + messageCapacity - 1)/ messageCapacity);
-
-		struct Container
-		{
-			struct iterator
-			{
-				iterator() : index(0) {}
-				iterator(u32 i) : index(i) {}
-				iterator operator++(int) { return iterator(index++); }
-				bool operator==(const iterator& o) const { return index == o.index; }
-				u32 index;
-			};
-			Container(u32 c) : count(c) {}
-			iterator begin() const { return iterator(0); }
-			iterator end() const { return iterator(count); }
-			u32 size() { return count; }
-			u32 count;
-		};
+		u64 messageCapacity = client.GetMessageMaxSize() - (messageHeader + sizeof(u16) + sizeof(u64));
+		u64 sendCount = (sendLeft + messageCapacity - 1) / messageCapacity;
 
 		u32 connectionCount = client.GetConnectionCount();
 		if (workManager && connectionCount > 1) // Use parallel send.. to make sure all sockets send buffers are fully utilized
 		{
+			struct Container
+			{
+				struct iterator
+				{
+					iterator() : index(0) {}
+					iterator(u32 i) : index(i) {}
+					iterator operator++(int) { return iterator(index++); }
+					bool operator==(const iterator& o) const { return index == o.index; }
+					u32 index;
+				};
+				Container(u64 c) : count(u32(c)) {}
+				iterator begin() const { return iterator(0); }
+				iterator end() const { return iterator(count); }
+				u32 size() { return count; }
+				u32 count;
+			} container(sendCount);
+
 			Atomic<bool> error;
-			Container container(sendCount);
 
 			workManager->ParallelFor(connectionCount - 1, container, [&](const WorkContext&, auto& it)
 				{
 					u64 sendPos2 = sendPos + it.index * messageCapacity;
-					StackBinaryWriter<SendMaxSize> writer;
-					NetworkMessage msg(client, StorageServiceId, StorageMessageType_StoreSegment, writer);
-					writer.WriteU16(storeId);
-					writer.WriteU64(sendPos2);
-
-					UBA_ASSERT(messageCapacity == writer.GetCapacityLeft());
-
 					const u8* writeData = readData + it.index * messageCapacity;
 					u64 toWrite = messageCapacity;
 					if (it.index == sendCount - 1)
 						toWrite = sendLeft % messageCapacity;
-
-					writer.WriteBytes(writeData, toWrite); // We always need a separate buffer on the send since we do in-place encryption of the buffers
-
-					if (!msg.Send())
+					if (!SendFileSegmentMessage(client, storeId, sendPos2, writeData, toWrite))
 						error = true;
-
 				}, TCV("SendFile"));
 
 			if (error)
@@ -381,17 +379,9 @@ namespace uba
 		{
 			while (sendLeft)
 			{
-				StackBinaryWriter<SendMaxSize> writer;
-				NetworkMessage msg(client, StorageServiceId, StorageMessageType_StoreSegment, writer);
-				writer.WriteU16(storeId);
-				writer.WriteU64(sendPos);
-
-				u64 toWrite = Min(sendLeft, writer.GetCapacityLeft());
-				writer.WriteBytes(readData, toWrite); // We always need a separate buffer on the send since we do in-place encryption of the buffers
-
-				if (!msg.Send())
+				u64 toWrite = Min(sendLeft, messageCapacity);
+				if (!SendFileSegmentMessage(client, storeId, sendPos, readData, toWrite))
 					return false;
-
 				readData += toWrite;
 				sendLeft -= toWrite;
 				sendPos += toWrite;
