[UBA]
* Reduced size of stack buffer when printing cache client summary to trace file

diff --git a/Engine/Source/Programs/UnrealBuildAccelerator/Common/Private/UbaCacheClient.cpp b/Engine/Source/Programs/UnrealBuildAccelerator/Common/Private/UbaCacheClient.cpp
--- a/Engine/Source/Programs/UnrealBuildAccelerator/Common/Private/UbaCacheClient.cpp
+++ b/Engine/Source/Programs/UnrealBuildAccelerator/Common/Private/UbaCacheClient.cpp
@@ -1065,7 +1065,7 @@ namespace uba
 
 	void CacheClient::TraceSummary()
 	{
-		StackBinaryWriter<256*1024> writer;
+		StackBinaryWriter<16*256> writer;
 		m_session.GetTrace().WriteSummaryText(writer, [&](Logger& logger)
 			{
 				m_client.PrintSummary(logger);
