[UBA]
* Fixed hover-over graph showing wrong stats

diff --git a/Engine/Source/Programs/UnrealBuildAccelerator/Visualizer/Private/UbaVisualizer.cpp b/Engine/Source/Programs/UnrealBuildAccelerator/Visualizer/Private/UbaVisualizer.cpp
--- a/Engine/Source/Programs/UnrealBuildAccelerator/Visualizer/Private/UbaVisualizer.cpp
+++ b/Engine/Source/Programs/UnrealBuildAccelerator/Visualizer/Private/UbaVisualizer.cpp
@@ -3099,16 +3099,18 @@ namespace uba
 							loop = false;
 						}
 
-						int prevPrevX = 100000;
-						int prevX = 0;
+						int prevX = -1;
+						int prevCenterX = -1;
 						for (; i!=e; ++i)
 						{
 							u64 updateTime = session.updates[i];
 							auto updateSend = session.networkSend[i];
 							auto updateRecv = session.networkRecv[i];
 
-							int hitOffset = (prevPrevX - prevX)/2;
-							if (pos.x + hitOffset <= prevX)
+							int x = posX + int(float(TimeToS(updateTime)) * scaleX);
+							int centerX = prevX != -1 ? (prevX + (x - prevX)/2) : x;
+
+							if (prevCenterX != -1 && pos.x >= prevCenterX && pos.x < centerX)
 							{
 								u64 prevTime = 0;
 								u64 prevSend = 0;
@@ -3144,9 +3146,9 @@ namespace uba
 								}
 								return;
 							}
-
-							prevPrevX = prevX;
-							prevX = posX + int(float(TimeToS(updateTime)) * scaleX);
+							
+							prevCenterX = centerX;
+							prevX = x;
 						}
 						++reconnectIndex;
 					}
