[UBA]
* Moved cache client trace summary report to dtor of cache client instead of Scheduler

diff --git a/Engine/Source/Programs/UnrealBuildAccelerator/Common/Private/UbaCacheClient.cpp b/Engine/Source/Programs/UnrealBuildAccelerator/Common/Private/UbaCacheClient.cpp
--- a/Engine/Source/Programs/UnrealBuildAccelerator/Common/Private/UbaCacheClient.cpp
+++ b/Engine/Source/Programs/UnrealBuildAccelerator/Common/Private/UbaCacheClient.cpp
@@ -156,7 +156,10 @@ namespace uba
 			m_client.SetWorkTracker(&m_session.GetTrace());
 	}
 
-	CacheClient::~CacheClient() = default;
+	CacheClient::~CacheClient()
+	{
+		TraceSummary();
+	}
 
 	bool CacheClient::RegisterPathHash(const tchar* path, const CasKey& hash)
 	{
diff --git a/Engine/Source/Programs/UnrealBuildAccelerator/Common/Private/UbaScheduler.cpp b/Engine/Source/Programs/UnrealBuildAccelerator/Common/Private/UbaScheduler.cpp
--- a/Engine/Source/Programs/UnrealBuildAccelerator/Common/Private/UbaScheduler.cpp
+++ b/Engine/Source/Programs/UnrealBuildAccelerator/Common/Private/UbaScheduler.cpp
@@ -177,9 +177,6 @@ namespace uba
 		m_processEntries.clear();
 		m_processEntriesStart = 0;
 		m_session.SetOuterScheduler(nullptr);
-
-		for (auto cacheClient : m_cacheClients)
-			cacheClient->TraceSummary();
 	}
 
 	void Scheduler::SetMaxLocalProcessors(u32 maxLocalProcessors)
