[UBA]
* Improved cache send stats in trace file

diff --git a/Engine/Source/Programs/UnrealBuildAccelerator/Common/Private/UbaCacheClient.cpp b/Engine/Source/Programs/UnrealBuildAccelerator/Common/Private/UbaCacheClient.cpp
--- a/Engine/Source/Programs/UnrealBuildAccelerator/Common/Private/UbaCacheClient.cpp
+++ b/Engine/Source/Programs/UnrealBuildAccelerator/Common/Private/UbaCacheClient.cpp
@@ -1247,7 +1247,7 @@ namespace uba
 				if (!SendFile(logger, m_client, &m_storage.m_workManager, casKey, fileData, fileSize, path.data))
 					return false;
 
-				stats.sendBytesComp += fileSize;
+				stats.sendFile.bytes += fileSize;
 			}
 			else // If we don't have the cas key it should be one of the normalized files.... otherwise there is a bug
 			{
@@ -1303,7 +1303,7 @@ namespace uba
 				if (!sender.SendFileCompressed(casKey, path.data, dataToSend, sizeToSend, path.data))
 					return m_logger.Warning(TC("Failed to send cas content for file %s"), path.data);
 
-				stats.sendBytesComp += sender.m_bytesSent;
+				stats.sendNormalizedFile.bytes += sender.m_bytesSent;
 			}
 
 		}
diff --git a/Engine/Source/Programs/UnrealBuildAccelerator/Common/Private/UbaStats.cpp b/Engine/Source/Programs/UnrealBuildAccelerator/Common/Private/UbaStats.cpp
--- a/Engine/Source/Programs/UnrealBuildAccelerator/Common/Private/UbaStats.cpp
+++ b/Engine/Source/Programs/UnrealBuildAccelerator/Common/Private/UbaStats.cpp
@@ -466,8 +466,6 @@ namespace uba
 		#define UBA_CACHE_STAT(type, var, ver) LogStat(logger, #var, var, frequency);
 		UBA_CACHE_SEND_STATS
 		#undef UBA_CACHE_STAT
-		if (sendBytesComp)
-			logger.Info(TC("   Bytes   Raw/Comp %9s %9s"), BytesToText(sendBytesRaw).str, BytesToText(sendBytesComp).str);
 	}
 
 	bool CacheSendStats::IsEmpty()
diff --git a/Engine/Source/Programs/UnrealBuildAccelerator/Common/Public/UbaStats.h b/Engine/Source/Programs/UnrealBuildAccelerator/Common/Public/UbaStats.h
--- a/Engine/Source/Programs/UnrealBuildAccelerator/Common/Public/UbaStats.h
+++ b/Engine/Source/Programs/UnrealBuildAccelerator/Common/Public/UbaStats.h
@@ -160,10 +160,8 @@ namespace uba
 		UBA_CACHE_STAT(Timer, sendTableWait, 0) \
 		UBA_CACHE_STAT(Timer, sendEntry, 0) \
 		UBA_CACHE_STAT(Timer, createCas, 0) \
-		UBA_CACHE_STAT(Timer, sendFile, 0) \
-		UBA_CACHE_STAT(Timer, sendNormalizedFile, 0) \
-		UBA_CACHE_STAT(AtomicU64, sendBytesRaw, 0) \
-		UBA_CACHE_STAT(AtomicU64, sendBytesComp, 0) \
+		UBA_CACHE_STAT(TimeAndBytes, sendFile, 0) \
+		UBA_CACHE_STAT(TimeAndBytes, sendNormalizedFile, 0) \
 
     struct CacheSendStats
 	{
diff --git a/Engine/Source/Programs/UnrealBuildAccelerator/Visualizer/Private/UbaTraceReader.cpp b/Engine/Source/Programs/UnrealBuildAccelerator/Visualizer/Private/UbaTraceReader.cpp
--- a/Engine/Source/Programs/UnrealBuildAccelerator/Visualizer/Private/UbaTraceReader.cpp
+++ b/Engine/Source/Programs/UnrealBuildAccelerator/Visualizer/Private/UbaTraceReader.cpp
@@ -1313,7 +1313,7 @@ namespace uba
 			CacheSendStats stats;
 			const u8* dataStart = reader.GetPositionData();
 			if (out.version < 45)
-				stats.sendBytesComp = reader.Read7BitEncoded();
+				stats.sendFile.bytes = reader.Read7BitEncoded();
 			else
 				stats.Read(reader, out.version);
 
