[UBA]
* Fixed so it is possible to set recv and send buffer sizes for connections

diff --git a/Engine/Source/Programs/UnrealBuildAccelerator/Common/Private/UbaNetworkBackendMemory.cpp b/Engine/Source/Programs/UnrealBuildAccelerator/Common/Private/UbaNetworkBackendMemory.cpp
--- a/Engine/Source/Programs/UnrealBuildAccelerator/Common/Private/UbaNetworkBackendMemory.cpp
+++ b/Engine/Source/Programs/UnrealBuildAccelerator/Common/Private/UbaNetworkBackendMemory.cpp
@@ -195,6 +195,14 @@ namespace uba
 		//delete m_connection;
 	}
 
+	void NetworkBackendMemory::SetSendBufferSize(void* connection, u32 bytes)
+	{
+	}
+
+	void NetworkBackendMemory::SetRecvBufferSize(void* connection, u32 bytes)
+	{
+	}
+
 	void NetworkBackendMemory::GetTotalSendAndRecv(u64& outSend, u64& outRecv)
 	{
 		outSend = 0;
diff --git a/Engine/Source/Programs/UnrealBuildAccelerator/Common/Private/UbaNetworkBackendTcp.cpp b/Engine/Source/Programs/UnrealBuildAccelerator/Common/Private/UbaNetworkBackendTcp.cpp
--- a/Engine/Source/Programs/UnrealBuildAccelerator/Common/Private/UbaNetworkBackendTcp.cpp
+++ b/Engine/Source/Programs/UnrealBuildAccelerator/Common/Private/UbaNetworkBackendTcp.cpp
@@ -571,8 +571,10 @@ namespace uba
 			// For some reason performance gets destroyed on wine when changing these
 			if (IsWindows && !IsRunningWine())
 			{
-				SetRecvBuf(logger, clientSocket, m_recvBufferSize);
-				SetSendBuf(logger, clientSocket, m_sendBufferSize);
+				if (m_recvBufferSize)
+					SetRecvBuf(logger, clientSocket, m_recvBufferSize);
+				if (m_sendBufferSize)
+					SetSendBuf(logger, clientSocket, m_sendBufferSize);
 			}
 
 			SCOPED_FUTEX(m_connectionsLock, lock);
@@ -887,6 +889,20 @@ namespace uba
 		}
 	}
 
+	void NetworkBackendTcp::SetSendBufferSize(void* connection, u32 bytes)
+	{
+		auto& conn = *(Connection*)connection;
+		SCOPED_CRITICAL_SECTION(conn.shutdownLock, lock);
+		SetSendBuf(m_logger, conn.socket, bytes);
+	}
+
+	void NetworkBackendTcp::SetRecvBufferSize(void* connection, u32 bytes)
+	{
+		auto& conn = *(Connection*)connection;
+		SCOPED_CRITICAL_SECTION(conn.shutdownLock, lock);
+		SetRecvBuf(m_logger, conn.socket, bytes);
+	}
+
 	void NetworkBackendTcp::GetTotalSendAndRecv(u64& outSend, u64& outRecv)
 	{
 		outSend = m_totalSend;
@@ -1396,11 +1412,12 @@ namespace uba
 
 				if (GetCongestionAlgorithm)
 					GetCongestionAlgorithm(s, alg, sizeof(alg));
+
 				#if PLATFORM_LINUX
 				getsockopt(s, SOL_SOCKET, SO_PRIORITY, (char*)&prio, &optlen);
 				#endif
 
-				m_logger.Info(TC("%i - Send: %s, Recv: %s, Prio %i, Alg: %hs"), (int)s, BytesToText(sendBufSize).str, BytesToText(recvBufSize).str, prio, alg);
+				m_logger.Info(TC("%i - Send: %s, Recv: %s, Prio %i%hs%hs"), (int)s, BytesToText(sendBufSize).str, BytesToText(recvBufSize).str, prio, *alg ? ", Alg: " : "", alg);
 				#endif
 
 			}
diff --git a/Engine/Source/Programs/UnrealBuildAccelerator/Common/Private/UbaNetworkClient.cpp b/Engine/Source/Programs/UnrealBuildAccelerator/Common/Private/UbaNetworkClient.cpp
--- a/Engine/Source/Programs/UnrealBuildAccelerator/Common/Private/UbaNetworkClient.cpp
+++ b/Engine/Source/Programs/UnrealBuildAccelerator/Common/Private/UbaNetworkClient.cpp
@@ -17,7 +17,8 @@ namespace uba
 		if (!tablePtr)
 			return;
 		const ConfigTable& table = *tablePtr;
-		(void)table;
+		table.GetValueAsU32(recvBufferSize, TC("RecvBufferSize"));
+		table.GetValueAsU32(sendBufferSize, TC("SendBufferSize"));
 	}
 
 	struct NetworkClient::Connection
@@ -53,6 +54,8 @@ namespace uba
 		m_sendSize = fixedSendSize;
 		m_receiveTimeoutSeconds = info.receiveTimeoutSeconds;
 		m_useMessagePriority = info.useMessagePriority;
+		m_recvBufferSize = info.recvBufferSize;
+		m_sendBufferSize = info.sendBufferSize;
 
 		m_connectionsIt = m_connections.end();
 
@@ -117,6 +120,14 @@ namespace uba
 
 		backend.SetRecvTimeout(backendConnection, 10*1000, nullptr, nullptr); // 10 seconds timeout during adding connection. RecvContext set it to right time in ctor
 
+		if (IsWindows && !IsRunningWine())
+		{
+			if (m_recvBufferSize)
+				backend.SetRecvBufferSize(backendConnection, m_recvBufferSize);
+			if (m_sendBufferSize)
+				backend.SetSendBufferSize(backendConnection, m_sendBufferSize);
+		}
+
 		RecvContext rc(*this, backend, backendConnection);
 
 		// The only way out of this function is to get a call to one of the below callbacks since exitScopeEvent must be set.
diff --git a/Engine/Source/Programs/UnrealBuildAccelerator/Common/Private/UbaNetworkServer.cpp b/Engine/Source/Programs/UnrealBuildAccelerator/Common/Private/UbaNetworkServer.cpp
--- a/Engine/Source/Programs/UnrealBuildAccelerator/Common/Private/UbaNetworkServer.cpp
+++ b/Engine/Source/Programs/UnrealBuildAccelerator/Common/Private/UbaNetworkServer.cpp
@@ -10,6 +10,12 @@ namespace uba
 {
 	void NetworkServerCreateInfo::Apply(Config& config, const tchar* tableName)
 	{
+		const ConfigTable* tablePtr = config.GetTable(tableName);
+		if (!tablePtr)
+			return;
+		const ConfigTable& table = *tablePtr;
+		table.GetValueAsU32(recvBufferSize, TC("RecvBufferSize"));
+		table.GetValueAsU32(sendBufferSize, TC("SendBufferSize"));
 	}
 
 	struct NetworkServer::Client
@@ -782,6 +788,9 @@ namespace uba
 		m_useKeepAlive = info.useKeepAlive;
 		m_useMessagePriority = info.useMessagePriority;
 
+		m_recvBufferSize = info.recvBufferSize;
+		m_sendBufferSize = info.sendBufferSize;
+
 		#if PLATFORM_MAC
 		m_useKeepAlive = true; // Always run keep alive on mac since the built-in one has a probe interval of 1 minute or something like that... so timeout is always 10 minutes
 		#endif
@@ -1546,6 +1555,14 @@ namespace uba
 			return false;
 		}
 
+		if (IsWindows && !IsRunningWine())
+		{
+			if (m_recvBufferSize)
+				backend.SetRecvBufferSize(backendConnection, m_recvBufferSize);
+			if (m_sendBufferSize)
+				backend.SetSendBufferSize(backendConnection, m_sendBufferSize);
+		}
+
 		m_connections.emplace_back(*this, backend, backendConnection, remoteSocketAddr, requiresCrypto, cryptoKey, m_connectionIdCounter++);
 		m_maxActiveConnections = Max(m_maxActiveConnections, u32(m_connections.size()));
 		return true;
diff --git a/Engine/Source/Programs/UnrealBuildAccelerator/Common/Public/UbaNetworkBackend.h b/Engine/Source/Programs/UnrealBuildAccelerator/Common/Public/UbaNetworkBackend.h
--- a/Engine/Source/Programs/UnrealBuildAccelerator/Common/Public/UbaNetworkBackend.h
+++ b/Engine/Source/Programs/UnrealBuildAccelerator/Common/Public/UbaNetworkBackend.h
@@ -70,6 +70,10 @@ namespace uba
 		// Remove connection (user guarantee connection is not referenced anywhere)
 		virtual void DeleteConnection(void* connection) = 0;
 
+		// Set internal buffer sizes for send and recv
+		virtual void SetSendBufferSize(void* connection, u32 bytes) = 0;
+		virtual void SetRecvBufferSize(void* connection, u32 bytes) = 0;
+
 		virtual void GetTotalSendAndRecv(u64& outSend, u64& outRecv) = 0;
 		virtual void Validate(Logger& logger, const Vector<void*>& connections, bool full) = 0;
 	};
diff --git a/Engine/Source/Programs/UnrealBuildAccelerator/Common/Public/UbaNetworkBackendMemory.h b/Engine/Source/Programs/UnrealBuildAccelerator/Common/Public/UbaNetworkBackendMemory.h
--- a/Engine/Source/Programs/UnrealBuildAccelerator/Common/Public/UbaNetworkBackendMemory.h
+++ b/Engine/Source/Programs/UnrealBuildAccelerator/Common/Public/UbaNetworkBackendMemory.h
@@ -28,6 +28,9 @@ namespace uba
 		virtual bool Connect(Logger& logger, const sockaddr& remoteSocketAddr, const ConnectedFunc& connectedFunc, bool* timedOut = nullptr, const tchar* nameHint = nullptr) override;
 		virtual void DeleteConnection(void* connection) override;
 
+		virtual void SetSendBufferSize(void* connection, u32 bytes) override;
+		virtual void SetRecvBufferSize(void* connection, u32 bytes) override;
+
 		virtual void GetTotalSendAndRecv(u64& outSend, u64& outRecv) override;
 		virtual void Validate(Logger& logger, const Vector<void*>& connections, bool full) override;
 
diff --git a/Engine/Source/Programs/UnrealBuildAccelerator/Common/Public/UbaNetworkBackendTcp.h b/Engine/Source/Programs/UnrealBuildAccelerator/Common/Public/UbaNetworkBackendTcp.h
--- a/Engine/Source/Programs/UnrealBuildAccelerator/Common/Public/UbaNetworkBackendTcp.h
+++ b/Engine/Source/Programs/UnrealBuildAccelerator/Common/Public/UbaNetworkBackendTcp.h
@@ -21,8 +21,8 @@ namespace uba
 		bool useOverlappedSend = true;
 		u32 iocpWorkerCount = 0; // Iocp is turned off by default
 
-		u32 recvBufferSize = 4*1024*1024;
-		u32 sendBufferSize = 2*1024*1024;
+		u32 recvBufferSize = 0; // 0 means default
+		u32 sendBufferSize = 0; // 0 means default
 
 		void Apply(const Config& config, const tchar* tableName = TC("NetworkBackendTcp"));
 	};
@@ -49,6 +49,9 @@ namespace uba
 		virtual bool Connect(Logger& logger, const sockaddr& remoteSocketAddr, const ConnectedFunc& connectedFunc, bool* timedOut = nullptr, const tchar* nameHint = nullptr) override;
 		virtual void DeleteConnection(void* connection) override;
 
+		virtual void SetSendBufferSize(void* connection, u32 bytes) override;
+		virtual void SetRecvBufferSize(void* connection, u32 bytes) override;
+
 		virtual void GetTotalSendAndRecv(u64& outSend, u64& outRecv) override;
 		virtual void Validate(Logger& logger, const Vector<void*>& connections, bool full) override;
 
diff --git a/Engine/Source/Programs/UnrealBuildAccelerator/Common/Public/UbaNetworkClient.h b/Engine/Source/Programs/UnrealBuildAccelerator/Common/Public/UbaNetworkClient.h
--- a/Engine/Source/Programs/UnrealBuildAccelerator/Common/Public/UbaNetworkClient.h
+++ b/Engine/Source/Programs/UnrealBuildAccelerator/Common/Public/UbaNetworkClient.h
@@ -28,6 +28,10 @@ namespace uba
 		u32 workerCount = ~0u; // ~0u means logical processor count will be used
 		bool useMessagePriority = false;
 		const u8* cryptoKey128 = nullptr;
+
+		u32 recvBufferSize = 0;
+		u32 sendBufferSize = 0;
+
 		void Apply(Config& config, const tchar* tableName = TC("NetworkClient"));
 	};
 
@@ -90,6 +94,8 @@ namespace uba
 		TString m_prefix;
 		LoggerWithWriter m_logger;
 		u32 m_sendSize;
+		u32 m_recvBufferSize;
+		u32 m_sendBufferSize;
 		u32 m_receiveTimeoutSeconds;
 		Atomic<u64> m_sendBytes;
 		Atomic<u64> m_recvBytes;
diff --git a/Engine/Source/Programs/UnrealBuildAccelerator/Common/Public/UbaNetworkServer.h b/Engine/Source/Programs/UnrealBuildAccelerator/Common/Public/UbaNetworkServer.h
--- a/Engine/Source/Programs/UnrealBuildAccelerator/Common/Public/UbaNetworkServer.h
+++ b/Engine/Source/Programs/UnrealBuildAccelerator/Common/Public/UbaNetworkServer.h
@@ -40,6 +40,9 @@ namespace uba
 		bool logConnections = true;
 		bool useKeepAlive = false;
 		bool useMessagePriority = false;
+
+		u32 recvBufferSize = 0;
+		u32 sendBufferSize = 0;
 	};
 
 	struct MessageInfo
@@ -173,6 +176,8 @@ namespace uba
 		u32 m_receiveTimeoutMs = 0;
 		u32 m_createdWorkerCount = 0;
 		u32 m_maxCreatedWorkerCount = 0;
+		u32 m_recvBufferSize;
+		u32 m_sendBufferSize;
 		bool m_workersEnabled = true;
 		bool m_logConnections = true;
 		bool m_useKeepAlive = false;
diff --git a/Engine/Source/Programs/UnrealBuildAccelerator/Host/Private/UbaExports.cpp b/Engine/Source/Programs/UnrealBuildAccelerator/Host/Private/UbaExports.cpp
--- a/Engine/Source/Programs/UnrealBuildAccelerator/Host/Private/UbaExports.cpp
+++ b/Engine/Source/Programs/UnrealBuildAccelerator/Host/Private/UbaExports.cpp
@@ -205,7 +205,6 @@ extern "C"
 		{
 			NetworkBackendTcpCreateInfo info(writer);
 			info.Apply(GetConfig());
-			//info.statusUpdateSeconds = 2;
 			networkBackend = new NetworkBackendTcp(info);
 		}
 
