[UBA]
* Added api for adding custom tasks into session trace from the outside of uba

diff --git a/Engine/Source/Programs/UnrealBuildAccelerator/Common/Private/UbaTrace.cpp b/Engine/Source/Programs/UnrealBuildAccelerator/Common/Private/UbaTrace.cpp
--- a/Engine/Source/Programs/UnrealBuildAccelerator/Common/Private/UbaTrace.cpp
+++ b/Engine/Source/Programs/UnrealBuildAccelerator/Common/Private/UbaTrace.cpp
@@ -543,6 +543,29 @@ namespace uba
 		writer.WriteBytes(data, dataSize);
 	}
 
+	void Trace::TaskBegin(u32 processId, const StringView& description, const StringView& details, const Color& color)
+	{
+		BEGIN_TRACE_ENTRY(TaskBegin);
+		writer.Write7BitEncoded(processId);
+		writer.WriteString(description);
+		writer.WriteString(details);
+		writer.WriteU32(color);
+	}
+
+	void Trace::TaskHint(u32 processId, const StringView& hint)
+	{
+		BEGIN_TRACE_ENTRY(TaskHint);
+		writer.Write7BitEncoded(processId);
+		writer.WriteString(hint);
+	}
+
+	void Trace::TaskEnd(u32 processId, bool success)
+	{
+		BEGIN_TRACE_ENTRY(TaskEnd);
+		writer.Write7BitEncoded(processId);
+		writer.WriteBool(success);
+	}
+
 	TraceChannel::TraceChannel(Logger& logger) : m_logger(logger)
 	{
 	}
diff --git a/Engine/Source/Programs/UnrealBuildAccelerator/Common/Public/UbaTrace.h b/Engine/Source/Programs/UnrealBuildAccelerator/Common/Public/UbaTrace.h
--- a/Engine/Source/Programs/UnrealBuildAccelerator/Common/Public/UbaTrace.h
+++ b/Engine/Source/Programs/UnrealBuildAccelerator/Common/Public/UbaTrace.h
@@ -62,6 +62,9 @@ namespace uba
 		UBA_TRACE_TYPE(WorkHint) \
 		UBA_TRACE_TYPE(DriveUpdate) \
 		UBA_TRACE_TYPE(CacheSummary) \
+		UBA_TRACE_TYPE(TaskBegin) \
+		UBA_TRACE_TYPE(TaskHint) \
+		UBA_TRACE_TYPE(TaskEnd) \
 
 	enum TraceType : u8
 	{
@@ -70,7 +73,7 @@ namespace uba
 		#undef UBA_TRACE_TYPE
 	};
 
-	static constexpr u32 TraceVersion = 45;
+	static constexpr u32 TraceVersion = 46;
 	static constexpr u32 TraceReadCompatibilityVersion = 6;
 
 	class Trace : public WorkTracker
@@ -115,6 +118,10 @@ namespace uba
 		void CacheEndWrite(u32 processId, bool success, const u8* data, u64 dataSize);
 		void CacheSummary(const u8* data, u64 dataSize);
 
+		void TaskBegin(u32 processId, const StringView& description, const StringView& details, const Color& color);
+		void TaskHint(u32 processId, const StringView& hint);
+		void TaskEnd(u32 processId, bool success);
+
 		// Writes the current state of the trace to the specified output file without stopping anything
 		bool Write(const tchar* writeFileName, bool writeSummary);
 
diff --git a/Engine/Source/Programs/UnrealBuildAccelerator/Host/Private/UbaExports.cpp b/Engine/Source/Programs/UnrealBuildAccelerator/Host/Private/UbaExports.cpp
--- a/Engine/Source/Programs/UnrealBuildAccelerator/Host/Private/UbaExports.cpp
+++ b/Engine/Source/Programs/UnrealBuildAccelerator/Host/Private/UbaExports.cpp
@@ -597,6 +597,23 @@ extern "C"
 			server->SaveSnapshotOfTrace();
 	}
 
+	uba::u32 SessionServer_TaskBegin(uba::SessionServer* server, const uba::tchar* description, const uba::tchar* details, uba::u32 color)
+	{
+		uba::u32 taskId = server->CreateProcessId();
+		server->GetTrace().TaskBegin(taskId, uba::ToView(description), uba::ToView(details), color);
+		return taskId;
+	}
+
+	void SessionServer_TaskHint(uba::SessionServer* server, uba::u32 taskId, const uba::tchar* hint)
+	{
+		server->GetTrace().TaskHint(taskId, uba::ToView(hint));
+	}
+
+	void SessionServer_TaskEnd(uba::SessionServer* server, uba::u32 taskId, bool success)
+	{
+		server->GetTrace().TaskEnd(taskId, success);
+	}
+
 	void SessionServer_Destroy(uba::SessionServer* server)
 	{
 		if (server)
diff --git a/Engine/Source/Programs/UnrealBuildAccelerator/Host/Public/UbaExports.h b/Engine/Source/Programs/UnrealBuildAccelerator/Host/Public/UbaExports.h
--- a/Engine/Source/Programs/UnrealBuildAccelerator/Host/Public/UbaExports.h
+++ b/Engine/Source/Programs/UnrealBuildAccelerator/Host/Public/UbaExports.h
@@ -144,6 +144,9 @@ extern "C"
 	UBA_API void SessionServer_RegisterCustomService(uba::SessionServer* server, SessionServer_CustomServiceFunction* function, void* userData = nullptr);
 	UBA_API void SessionServer_RegisterCrossArchitectureMapping(uba::SessionServer* server, const uba::tchar* from, const uba::tchar* to);
 	UBA_API void SessionServer_SaveSnapshotOfTrace(uba::SessionServer* server);
+	UBA_API uba::u32 SessionServer_TaskBegin(uba::SessionServer* server, const uba::tchar* description, const uba::tchar* details, uba::u32 color);
+	UBA_API void SessionServer_TaskHint(uba::SessionServer* server, uba::u32 taskId, const uba::tchar* hint);
+	UBA_API void SessionServer_TaskEnd(uba::SessionServer* server, uba::u32 taskId, bool success);
 	UBA_API void SessionServer_Destroy(uba::SessionServer* server);
 
 	// Scheduler
diff --git a/Engine/Source/Programs/UnrealBuildAccelerator/Visualizer/Private/UbaTraceReader.cpp b/Engine/Source/Programs/UnrealBuildAccelerator/Visualizer/Private/UbaTraceReader.cpp
--- a/Engine/Source/Programs/UnrealBuildAccelerator/Visualizer/Private/UbaTraceReader.cpp
+++ b/Engine/Source/Programs/UnrealBuildAccelerator/Visualizer/Private/UbaTraceReader.cpp
@@ -630,8 +630,7 @@ namespace uba
 		{
 			u32 sessionIndex = reader.ReadU32();
 			u32 id = reader.ReadU32();
-			StringBuffer<> desc;
-			reader.ReadString(desc);
+			TString desc = reader.ReadString();
 
 			TString breadcrumbs;
 			if (out.version >= 35)
@@ -663,7 +662,7 @@ namespace uba
 				}
 			}
 
-			ProcessBegin(out, sessionIndex, id, time, desc, breadcrumbs);
+			ProcessBegin(out, sessionIndex, id, time, std::move(desc), std::move(breadcrumbs));
 			break;
 		}
 		case TraceType_ProcessExited:
@@ -1146,6 +1145,31 @@ namespace uba
 
 			break;
 		}
+		case TraceType_TaskBegin:
+		{
+			u32 id = u32(reader.Read7BitEncoded());
+			TString desc = reader.ReadString();
+			TString details = reader.ReadString();
+			u32 color =  reader.ReadU32();(void)color;
+			ProcessBegin(out, 0, id, time, std::move(desc), std::move(details))->type = TraceView::ProcessType_Task;
+			break;
+		}
+		case TraceType_TaskHint:
+		{
+			u32 id = u32(reader.Read7BitEncoded());(void)id;
+			TString hint = reader.ReadString();
+			break;
+		}
+		case TraceType_TaskEnd:
+		{
+			u32 id = u32(reader.Read7BitEncoded());
+			bool success = reader.ReadBool();
+
+			u32 sessionIndex;
+			TraceView::Process& process = *ProcessEnd(out, sessionIndex, id, time);
+			process.exitCode = success ? 0 : -1;
+			break;
+		}
 		case TraceType_ProgressUpdate:
 		{
 			out.progressProcessesTotal = u32(reader.Read7BitEncoded());
@@ -1267,9 +1291,8 @@ namespace uba
 		case TraceType_CacheBeginFetch:
 		{
 			u32 id = u32(reader.Read7BitEncoded());
-			StringBuffer<> desc;
-			reader.ReadString(desc);
-			ProcessBegin(out, 0, id, time, desc, StringView())->cacheFetch = true;
+			TString desc = reader.ReadString();
+			ProcessBegin(out, 0, id, time, std::move(desc), {})->type = TraceView::ProcessType_CacheFetch;
 			break;
 		}
 		case TraceType_CacheEndFetch:
@@ -1463,7 +1486,7 @@ namespace uba
 		return nullptr;
 	}
 
-	TraceView::Process* TraceReader::ProcessBegin(TraceView& out, u32 sessionIndex, u32 id, u64 time, const StringView& description, const StringView& breadcrumbs)
+	TraceView::Process* TraceReader::ProcessBegin(TraceView& out, u32 sessionIndex, u32 id, u64 time, TString&& description, TString&& breadcrumbs)
 	{
 		TraceView::Processor* processor = nullptr;
 		TraceView::Session& session = GetSession(out, sessionIndex);
@@ -1498,8 +1521,8 @@ namespace uba
 		++out.totalProcessActiveCount;
 
 		process.id = id;
-		process.description = description.data;
-		process.breadcrumbs = breadcrumbs.data;
+		process.description = std::move(description);
+		process.breadcrumbs = std::move(breadcrumbs);
 		process.start = time;
 		process.stop = ~u64(0);
 		process.exitCode = ~0u;
diff --git a/Engine/Source/Programs/UnrealBuildAccelerator/Visualizer/Private/UbaTraceReader.h b/Engine/Source/Programs/UnrealBuildAccelerator/Visualizer/Private/UbaTraceReader.h
--- a/Engine/Source/Programs/UnrealBuildAccelerator/Visualizer/Private/UbaTraceReader.h
+++ b/Engine/Source/Programs/UnrealBuildAccelerator/Visualizer/Private/UbaTraceReader.h
@@ -13,6 +13,8 @@ namespace uba
 	class TraceView
 	{
 	public:
+		enum ProcessType : u8 { ProcessType_Normal, ProcessType_CacheFetch, ProcessType_Task };
+
 		struct Process
 		{
 			u32 id = 0;
@@ -24,8 +26,8 @@ namespace uba
 			TString breadcrumbs;
 			HBITMAP bitmap = 0;
 			u32 bitmapOffset = 0;
+			ProcessType type = ProcessType_Normal;
 			bool bitmapDirty = true;
-			bool cacheFetch = false;
 			bool isRemote = false;
 			bool isReuse = false;
 			u64 createFilesTime = 0;
@@ -229,7 +231,7 @@ namespace uba
 		TraceView::Session& GetSession(TraceView& out, u32 sessionIndex);
 		TraceView::Session* GetSession(TraceView& out, const Guid& clientUid);
 
-		TraceView::Process* ProcessBegin(TraceView& out, u32 sessionIndex, u32 id, u64 time, const StringView& description, const StringView& breadcrumbs);
+		TraceView::Process* ProcessBegin(TraceView& out, u32 sessionIndex, u32 id, u64 time, TString&& description, TString&& breadcrumbs);
 		TraceView::Process* ProcessEnd(TraceView& out, u32& outSessionIndex, u32 id, u64 time);
 
 		UnorderedMap<u32, TraceView::ProcessLocation> m_activeProcesses;
diff --git a/Engine/Source/Programs/UnrealBuildAccelerator/Visualizer/Private/UbaVisualizer.cpp b/Engine/Source/Programs/UnrealBuildAccelerator/Visualizer/Private/UbaVisualizer.cpp
--- a/Engine/Source/Programs/UnrealBuildAccelerator/Visualizer/Private/UbaVisualizer.cpp
+++ b/Engine/Source/Programs/UnrealBuildAccelerator/Visualizer/Private/UbaVisualizer.cpp
@@ -1275,7 +1275,7 @@ namespace uba
 
 						if (process.isRemote)
 							str.Append(L" [").Append(session.name).Append(']');
-						else if (process.cacheFetch)
+						else if (process.type == TraceView::ProcessType_CacheFetch)
 							str.Append(L" [cache]");
 						if (boxRect.left < 0)
 							str.Appendf(L"   %s", TimeToText(playTime - process.start, true).str);
@@ -2091,8 +2091,11 @@ namespace uba
 
 			logger.AddSpace(2);
 			logger.Info(L"  %ls", process.description.c_str());
-			logger.Info(L"  Host:        %ls", m_processSelectedLocation.sessionIndex == 0 ? TC("local") : m_traceView.GetSession(m_processSelectedLocation).name.c_str());
-			logger.Info(L"  ProcessId:  %6u", process.id);
+			if (process.type != TraceView::ProcessType_Task)
+			{
+				logger.Info(L"  Host:        %ls", m_processSelectedLocation.sessionIndex == 0 ? TC("local") : m_traceView.GetSession(m_processSelectedLocation).name.c_str());
+				logger.Info(L"  ProcessId:  %6u", process.id);
+			}
 			logger.Info(L"  Start:     %7ls (%ls)", TimeToText(process.start, true).str, GetWorldTime(process.start).data);
 			logger.Info(L"  Duration:  %7ls", TimeToText(duration, true).str);
 			if (!process.returnedReason.empty())
@@ -2145,7 +2148,7 @@ namespace uba
 				KernelStats kernelStats;
 				CacheFetchStats cacheStats;
 
-				if (process.cacheFetch)
+				if (process.type == TraceView::ProcessType_CacheFetch)
 				{
 					if (!process.returnedReason.empty())
 						logger.Info(L"  Cache:       Miss");
@@ -2510,7 +2513,7 @@ namespace uba
 			brush = m_processBrushes[selected].returned;
 		else if (!done)
 			brush = m_processBrushes[selected].inProgress;
-		else if (process.cacheFetch)
+		else if (process.type == TraceView::ProcessType_CacheFetch)
 			brush = m_processBrushes[selected].cacheFetch;
 		else if (process.exitCode == ProcessCancelExitCode)
 			brush = m_processBrushes[selected].returned;
