[UBA]
* Fixed hang in visualizer if assert happens inside paint function

diff --git a/Engine/Source/Programs/UnrealBuildAccelerator/Visualizer/Private/UbaVisualizer.cpp b/Engine/Source/Programs/UnrealBuildAccelerator/Visualizer/Private/UbaVisualizer.cpp
--- a/Engine/Source/Programs/UnrealBuildAccelerator/Visualizer/Private/UbaVisualizer.cpp
+++ b/Engine/Source/Programs/UnrealBuildAccelerator/Visualizer/Private/UbaVisualizer.cpp
@@ -16,6 +16,10 @@
 
 namespace uba
 {
+	#if PLATFORM_WINDOWS
+	extern bool g_isInAssertMessageBox;
+	#endif
+
 	enum
 	{
 		Popup_CopySessionInfo = 3,
@@ -905,7 +909,12 @@ namespace uba
 		}
 		HGDIOBJ oldBmp = SelectObject(memDC, m_cachedBitmap);
 		
-		paintFunc(hdc, memDC, rect);
+		if (!m_isInPaint)
+		{
+			m_isInPaint = true;
+			paintFunc(hdc, memDC, rect);
+			m_isInPaint = false;
+		}
 
 		SelectObject(memDC, oldBmp);
 		DeleteDC(memDC);
diff --git a/Engine/Source/Programs/UnrealBuildAccelerator/Visualizer/Private/UbaVisualizer.h b/Engine/Source/Programs/UnrealBuildAccelerator/Visualizer/Private/UbaVisualizer.h
--- a/Engine/Source/Programs/UnrealBuildAccelerator/Visualizer/Private/UbaVisualizer.h
+++ b/Engine/Source/Programs/UnrealBuildAccelerator/Visualizer/Private/UbaVisualizer.h
@@ -261,6 +261,8 @@ namespace uba
 		bool m_paused = false;
 		u64 m_pauseStart = 0;
 
+		bool m_isInPaint = false;
+
 		static constexpr int BitmapCacheHeight = 1024*1024;
 		HBITMAP m_lastBitmap = 0;
 		int m_lastBitmapOffset = BitmapCacheHeight;
