[UBA]
* Fixed bug when reading older trace files that contains cache writes

diff --git a/Engine/Source/Programs/UnrealBuildAccelerator/Visualizer/Private/UbaVisualizer.cpp b/Engine/Source/Programs/UnrealBuildAccelerator/Visualizer/Private/UbaVisualizer.cpp
--- a/Engine/Source/Programs/UnrealBuildAccelerator/Visualizer/Private/UbaVisualizer.cpp
+++ b/Engine/Source/Programs/UnrealBuildAccelerator/Visualizer/Private/UbaVisualizer.cpp
@@ -808,7 +808,13 @@ namespace uba
 
 			BinaryReader reader(write.stats.data(), 0, write.stats.size());
 			CacheSendStats stats;
-			stats.Read(reader, m_traceView.version);
+			if (m_traceView.version < 45)
+			{
+				stats.sendFile.bytes = reader.Read7BitEncoded();
+				stats.sendFile.count = 1;
+			}
+			else
+				stats.Read(reader, m_traceView.version);
 			stats.Print(logger, m_traceView.frequency);
 		}
 		else
