[UBA]
* Made trace reader more resilient to bad .uba files

diff --git a/Engine/Source/Programs/UnrealBuildAccelerator/Visualizer/Private/UbaTraceReader.cpp b/Engine/Source/Programs/UnrealBuildAccelerator/Visualizer/Private/UbaTraceReader.cpp
--- a/Engine/Source/Programs/UnrealBuildAccelerator/Visualizer/Private/UbaTraceReader.cpp
+++ b/Engine/Source/Programs/UnrealBuildAccelerator/Visualizer/Private/UbaTraceReader.cpp
@@ -1166,8 +1166,8 @@ namespace uba
 			bool success = reader.ReadBool();
 
 			u32 sessionIndex;
-			TraceView::Process& process = *ProcessEnd(out, sessionIndex, id, time);
-			process.exitCode = success ? 0 : -1;
+			if (TraceView::Process* process = ProcessEnd(out, sessionIndex, id, time))
+				process->exitCode = success ? 0 : -1;
 			break;
 		}
 		case TraceType_ProgressUpdate:
