[UBA]
* Added more logging when failing to find input or output file used to create cache entry

diff --git a/Engine/Source/Programs/UnrealBuildAccelerator/Common/Private/UbaCacheClient.cpp b/Engine/Source/Programs/UnrealBuildAccelerator/Common/Private/UbaCacheClient.cpp
--- a/Engine/Source/Programs/UnrealBuildAccelerator/Common/Private/UbaCacheClient.cpp
+++ b/Engine/Source/Programs/UnrealBuildAccelerator/Common/Private/UbaCacheClient.cpp
@@ -365,12 +365,12 @@ namespace uba
 				if (isOutput)
 				{
 					if (!m_storage.StoreCasFile(casKey, path.data, CasKeyZero, deferCreation))
-						return false;
+						return m_logger.Info(TC("Failed to store output file %s when writing to cache"), path.data).ToFalse();
 				}
 				else
 				{
 					if (!m_storage.StoreCasKey(casKey, path.data, CasKeyZero))
-						return false;
+						return m_logger.Info(TC("Failed to store input file %s when writing to cache"), path.data).ToFalse();
 				}
 
 				if (casKey == CasKeyZero) // If file is not found it was a temporary file that was deleted and is not really an output
