[UBA]
* Added some more logging about when failing to store input/output files

diff --git a/Engine/Source/Programs/UnrealBuildAccelerator/Common/Private/UbaCacheClient.cpp b/Engine/Source/Programs/UnrealBuildAccelerator/Common/Private/UbaCacheClient.cpp
--- a/Engine/Source/Programs/UnrealBuildAccelerator/Common/Private/UbaCacheClient.cpp
+++ b/Engine/Source/Programs/UnrealBuildAccelerator/Common/Private/UbaCacheClient.cpp
@@ -656,8 +656,10 @@ namespace uba
 							}
 					}
 					else
-						m_storage.StoreCasKey(localCasKey, fileNameKey, path.data, CasKeyZero);
-
+					{
+						if (!m_storage.StoreCasKey(localCasKey, fileNameKey, path.data, CasKeyZero))
+							m_logger.Info(TC("Failed to store input file %s when fetching from cache"), path.data);
+					}
 					#if UBA_TRACE_FETCH_CACHE
 					storeTime += GetTime() - t;
 					#endif
