[UBA]
* Fixed bug in proxy code where failing begin fetch message was not sending correct bytes to the clients of the proxy

diff --git a/Engine/Source/Programs/UnrealBuildAccelerator/Common/Private/UbaStorageProxy.cpp b/Engine/Source/Programs/UnrealBuildAccelerator/Common/Private/UbaStorageProxy.cpp
--- a/Engine/Source/Programs/UnrealBuildAccelerator/Common/Private/UbaStorageProxy.cpp
+++ b/Engine/Source/Programs/UnrealBuildAccelerator/Common/Private/UbaStorageProxy.cpp
@@ -260,22 +260,14 @@ namespace uba
 				file.error = error;
 				fileLock.Leave();
 
-				if (file.error)
+				StackBinaryWriter<SendMaxSize> writer;
+				bool isFirst = true;
+				for (auto& deferredResponse : f->deferredResponses)
 				{
-					for (auto& deferredResponse : f->deferredResponses)
-						m_server.SendResponse(deferredResponse.info, nullptr, 0);
-				}
-				else
-				{
-					StackBinaryWriter<SendMaxSize> writer;
-					bool isFirst = true;
-					for (auto& deferredResponse : f->deferredResponses)
-					{
-						writer.Reset();
-						SendFetchBeginResponse(deferredResponse.clientId, file, writer, isFirst);
-						m_server.SendResponse(deferredResponse.info, writer.GetData(), u32(writer.GetPosition()));
-						isFirst = false;
-					}
+					writer.Reset();
+					SendFetchBeginResponse(deferredResponse.clientId, file, writer, isFirst);
+					m_server.SendResponse(deferredResponse.info, writer.GetData(), u32(writer.GetPosition()));
+					isFirst = false;
 				}
 				delete f;
 			});
