[UBA]
* Improved error message

diff --git a/Engine/Source/Programs/UnrealBuildAccelerator/Detours/Private/Windows/UbaDetoursFunctionsKernelBase.inl b/Engine/Source/Programs/UnrealBuildAccelerator/Detours/Private/Windows/UbaDetoursFunctionsKernelBase.inl
--- a/Engine/Source/Programs/UnrealBuildAccelerator/Detours/Private/Windows/UbaDetoursFunctionsKernelBase.inl
+++ b/Engine/Source/Programs/UnrealBuildAccelerator/Detours/Private/Windows/UbaDetoursFunctionsKernelBase.inl
@@ -3278,7 +3278,7 @@ BOOL Detoured_CreateProcessW(LPCWSTR lpApplicationName, LPWSTR lpCommandLine, LP
 		if (retryCount++ >= 5)
 			break;
 		const wchar_t* errorText = lastError == ERROR_ACCESS_DENIED ? L"access denied" : L"internal error";
-		Rpc_WriteLogf(L"DetourCreateProcessWithDllEx failed with %ls, retrying %ls (Working dir: %ls)", errorText, originalCmd.data, currentDir.c_str());
+		Rpc_WriteLogf(L"DetourCreateProcessWithDllEx failed with %ls (code: %u), retrying %ls (Working dir: %ls)", errorText, lastError, originalCmd.data, currentDir.c_str());
 		Sleep(100 + (rand() % 200)); // We have no idea
 		continue;
 	}
