[UBA]
* Fixed so oodle uses mimalloc in detoured processes
* Added so a failed allocation for oodle is causing a FatalError and immediately exit the process

diff --git a/Engine/Source/Programs/UnrealBuildAccelerator/Core/Private/UbaPlatform.cpp b/Engine/Source/Programs/UnrealBuildAccelerator/Core/Private/UbaPlatform.cpp
--- a/Engine/Source/Programs/UnrealBuildAccelerator/Core/Private/UbaPlatform.cpp
+++ b/Engine/Source/Programs/UnrealBuildAccelerator/Core/Private/UbaPlatform.cpp
@@ -719,7 +719,10 @@ namespace uba
 	#if UBA_USE_MIMALLOC
 	void* Oodle_MallocAligned(OO_SINTa bytes, OO_S32 alignment)
 	{
-		return mi_malloc_aligned(bytes, alignment);
+		if (void* mem = mi_malloc_aligned(bytes, alignment))
+			return mem;
+		FatalError(9884, TC("Failed to allocate %llu bytes for oodle"), u64(bytes));
+		return nullptr;
 	}
 
 	void Oodle_Free(void* ptr)
diff --git a/Engine/Source/Programs/UnrealBuildAccelerator/Detours/Private/PosixOS/UbaDetoursMainPosix.cpp b/Engine/Source/Programs/UnrealBuildAccelerator/Detours/Private/PosixOS/UbaDetoursMainPosix.cpp
--- a/Engine/Source/Programs/UnrealBuildAccelerator/Detours/Private/PosixOS/UbaDetoursMainPosix.cpp
+++ b/Engine/Source/Programs/UnrealBuildAccelerator/Detours/Private/PosixOS/UbaDetoursMainPosix.cpp
@@ -83,6 +83,7 @@ static void __attribute__((constructor(102))) PreInitCtor()
 	g_messageMappingMem = rptr + sizeof(SharedEvent) * 3;
 
 	AddExceptionHandler();
+	InitMemory();
 
 	g_sessionPid = atoi(getenv("UBA_SESSION_PROCESS"));
 
diff --git a/Engine/Source/Programs/UnrealBuildAccelerator/Detours/Private/Windows/UbaDetoursFunctionsWin.cpp b/Engine/Source/Programs/UnrealBuildAccelerator/Detours/Private/Windows/UbaDetoursFunctionsWin.cpp
--- a/Engine/Source/Programs/UnrealBuildAccelerator/Detours/Private/Windows/UbaDetoursFunctionsWin.cpp
+++ b/Engine/Source/Programs/UnrealBuildAccelerator/Detours/Private/Windows/UbaDetoursFunctionsWin.cpp
@@ -980,6 +980,7 @@ void PreInit(const DetoursPayload& payload)
 void Init(const DetoursPayload& payload, u64 startTime)
 {
 	AddExceptionHandler();
+	InitMemory();
 
 	DetourAttachFunctions(g_runningRemote);
 
