[UBA]
* Fixed static analysis error

diff --git a/Engine/Source/Programs/UnrealBuildAccelerator/Core/Private/UbaPlatform.cpp b/Engine/Source/Programs/UnrealBuildAccelerator/Core/Private/UbaPlatform.cpp
--- a/Engine/Source/Programs/UnrealBuildAccelerator/Core/Private/UbaPlatform.cpp
+++ b/Engine/Source/Programs/UnrealBuildAccelerator/Core/Private/UbaPlatform.cpp
@@ -722,7 +722,6 @@ namespace uba
 		if (void* mem = mi_malloc_aligned(bytes, alignment))
 			return mem;
 		FatalError(9884, TC("Failed to allocate %llu bytes for oodle"), u64(bytes));
-		return nullptr;
 	}
 
 	void Oodle_Free(void* ptr)
