[UBA]
* Brought back code only running in UBA_DEBUG related to adding cas files not deferred.

diff --git a/Engine/Source/Programs/UnrealBuildAccelerator/Common/Private/UbaStorage.cpp b/Engine/Source/Programs/UnrealBuildAccelerator/Common/Private/UbaStorage.cpp
--- a/Engine/Source/Programs/UnrealBuildAccelerator/Common/Private/UbaStorage.cpp
+++ b/Engine/Source/Programs/UnrealBuildAccelerator/Common/Private/UbaStorage.cpp
@@ -562,23 +562,28 @@ namespace uba
 		}
 		else
 		{
-			#if UBA_DEBUG // This should never happen
+			// This should never happen
+			// But it does. LOL. I've seen it on the farm on big builds containing both lots of helpers and lots of cache misses and hits
+
 			SCOPED_FUTEX(m_deferredCasCreationLookupLock, deferredLock);
+			#if UBA_DEBUG
 			auto findIt = m_deferredCasCreationLookupByName.find(fileNameKey);
 			if (findIt != m_deferredCasCreationLookupByName.end() && findIt->second != casKey)
 				return m_logger.Error(TC("Found deferred entry for %s with different cas key while writing cas. A ReportFileWrite is missing somewhere"), fileName);
+			#endif
 
 			auto findIt2 = m_deferredCasCreationLookup.find(casKey);
 			if (findIt2 != m_deferredCasCreationLookup.end())
 			{
-				#if UBA_DEBUG
-				return m_logger.Error(TC("This should not happen. (%s)"), fileName);
-				#endif
+				//#if UBA_DEBUG
+				//return m_logger.Error(TC("This should not happen. (%s)"), fileName);
+				//#endif
+				m_logger.Info(TC("Investigate how we get here. (%s)"), fileName);
+
 				for (auto& it2 : findIt2->second.names)
 					m_deferredCasCreationLookupByName.erase(it2);
 				m_deferredCasCreationLookup.erase(findIt2);
 			}
-			#endif
 		}
 
 		casEntry.verified = true;
