[UBA]
* Rearranged code to make linux target happy (non-reachable code)

diff --git a/Engine/Source/Programs/UnrealBuildAccelerator/Core/Private/UbaPlatform.cpp b/Engine/Source/Programs/UnrealBuildAccelerator/Core/Private/UbaPlatform.cpp
--- a/Engine/Source/Programs/UnrealBuildAccelerator/Core/Private/UbaPlatform.cpp
+++ b/Engine/Source/Programs/UnrealBuildAccelerator/Core/Private/UbaPlatform.cpp
@@ -719,9 +719,10 @@ namespace uba
 	#if UBA_USE_MIMALLOC
 	void* Oodle_MallocAligned(OO_SINTa bytes, OO_S32 alignment)
 	{
-		if (void* mem = mi_malloc_aligned(bytes, alignment))
-			return mem;
-		FatalError(9884, TC("Failed to allocate %llu bytes for oodle"), u64(bytes));
+		void* mem = mi_malloc_aligned(bytes, alignment);
+		if (!mem)
+			FatalError(9884, TC("Failed to allocate %llu bytes for oodle"), u64(bytes));
+		return mem;
 	}
 
 	void Oodle_Free(void* ptr)
