[UBA]
* Added UBA_SHIPPING_ASSERT
* Added temporary hack to try to get more info from error we see on the farm

diff --git a/Engine/Source/Programs/UnrealBuildAccelerator/Common/Private/UbaFile.cpp b/Engine/Source/Programs/UnrealBuildAccelerator/Common/Private/UbaFile.cpp
--- a/Engine/Source/Programs/UnrealBuildAccelerator/Common/Private/UbaFile.cpp
+++ b/Engine/Source/Programs/UnrealBuildAccelerator/Common/Private/UbaFile.cpp
@@ -146,6 +146,12 @@ namespace uba
 		if (lastError == ERROR_SHARING_VIOLATION)
 			GetProcessHoldingFile(additionalInfo, fileName);
 		#endif
+
+		// HACK HACK HACK... want callstack from this error
+		#if PLATFORM_WINDOWS
+		UBA_SHIPPING_ASSERTF(!Contains(fileName, TC("<PathHash>")), TC("ERROR opening file %s for read (%s%s)"), fileName, LastErrorToText(lastError).data, additionalInfo.data);
+		#endif
+
 		return logger.Error(TC("ERROR opening file %s for read (%s%s)"), fileName, LastErrorToText(lastError).data, additionalInfo.data);
 	}
 
diff --git a/Engine/Source/Programs/UnrealBuildAccelerator/Core/Public/UbaPlatform.h b/Engine/Source/Programs/UnrealBuildAccelerator/Core/Public/UbaPlatform.h
--- a/Engine/Source/Programs/UnrealBuildAccelerator/Core/Public/UbaPlatform.h
+++ b/Engine/Source/Programs/UnrealBuildAccelerator/Core/Public/UbaPlatform.h
@@ -68,9 +68,12 @@ namespace uba
 	#define	sizeof_array(array) int(sizeof(array)/sizeof(array[0]))
 	#endif
 
+	#define UBA_SHIPPING_ASSERT(x) do { if (x) break; uba::UbaAssert(TC(""), __FILE__, __LINE__, #x, true, 543221, nullptr, 0); } while(false)
+	#define UBA_SHIPPING_ASSERTF(x, fmt, ...) do { UBA_FMT_CHECK(fmt, ##__VA_ARGS__); if (x) break; uba::StringBuffer<> _buf; _buf.Appendf(fmt, ##__VA_ARGS__); uba::UbaAssert(_buf.data, __FILE__, __LINE__, #x, true, 543221, nullptr, 0); } while(false)
+
 	#if UBA_DEBUG
-	#define UBA_ASSERT(x) do { if (x) break; uba::UbaAssert(TC(""), __FILE__, __LINE__, #x, true, 543221, nullptr, 0); } while(false)
-	#define UBA_ASSERTF(x, fmt, ...) do { UBA_FMT_CHECK(fmt, ##__VA_ARGS__); if (x) break; uba::StringBuffer<> _buf; _buf.Appendf(fmt, ##__VA_ARGS__); uba::UbaAssert(_buf.data, __FILE__, __LINE__, #x, true, 543221, nullptr, 0); } while(false)
+	#define UBA_ASSERT(x) UBA_SHIPPING_ASSERT(x)
+	#define UBA_ASSERTF(x, fmt, ...) UBA_SHIPPING_ASSERTF(x, fmt, ##__VA_ARGS__)
 	#define UBA_ASSERT_MESSAGEBOX 0
 	inline constexpr bool IsDebug = true;
 	#else
