[UBA]
* More hacks to try to catch issue on farm'

diff --git a/Engine/Source/Programs/UnrealBuildAccelerator/Common/Private/UbaFile.cpp b/Engine/Source/Programs/UnrealBuildAccelerator/Common/Private/UbaFile.cpp
--- a/Engine/Source/Programs/UnrealBuildAccelerator/Common/Private/UbaFile.cpp
+++ b/Engine/Source/Programs/UnrealBuildAccelerator/Common/Private/UbaFile.cpp
@@ -149,7 +149,7 @@ namespace uba
 
 		// HACK HACK HACK... want callstack from this error
 		#if PLATFORM_WINDOWS
-		UBA_SHIPPING_ASSERTF(!Contains(fileName, TC("<PathHash>")), TC("ERROR opening file %s for read (%s%s)"), fileName, LastErrorToText(lastError).data, additionalInfo.data);
+		UBA_SHIPPING_ASSERTF(!Contains(fileName, TC("<PathHash>")), TC("ERROR opening file for read"));
 		#endif
 
 		return logger.Error(TC("ERROR opening file %s for read (%s%s)"), fileName, LastErrorToText(lastError).data, additionalInfo.data);
