[UBA]
* Changed so callstack resolving only allows pdb because all other symbol resolving seems broken and return garbage

diff --git a/Engine/Source/Programs/UnrealBuildAccelerator/Common/Private/UbaLogger.cpp b/Engine/Source/Programs/UnrealBuildAccelerator/Common/Private/UbaLogger.cpp
--- a/Engine/Source/Programs/UnrealBuildAccelerator/Common/Private/UbaLogger.cpp
+++ b/Engine/Source/Programs/UnrealBuildAccelerator/Common/Private/UbaLogger.cpp
@@ -627,7 +627,14 @@ namespace uba
 				auto res = SymLoadModuleExW(processHandle, NULL, mod.name.c_str(), NULL, mod.start, (DWORD)mod.size, NULL, 0);
 				if (res != 0 || GetLastError() == ERROR_SUCCESS)
 				{
-					mod.handled = true;
+					// This is very annoying but it can actually succeed in loading symbols even though pdb is not available.
+					// .. and use internal symbols or something? those symbols are just wrong and messed up and we rather
+					// have have the relative addresses printed out!
+					IMAGEHLP_MODULEW64 moduleInfo = {0};
+					moduleInfo.SizeOfStruct = sizeof(moduleInfo);
+					if (SymGetModuleInfoW64(processHandle, res, &moduleInfo))
+						if (*moduleInfo.LoadedPdbName)
+							mod.handled = true;
 					continue;
 				}
 
