[UBA]
* Added support for suppressing error messages for StoreCas
* Changed so file errors are suppressed when doing StoreCas for fetching cache entries

diff --git a/Engine/Source/Programs/UnrealBuildAccelerator/Common/Private/UbaCacheClient.cpp b/Engine/Source/Programs/UnrealBuildAccelerator/Common/Private/UbaCacheClient.cpp
--- a/Engine/Source/Programs/UnrealBuildAccelerator/Common/Private/UbaCacheClient.cpp
+++ b/Engine/Source/Programs/UnrealBuildAccelerator/Common/Private/UbaCacheClient.cpp
@@ -657,7 +657,7 @@ namespace uba
 					}
 					else
 					{
-						if (!m_storage.StoreCasKey(localCasKey, fileNameKey, path.data, CasKeyZero))
+						if (!m_storage.StoreCasKey(localCasKey, fileNameKey, path.data, CasKeyZero, false))
 							m_logger.Info(TC("Failed to store input file %s when fetching from cache"), path.data);
 					}
 					#if UBA_TRACE_FETCH_CACHE
diff --git a/Engine/Source/Programs/UnrealBuildAccelerator/Common/Private/UbaFile.cpp b/Engine/Source/Programs/UnrealBuildAccelerator/Common/Private/UbaFile.cpp
--- a/Engine/Source/Programs/UnrealBuildAccelerator/Common/Private/UbaFile.cpp
+++ b/Engine/Source/Programs/UnrealBuildAccelerator/Common/Private/UbaFile.cpp
@@ -146,16 +146,10 @@ namespace uba
 		if (lastError == ERROR_SHARING_VIOLATION)
 			GetProcessHoldingFile(additionalInfo, fileName);
 		#endif
-
-		// HACK HACK HACK... want callstack from this error
-		#if PLATFORM_WINDOWS
-		UBA_SHIPPING_ASSERTF(!Contains(fileName, TC("<PathHash>")), TC("ERROR opening file for read"));
-		#endif
-
 		return logger.Error(TC("ERROR opening file %s for read (%s%s)"), fileName, LastErrorToText(lastError).data, additionalInfo.data);
 	}
 
-	bool OpenFileSequentialRead(Logger& logger, const tchar* fileName, FileHandle& outHandle, bool fileNotFoundIsError, bool overlapped)
+	bool OpenFileSequentialRead(Logger& logger, const tchar* fileName, FileHandle& outHandle, bool fileNotFoundIsError, bool overlapped, bool showErrors)
 	{
 		u32 dwFlagsAndAttributes = DefaultAttributes();
 		#if PLATFORM_WINDOWS
@@ -177,6 +171,10 @@ namespace uba
 				if (IsDirectory(attr))
 					return false;
 		}
+
+		if (!showErrors)
+			return false;
+
 		return ReportOpenFileError(logger, fileName, lastError);
 	}
 
diff --git a/Engine/Source/Programs/UnrealBuildAccelerator/Common/Private/UbaStorage.cpp b/Engine/Source/Programs/UnrealBuildAccelerator/Common/Private/UbaStorage.cpp
--- a/Engine/Source/Programs/UnrealBuildAccelerator/Common/Private/UbaStorage.cpp
+++ b/Engine/Source/Programs/UnrealBuildAccelerator/Common/Private/UbaStorage.cpp
@@ -1758,7 +1758,7 @@ namespace uba
 		return StoreCasKey(out, fileNameKey, fileName, casKeyOverride);
 	}
 
-	bool StorageImpl::StoreCasKey(CasKey& out, const StringKey& fileNameKey, const tchar* fileName, const CasKey& casKeyOverride)
+	bool StorageImpl::StoreCasKey(CasKey& out, const StringKey& fileNameKey, const tchar* fileName, const CasKey& casKeyOverride, bool showErrors)
 	{
 		FileEntry& fileEntry = GetOrCreateFileEntry(fileNameKey);
 
@@ -1773,12 +1773,20 @@ namespace uba
 
 		// Use OpenFile+GetFileInformationByHandle+close instead of Getting file attributes because it is actually faster on cloud setups (weirdly enough)
 		FileHandle fileHandle;
-		if (!OpenFileSequentialRead(m_logger, fileName, fileHandle))
+		if (!OpenFileSequentialRead(m_logger, fileName, fileHandle, false, false, showErrors))
+		{
+			fileEntry.casKey = CasKeyZero;
+			out = CasKeyZero;
+			return false;
+		}
+
+		if (fileHandle == InvalidFileHandle)
 		{
 			fileEntry.casKey = CasKeyZero;
 			out = CasKeyZero;
 			return true;
 		}
+
 		auto fileGuard = MakeGuard([&](){ CloseFile(fileName, fileHandle); });
 
 		FileBasicInformation info;
diff --git a/Engine/Source/Programs/UnrealBuildAccelerator/Common/Public/UbaFile.h b/Engine/Source/Programs/UnrealBuildAccelerator/Common/Public/UbaFile.h
--- a/Engine/Source/Programs/UnrealBuildAccelerator/Common/Public/UbaFile.h
+++ b/Engine/Source/Programs/UnrealBuildAccelerator/Common/Public/UbaFile.h
@@ -64,7 +64,7 @@ namespace uba
 	bool GetFileInformation(FileInformation& out, Logger& logger, const tchar* fileName);
 
 	bool ReadFile(Logger& logger, const tchar* fileName, FileHandle fileHandle, void* b, u64 bufferLen);
-	bool OpenFileSequentialRead(Logger& logger, const tchar* fileName, FileHandle& outHandle, bool fileNotFoundIsError = true, bool overlapped = false);
+	bool OpenFileSequentialRead(Logger& logger, const tchar* fileName, FileHandle& outHandle, bool fileNotFoundIsError = true, bool overlapped = false, bool showErrors = true);
 	bool FileExists(Logger& logger, const tchar* fileName, u64* outSize = nullptr, u32* outAttributes = nullptr, u64* lastWriteTime = nullptr);
 	bool SetFilePointer(Logger& logger, const tchar* fileName, FileHandle handle, u64 position);
 	bool SetEndOfFile(Logger& logger, const tchar* fileName, FileHandle handle, u64 size);
diff --git a/Engine/Source/Programs/UnrealBuildAccelerator/Common/Public/UbaStorage.h b/Engine/Source/Programs/UnrealBuildAccelerator/Common/Public/UbaStorage.h
--- a/Engine/Source/Programs/UnrealBuildAccelerator/Common/Public/UbaStorage.h
+++ b/Engine/Source/Programs/UnrealBuildAccelerator/Common/Public/UbaStorage.h
@@ -184,7 +184,7 @@ namespace uba
 		CasKey CalculateCasKey(const tchar* fileName, FileHandle fileHandle, u64 fileSize, bool storeCompressed);
 		CasKey CalculateCasKey(u8* fileMem, u64 fileSize, bool storeCompressed);
 		bool StoreCasKey(CasKey& out, const tchar* fileName, const CasKey& casKeyOverride);
-		bool StoreCasKey(CasKey& out, const StringKey& fileNameKey, const tchar* fileName, const CasKey& casKeyOverride);
+		bool StoreCasKey(CasKey& out, const StringKey& fileNameKey, const tchar* fileName, const CasKey& casKeyOverride, bool showErrors = true);
 		bool IsFileVerified(const StringKey& fileNameKey);
 		void ReportFileInfoWeak(const StringKey& fileNameKey, u64 verifiedLastWriteTime, u64 verifiedSize);
 		bool HasBeenSeen(const CasKey& casKey);
