[UBA]
* Hardened code that parses binaries on win32 to find explicit dependencies

diff --git a/Engine/Source/Programs/UnrealBuildAccelerator/Common/Private/UbaSession.cpp b/Engine/Source/Programs/UnrealBuildAccelerator/Common/Private/UbaSession.cpp
--- a/Engine/Source/Programs/UnrealBuildAccelerator/Common/Private/UbaSession.cpp
+++ b/Engine/Source/Programs/UnrealBuildAccelerator/Common/Private/UbaSession.cpp
@@ -1493,11 +1493,13 @@ namespace uba
 
 		StringBuffer<> errorStr;
 		BinaryInfo info;
-		ParseBinary(temp3, ToView(applicationDir), info, [&](const tchar* importName, bool isKnown, const char* const* importLoaderPaths)
+		if (!ParseBinary(temp3, ToView(applicationDir), info, [&](const tchar* importName, bool isKnown, const char* const* importLoaderPaths)
 			{
 				if (result && !isKnown)
 					result = CopyImports(out, importName, applicationDir, applicationDirEnd, handledImports, importLoaderPaths);
-			}, errorStr);
+			}, errorStr))
+			return m_logger.Error(TC("Failed to parse binary %s for imports"), temp3.data);
+
 		if (errorStr.count)
 			return m_logger.Error(errorStr.data);
 
diff --git a/Engine/Source/Programs/UnrealBuildAccelerator/Core/Public/UbaWinBinDependencyParser.h b/Engine/Source/Programs/UnrealBuildAccelerator/Core/Public/UbaWinBinDependencyParser.h
--- a/Engine/Source/Programs/UnrealBuildAccelerator/Core/Public/UbaWinBinDependencyParser.h
+++ b/Engine/Source/Programs/UnrealBuildAccelerator/Core/Public/UbaWinBinDependencyParser.h
@@ -11,6 +11,7 @@
 #define True_CreateFileMappingW ::CreateFileMappingW
 #define True_MapViewOfFile ::MapViewOfFile
 #define True_UnmapViewOfFile ::UnmapViewOfFile
+#define True_GetFileSize ::GetFileSize
 #endif
 
 namespace uba
@@ -106,9 +107,9 @@ namespace uba
 
 	DWORD UbaExceptionHandlerAllowAccessViolation(DWORD code, _EXCEPTION_POINTERS* info, const tchar* fileName);
 
-	inline bool FindImportsInMem(const tchar* fileName, void* mem, const Function<void(const tchar* import, bool isKnown, const char* const* importLoaderPaths)>& func)
+	inline bool FindImportsInMem(const tchar* fileName, void* mem, u64 memSize, const Function<void(const tchar* import, bool isKnown, const char* const* importLoaderPaths)>& func)
 	{
-		auto hdrs = (PIMAGE_NT_HEADERS)(PCHAR(mem) + PIMAGE_DOS_HEADER(mem)->e_lfanew);   
+		auto hdrs = (PIMAGE_NT_HEADERS)(PCHAR(mem) + PIMAGE_DOS_HEADER(mem)->e_lfanew);
 		auto pSech=IMAGE_FIRST_SECTION(hdrs);
 		__try
 		{
@@ -117,9 +118,16 @@ namespace uba
 				return true; // No import table
 
 			auto importDesc = (PIMAGE_IMPORT_DESCRIPTOR)((DWORD_PTR)mem + GetOffset(dataDir.VirtualAddress, pSech, hdrs));
-			while(importDesc->Name != NULL)
+
+			while(true)
 			{
-				const char* name = (PCHAR)((DWORD_PTR)mem + GetOffset(importDesc->Name, pSech, hdrs));
+				if (u64((u8*)importDesc - (u8*)mem) + sizeof(IMAGE_IMPORT_DESCRIPTOR) > memSize)
+					return false;
+				if (importDesc->Name == NULL)
+					break;
+				char* name = (PCHAR)((DWORD_PTR)mem + GetOffset(importDesc->Name, pSech, hdrs));
+				if (u64((u8*)importDesc - (u8*)name) + 2 > memSize)
+					return false;
 				StringBuffer<256> wname;
 				wname.Append(name);
 				func(wname.data, IsKnownSystemFile(wname.data), nullptr);
@@ -148,8 +156,12 @@ namespace uba
 		void* mem = True_MapViewOfFile(fileMapping, FILE_MAP_READ, 0, 0, 0);
 		if (!mem)
 			return false;
+		DWORD fileSize = True_GetFileSize(fileHandle, NULL);
+		if (fileSize == INVALID_FILE_SIZE)
+			return false;
+
 		auto unmap = MakeGuard([&]() { True_UnmapViewOfFile(mem); });
 
-		return FindImportsInMem(filePath.data, mem, func);
+		return FindImportsInMem(filePath.data, mem, fileSize, func);
 	}
 }
diff --git a/Engine/Source/Programs/UnrealBuildAccelerator/Detours/Private/Windows/UbaDetoursFunctionsKernelBase.inl b/Engine/Source/Programs/UnrealBuildAccelerator/Detours/Private/Windows/UbaDetoursFunctionsKernelBase.inl
--- a/Engine/Source/Programs/UnrealBuildAccelerator/Detours/Private/Windows/UbaDetoursFunctionsKernelBase.inl
+++ b/Engine/Source/Programs/UnrealBuildAccelerator/Detours/Private/Windows/UbaDetoursFunctionsKernelBase.inl
@@ -2743,7 +2743,7 @@ HMODULE Recursive_LoadLibraryExW(const StringView& filePath, LPCWSTR originalNam
 					importedModules.emplace_back(import, isKnown);
 			}, error))
 		{
-			UBA_ASSERTF(false, L"Failed to find imports for binary %ls (%ls)", filePath.data, originalName);
+			FatalError(9887, L"Failed to find imports for binary %s (%s)", filePath.data, originalName);
 		}
 	}
 	for (auto& importedModule : importedModules)
