[UBA]
* Changed code from GetFileSize to GetFileSizeEx just to make bughawk happy (we don't expect an exe or dll to be more than 4gb, that is not allowed on windows)

diff --git a/Engine/Source/Programs/UnrealBuildAccelerator/Core/Public/UbaWinBinDependencyParser.h b/Engine/Source/Programs/UnrealBuildAccelerator/Core/Public/UbaWinBinDependencyParser.h
--- a/Engine/Source/Programs/UnrealBuildAccelerator/Core/Public/UbaWinBinDependencyParser.h
+++ b/Engine/Source/Programs/UnrealBuildAccelerator/Core/Public/UbaWinBinDependencyParser.h
@@ -11,7 +11,7 @@
 #define True_CreateFileMappingW ::CreateFileMappingW
 #define True_MapViewOfFile ::MapViewOfFile
 #define True_UnmapViewOfFile ::UnmapViewOfFile
-#define True_GetFileSize ::GetFileSize
+#define True_GetFileSizeEx ::GetFileSizeEx
 #endif
 
 namespace uba
@@ -156,12 +156,12 @@ namespace uba
 		void* mem = True_MapViewOfFile(fileMapping, FILE_MAP_READ, 0, 0, 0);
 		if (!mem)
 			return false;
-		DWORD fileSize = True_GetFileSize(fileHandle, NULL);
-		if (fileSize == INVALID_FILE_SIZE)
+		LARGE_INTEGER fileSize;
+		if (!True_GetFileSizeEx(fileHandle, &fileSize))
 			return false;
 
 		auto unmap = MakeGuard([&]() { True_UnmapViewOfFile(mem); });
 
-		return FindImportsInMem(filePath.data, mem, fileSize, func);
+		return FindImportsInMem(filePath.data, mem, fileSize.QuadPart, func);
 	}
 }
