[UBA]
* Fixed bug in new hardened win32 binary import parser

diff --git a/Engine/Source/Programs/UnrealBuildAccelerator/Core/Public/UbaWinBinDependencyParser.h b/Engine/Source/Programs/UnrealBuildAccelerator/Core/Public/UbaWinBinDependencyParser.h
--- a/Engine/Source/Programs/UnrealBuildAccelerator/Core/Public/UbaWinBinDependencyParser.h
+++ b/Engine/Source/Programs/UnrealBuildAccelerator/Core/Public/UbaWinBinDependencyParser.h
@@ -126,7 +126,7 @@ namespace uba
 				if (importDesc->Name == NULL)
 					break;
 				char* name = (PCHAR)((DWORD_PTR)mem + GetOffset(importDesc->Name, pSech, hdrs));
-				if (u64((u8*)importDesc - (u8*)name) + 2 > memSize)
+				if (u64((u8*)name - (u8*)mem) + 2 > memSize)
 					return false;
 				StringBuffer<256> wname;
 				wname.Append(name);
