[UBA]
* Fixed bug in Rpc_WriteLogf when string had 0 length (which is valid)

diff --git a/Engine/Source/Programs/UnrealBuildAccelerator/Detours/Private/UbaDetoursShared.cpp b/Engine/Source/Programs/UnrealBuildAccelerator/Detours/Private/UbaDetoursShared.cpp
--- a/Engine/Source/Programs/UnrealBuildAccelerator/Detours/Private/UbaDetoursShared.cpp
+++ b/Engine/Source/Programs/UnrealBuildAccelerator/Detours/Private/UbaDetoursShared.cpp
@@ -367,7 +367,7 @@ namespace uba
 		va_start(arg, format);
 		tchar buffer[1024];
 		int count = Tvsprintf_s(buffer, 1024, format, arg);
-		if (count <= 0)
+		if (count < 0)
 		{
 			TStrcpy_s(buffer, 1024, format);
 			count = int(TStrlen(buffer));
