[UBA]
* Fixed detoured ImageGetDigestStream not handling .tls segment

diff --git a/Engine/Source/Programs/UnrealBuildAccelerator/Detours/Private/Windows/UbaDetoursFunctionsImagehlp.inl b/Engine/Source/Programs/UnrealBuildAccelerator/Detours/Private/Windows/UbaDetoursFunctionsImagehlp.inl
--- a/Engine/Source/Programs/UnrealBuildAccelerator/Detours/Private/Windows/UbaDetoursFunctionsImagehlp.inl
+++ b/Engine/Source/Programs/UnrealBuildAccelerator/Detours/Private/Windows/UbaDetoursFunctionsImagehlp.inl
@@ -8,6 +8,8 @@ ImageGetDigestStreamFunc* True_ImageGetDigestStream;
 
 BOOL Detoured_ImageGetDigestStream(HANDLE FileHandle, DWORD DigestLevel, DIGEST_FUNCTION DigestFunction, DIGEST_HANDLE DigestHandle)
 {
+	DEBUG_LOG_DETOURED(L"ImageGetDigestStream", L"%llu (%s)", u64(FileHandle), HandleToName(FileHandle));
+
 	// This function is _very_ limited and does not cover all kinds of binaries. Only here to handle the dlls used by ShaderCompileWorker
 	// Expects CERT_PE_IMAGE_DIGEST_ALL_IMPORT_INFO | CERT_PE_IMAGE_DIGEST_RESOURCES
 	UBA_ASSERT(DigestLevel == (DWORD)(0x04 | 0x02));
@@ -20,15 +22,15 @@ BOOL Detoured_ImageGetDigestStream(HANDLE FileHandle, DWORD DigestLevel, DIGEST_
 	if (fileSize == INVALID_FILE_SIZE || fileSize < sizeof(IMAGE_DOS_HEADER))
 		return ErrorInvalidParameter();
 
-	auto mappingHandle = CreateFileMappingW(FileHandle, NULL, PAGE_READONLY, 0, 0, NULL);
+	auto mappingHandle = ::CreateFileMappingW(FileHandle, NULL, PAGE_READONLY, 0, 0, NULL);
 	if (mappingHandle == INVALID_HANDLE_VALUE)
 		return ErrorInvalidParameter();
-	auto hg = MakeGuard([&]() { CloseHandle(mappingHandle); });
+	auto hg = MakeGuard([&]() { ::CloseHandle(mappingHandle); });
 
-	auto mem = (BYTE*)MapViewOfFile(mappingHandle, FILE_MAP_COPY, 0, 0, 0);
+	auto mem = (BYTE*)::MapViewOfFile(mappingHandle, FILE_MAP_COPY, 0, 0, 0);
 	if (!mem)
 		return ErrorInvalidParameter();
-	auto mg = MakeGuard([&]() { UnmapViewOfFile(mem); });
+	auto mg = MakeGuard([&]() { ::UnmapViewOfFile(mem); });
 
 	auto& dosHeader = *(IMAGE_DOS_HEADER*)mem;
 	if (dosHeader.e_magic != IMAGE_DOS_SIGNATURE)
@@ -85,6 +87,7 @@ BOOL Detoured_ImageGetDigestStream(HANDLE FileHandle, DWORD DigestLevel, DIGEST_
 	ReportSection("_RDATA");
 	ReportSection(".debug");
 	ReportSection(".didat");
+	ReportSection(".tls");
 	ReportSection(".rsrc");
 	ReportSection(".reloc");
 	return true;
