[UBA]
* Added commented out unit test TestImageDigestStream. Submitted because we might want to revisit this if dxcompiler.dll or friends change and something breaks

diff --git a/Engine/Source/Programs/UnrealBuildAccelerator/Test/Private/UbaTestAll.h b/Engine/Source/Programs/UnrealBuildAccelerator/Test/Private/UbaTestAll.h
--- a/Engine/Source/Programs/UnrealBuildAccelerator/Test/Private/UbaTestAll.h
+++ b/Engine/Source/Programs/UnrealBuildAccelerator/Test/Private/UbaTestAll.h
@@ -53,6 +53,7 @@ namespace uba
 		UBA_TEST(TestFileMappingBuffer) \
 		UBA_TEST(TestRemoteDirectoryTable) \
 		UBA_TEST(TestThreads) \
+		UBA_TEST(TestImageDigestStream) \
 		UBA_TEST_EXTRA \
 
 #define UBA_POSIX_TESTS \
diff --git a/Engine/Source/Programs/UnrealBuildAccelerator/Test/Private/UbaTestBasics.cpp b/Engine/Source/Programs/UnrealBuildAccelerator/Test/Private/UbaTestBasics.cpp
--- a/Engine/Source/Programs/UnrealBuildAccelerator/Test/Private/UbaTestBasics.cpp
+++ b/Engine/Source/Programs/UnrealBuildAccelerator/Test/Private/UbaTestBasics.cpp
@@ -15,6 +15,16 @@
 #include "UbaTimer.h"
 #include "UbaDirectoryIterator.h"
 
+#if 0 // PLATFORM_WINDOWS
+#define DEBUG_LOG_DETOURED(...)
+namespace uba
+{
+	#include "../../Detours/Private/Windows/UbaDetoursFunctionsImagehlp.inl"
+}
+#include "ImageHlp.h"
+#pragma comment (lib, "imagehlp.lib")
+#endif
+
 #define VA_ARGS(...) , ##__VA_ARGS__
 #define UBA_TEST_CHECK(expr, fmt, ...) if (!(expr)) return logger.Error(TC(fmt) VA_ARGS(__VA_ARGS__));
 
@@ -650,4 +660,29 @@ namespace uba
 		traverseDone.Set();
 		return true;
 	}
+
+	bool TestImageDigestStream(LoggerWithWriter& logger, const StringBufferBase& rootDir)
+	{
+		#if 0//PLATFORM_WINDOWS
+
+		const tchar* fileName = TC("e:\\dev\\fn\\Engine\\Binaries\\ThirdParty\\ShaderConductor\\Win64\\dxcompiler.dll");
+		HANDLE file = ::CreateFileW(fileName, GENERIC_READ, FILE_SHARE_READ, nullptr, OPEN_EXISTING, 0, nullptr);
+
+		auto digestHash = [](DIGEST_HANDLE Handle, PBYTE Data, DWORD Length) -> BOOL
+		{
+			printf("SEG: %u\n", Length);
+			return true;
+		};
+
+		printf("Native:\n");
+		if (!ImageGetDigestStream(file, CERT_PE_IMAGE_DIGEST_ALL_IMPORT_INFO | CERT_PE_IMAGE_DIGEST_RESOURCES, digestHash, nullptr))
+			return logger.Error(TC("ImageGetDigestStream failed"))
+
+		printf("Detoured:\n");
+		if (!Detoured_ImageGetDigestStream(file, CERT_PE_IMAGE_DIGEST_ALL_IMPORT_INFO | CERT_PE_IMAGE_DIGEST_RESOURCES, digestHash, nullptr))
+			return logger.Error(TC("Detoured_ImageGetDigestStream failed"))
+
+		#endif
+		return true;
+	}
 }
