[UBA]
* Implemented two-stage download of cache. First stage checks if there is a match and second actually downloads data. This will make it possible to quicker find the cache misses
* Fixed so cache client calculates max cas key offset and path offset so fetch cas/path tables are only called max once each per cache request.
* Changed so scheduler uses two-stage download and fetch test result first in order to find cache misses as early as possible without having to wait for downloads to finish.
* Changed so message priority is set to normal for fetch entries if there are 0 active downloads

diff --git a/Engine/Source/Programs/UnrealBuildAccelerator/Common/Private/UbaCacheClient.cpp b/Engine/Source/Programs/UnrealBuildAccelerator/Common/Private/UbaCacheClient.cpp
--- a/Engine/Source/Programs/UnrealBuildAccelerator/Common/Private/UbaCacheClient.cpp
+++ b/Engine/Source/Programs/UnrealBuildAccelerator/Common/Private/UbaCacheClient.cpp
@@ -37,6 +37,8 @@
 
 namespace uba
 {
+	////////////////////////////////////////////////////////////////////////////////////////////////////
+
 	struct CacheClient::DowngradedLogger : public LoggerWithWriter
 	{
 		DowngradedLogger(Atomic<bool>& c, LogWriter& writer, const tchar* prefix) : LoggerWithWriter(writer, prefix), connected(c) {}
@@ -59,6 +61,7 @@ namespace uba
 		table.GetValueAsBool(useCacheHit, TC("UseCacheHit"));
 	}
 
+	////////////////////////////////////////////////////////////////////////////////////////////////////
 
 	struct CacheClient::Bucket
 	{
@@ -90,6 +93,57 @@ namespace uba
 		Atomic<u32> casKeyTableSizeMaxToSend;
 	};
 
+	////////////////////////////////////////////////////////////////////////////////////////////////////
+
+	struct CacheClient::FetchScope
+	{
+		FetchScope(Session& session, bool& success, u8* memory, const tchar* desc) : m_session(session), m_success(success), m_memory(memory)
+		,	storageStatsScope(storageStats)
+		,	kernelStatsScope(kernelStats)
+		{
+			m_fetchId = m_session.CreateProcessId();
+			m_session.GetTrace().CacheBeginFetch(m_fetchId, desc);
+		}
+
+		~FetchScope()
+		{
+			BinaryWriter writer(m_memory, 0, SendMaxSize);
+			cacheStats.Write(writer);
+			storageStats.Write(writer);
+			kernelStats.Write(writer);
+			m_session.GetTrace().CacheEndFetch(m_fetchId, m_success, m_memory, writer.GetPosition());
+
+			KernelStats::GetGlobal().Add(kernelStats);
+			m_session.GetStorage().AddStats(storageStats);
+		}
+
+		Session& m_session;
+		bool& m_success;
+		u8* m_memory;
+		u32 m_fetchId;
+
+		CacheFetchStats cacheStats;
+		StorageStats storageStats;
+		KernelStats kernelStats;
+
+		StorageStatsScope storageStatsScope;
+		KernelStatsScope kernelStatsScope;
+	};
+
+	////////////////////////////////////////////////////////////////////////////////////////////////////
+
+	CacheClient::FetchContext::FetchContext(const ProcessStartInfo& i, Session& s, RootsHandle rh, MessagePriority p) : info(i), priority(p)
+	{
+		s.PopulateLocalToIndexRoots(rootPaths, rh);
+	}
+
+	CacheClient::FetchContext::FetchContext(const ProcessStartInfo& i, const RootPaths& rp) : info(i), priority(HasPriority)
+	{
+		rootPaths = rp;
+	}
+
+	////////////////////////////////////////////////////////////////////////////////////////////////////
+
 	CacheClient::CacheClient(const CacheClientCreateInfo& info)
 	:	m_logger(info.writer, TC("UbaCacheClient"))
 	,	m_storage(info.storage)
@@ -492,14 +546,28 @@ namespace uba
 
 	bool CacheClient::FetchFromCache(CacheResult& outResult, RootsHandle rootsHandle, u32 bucketId, const ProcessStartInfo& info)
 	{
-		RootPaths rootPaths;
-		m_session.PopulateLocalToIndexRoots(rootPaths, rootsHandle);
-		return FetchFromCache(outResult, rootPaths, bucketId, info);
+		FetchContext context(info, m_session, rootsHandle, HasPriority);
+		if (!FetchEntryFromCache(outResult, context, bucketId))
+			return false;
+		if (!outResult.hit)
+			return false;
+		return FetchFilesFromCache(outResult, context);
 	}
 
 	bool CacheClient::FetchFromCache(CacheResult& outResult, const RootPaths& rootPaths, u32 bucketId, const ProcessStartInfo& info)
+	{
+		FetchContext context(info, rootPaths);
+		if (!FetchEntryFromCache(outResult, context, bucketId))
+			return false;
+		if (!outResult.hit)
+			return false;
+		return FetchFilesFromCache(outResult, context);
+	}
+
+	bool CacheClient::FetchEntryFromCache(CacheResult& outResult, FetchContext& context, u32 bucketId)
 	{
 		outResult.hit = false;
+		outResult.logLines.clear();
 
 		if (!m_connected)
 			return false;
@@ -510,17 +578,8 @@ namespace uba
 		TrackWorkScope tws;
 		#endif
 
-		CacheFetchStats cacheStats;
-		StorageStats storageStats;
-		KernelStats kernelStats;
-		auto kg = MakeGuard([&]() { KernelStats::GetGlobal().Add(kernelStats); m_storage.AddStats(storageStats); });
-
-		StorageStatsScope __(storageStats);
-		KernelStatsScope _(kernelStats);
-
-		CasKey cmdKey = GetCmdKey(rootPaths, info, m_reportCacheKey, bucketId);
-		if (cmdKey == CasKeyZero)
-			return false;
+		auto& rootPaths = context.rootPaths;
+		auto& info = context.info;
 
 #if PLATFORM_MAC
 		u8* memory = (u8*)malloc(SendMaxSize);
@@ -529,19 +588,14 @@ namespace uba
 		u8 memory[SendMaxSize];
 #endif
 
-		u32 fetchId = m_session.CreateProcessId();
-		m_session.GetTrace().CacheBeginFetch(fetchId, info.GetDescription());
 		bool success = false;
-		auto tg = MakeGuard([&]()
-			{
-				BinaryWriter writer(memory, 0, SendMaxSize);
-				cacheStats.Write(writer);
-				storageStats.Write(writer);
-				kernelStats.Write(writer);
-				m_session.GetTrace().CacheEndFetch(fetchId, success, memory, writer.GetPosition());
-			});
+		FetchScope scope(m_session, success, memory, info.GetDescription());
+		auto& cacheStats = scope.cacheStats;
+
 
-		auto subtractTimer = MakeGuard([&](){ cacheStats.testEntry.time -= (cacheStats.fetchCasTable.time + cacheStats.normalizeFile.time); });
+		CasKey cmdKey = GetCmdKey(rootPaths, info, m_reportCacheKey, bucketId);
+		if (cmdKey == CasKeyZero)
+			return false;
 
 		BinaryReader reader(memory, 0, SendMaxSize);
 
@@ -555,7 +609,7 @@ namespace uba
 			TimerScope ts(cacheStats.fetchEntries);
 			// Fetch entries.. server will provide as many as fits. TODO: Should it be possible to ask for more entries?
 			StackBinaryWriter<32> writer;
-			NetworkMessage msg(m_client, CacheServiceId, CacheMessageType_FetchEntries, writer, HasPriority);
+			NetworkMessage msg(m_client, CacheServiceId, CacheMessageType_FetchEntries, writer, context.priority);
 			writer.Write7BitEncoded(MakeId(bucket.id));
 			writer.WriteCasKey(cmdKey);
 			if (!msg.Send(reader))
@@ -606,9 +660,6 @@ namespace uba
 					cachedIsMatch = &insres.first->second;
 				}
 
-				if (!FetchCasTable(tws, bucket, cacheStats, casKeyOffset))
-					return false;
-
 				StringBuffer<MaxPath> path;
 				if (!GetLocalPathAndCasKey(bucket, rootPaths, path, cacheCasKey, bucket.serverCasKeyTable, bucket.serverPathTable, casKeyOffset))
 					return false;
@@ -678,28 +729,78 @@ namespace uba
 			};
 
 
-		struct Range
+		u64 sharedSize = reader.Read7BitEncoded();
+		const u8* sharedPos = reader.GetPositionData();
+		reader.Skip(sharedSize);
+
+		u64 sharedLogLinesSize = reader.Read7BitEncoded();
+		const u8* sharedLogLines = reader.GetPositionData();
+		reader.Skip(sharedLogLinesSize);
+
+		// Calculate max caskey offset and fetch tables (faster to do in one call)
+		u64 maxCasKeyOffset = 0;
 		{
-			u32 begin;
-			u32 end;
-		};
-		Vector<Range> sharedMatchingRanges;
+			BinaryReader sharedReader(sharedPos, 0, sharedSize);
+			while (sharedReader.GetLeft())
+				maxCasKeyOffset = Max(maxCasKeyOffset, sharedReader.Read7BitEncoded());
+
+			BinaryReader entryReader(reader.GetPositionData(), 0, reader.GetLeft());
+			for (u32 i=0; i!=entryCount; ++i)
+			{
+				entryReader.Read7BitEncoded(); // Id
+				u64 extraSize = entryReader.Read7BitEncoded();
+				BinaryReader extraReader(entryReader.GetPositionData(), 0, extraSize);
+				while (extraReader.GetLeft())
+					maxCasKeyOffset = Max(maxCasKeyOffset, extraReader.Read7BitEncoded());
+				entryReader.Skip(extraSize);
+				entryReader.Skip(entryReader.Read7BitEncoded()); // Skip ranges
+				entryReader.Skip(entryReader.Read7BitEncoded()); // Skip outs
+				entryReader.ReadByte(); // Skip log lines type
+			}
+			if (!FetchCasTable(tws, bucket, cacheStats, u32(maxCasKeyOffset)))
+				return false;
+		}
+
+		// Calculate max path offset and fetch path table
+		{
+			u64 maxPathTableOffset = 0;
+			BinaryReader casKeyReader(bucket.serverCasKeyTable.GetMemory(), 0, bucket.serverCasKeyTableSize);
+			BinaryReader sharedReader(sharedPos, 0, sharedSize);
+			while (sharedReader.GetLeft())
+			{
+				casKeyReader.SetPosition(sharedReader.Read7BitEncoded());
+				maxPathTableOffset = Max(maxPathTableOffset, casKeyReader.Read7BitEncoded());
+			}
 
-		const u8* sharedLogLines;
-		u64 sharedLogLinesSize;
+			BinaryReader entryReader(reader.GetPositionData(), 0, reader.GetLeft());
+			for (u32 i=0; i!=entryCount; ++i)
+			{
+				entryReader.Read7BitEncoded(); // Id
+				u64 extraSize = entryReader.Read7BitEncoded();
+				BinaryReader extraReader(entryReader.GetPositionData(), 0, extraSize);
+				while (extraReader.GetLeft())
+				{
+					casKeyReader.SetPosition(extraReader.Read7BitEncoded());
+					maxPathTableOffset = Max(maxPathTableOffset, casKeyReader.Read7BitEncoded());
+				}
+				entryReader.Skip(extraSize);
+				entryReader.Skip(entryReader.Read7BitEncoded()); // Skip ranges
+				entryReader.Skip(entryReader.Read7BitEncoded()); // Skip outs
+				entryReader.ReadByte(); // Skip log lines type
+			}
+			if (!FetchPathTable(tws, bucket, cacheStats, u32(maxPathTableOffset)))
+				return false;
+		}
+
+		struct Range { u32 begin; u32 end; };
+		Vector<Range> sharedMatchingRanges;
 
 		// Create ranges out of shared offsets that matches local state
 		{
 			UBA_TRACE_FETCH_HINT("TestSharedMatch")
 			TimerScope ts(cacheStats.testEntry);
-			u64 sharedSize = reader.Read7BitEncoded();
-
-			BinaryReader sharedReader(reader.GetPositionData(), 0, sharedSize);
-			reader.Skip(sharedSize);
 
-			sharedLogLinesSize = reader.Read7BitEncoded();
-			sharedLogLines = reader.GetPositionData();
-			reader.Skip(sharedLogLinesSize);
+			BinaryReader sharedReader(sharedPos, 0, sharedSize);
 
 			u32 rangeBegin = 0;
 
@@ -752,7 +853,6 @@ namespace uba
 			UBA_TRACE_FETCH_HINT("TestEntriesMatch")
 			--cacheStats.testEntry.count; // Remove the shared one
 
-			BinaryReader entryReader(reader.GetPositionData(), 0, reader.GetLeft());
 			u32 entryIndex = 0;
 			for (; entryIndex!=entryCount; ++entryIndex)
 			{
@@ -803,8 +903,6 @@ namespace uba
 						continue;
 				}
 
-				subtractTimer.Execute();
-
 				#if UBA_TRACE_FETCH_CACHE
 				addStoreKeyHint();
 				#endif
@@ -821,38 +919,9 @@ namespace uba
 				if (!ReportUsedEntry(outResult.logLines, logLinesType == LogLinesType_Owned, bucket, cmdKey, entryId))
 					return false;
 
-				// Fetch output files from cache (and some files need to be "denormalized" before written to disk
-
-				Vector<u32> casKeyOffsets;
-
 				while (outputsReader.GetLeft())
-					casKeyOffsets.push_back(u32(outputsReader.Read7BitEncoded()));
-
-				UBA_TRACE_FETCH_HINT("FetchTableAndFiles")
-
-				if (!casKeyOffsets.empty())
-				{
-					Atomic<bool> fetchSuccess = true;
-					m_storage.m_workManager.ParallelFor(u32(casKeyOffsets.size() - 1), casKeyOffsets, [&](const WorkContext&, auto& it)
-						{
-							u32 casKeyOffset = *it;
-							if (!FetchCasTable(tws, bucket, cacheStats, casKeyOffset))
-							{
-								fetchSuccess = false;
-								return;
-							}
-
-							KernelStatsScope _(kernelStats);
-
-							UBA_TRACE_FETCH_HINT_SCOPE("FetchFile")
-							TimerScope fts(cacheStats.fetchOutput);
-							if (!FetchFile(bucket, rootPaths, info, cacheStats, storageStats, casKeyOffset))
-								fetchSuccess = false;
-
-						}, AsView(TC("CacheFetchOutput")));
-					if (!fetchSuccess)
-						return false;
-				}
+					context.casKeyOffsets.push_back(u32(outputsReader.Read7BitEncoded()));
+				context.bucket = &bucket;
 
 				outResult.hit = true;
 				success = true;
@@ -866,6 +935,60 @@ namespace uba
 		return false;
 	}
 
+	bool CacheClient::FetchFilesFromCache(CacheResult& outResult, FetchContext& context)
+	{
+		UBA_ASSERT(context.bucket);
+		auto& bucket = *context.bucket;
+
+		if (context.casKeyOffsets.empty())
+			return true;
+
+		#if UBA_TRACE_FETCH_CACHE
+		TrackWorkScope tws(m_client, AsView(TC("FetchFromCache")));
+		#else
+		TrackWorkScope tws;
+		#endif
+
+		bool success = false;
+		u8 memory[1024];
+		FetchScope scope(m_session, success, memory, context.info.GetDescription());
+
+		// Make sure caskey/path tables has enough data
+		{
+			u64 maxCasKeyOffset = 0;
+			for (u64 offset : context.casKeyOffsets)
+				maxCasKeyOffset = Max(maxCasKeyOffset, offset);
+			if (!FetchCasTable(tws, bucket, scope.cacheStats, u32(maxCasKeyOffset)))
+				return false;
+
+			u64 maxPathTableOffset = 0;
+			BinaryReader casKeyReader(bucket.serverCasKeyTable.GetMemory(), 0, bucket.serverCasKeyTableSize);
+			for (u64 offset : context.casKeyOffsets)
+			{
+				casKeyReader.SetPosition(offset);
+				maxPathTableOffset = Max(maxPathTableOffset, casKeyReader.Read7BitEncoded());
+			}
+			if (!FetchPathTable(tws, bucket, scope.cacheStats, u32(maxPathTableOffset)))
+				return false;
+		}
+
+		// Go wide on multiple files
+		Atomic<bool> fetchSuccess = true;
+		m_storage.m_workManager.ParallelFor(u32(context.casKeyOffsets.size() - 1), context.casKeyOffsets, [&](const WorkContext&, auto& it)
+			{
+				KernelStatsScope _(scope.kernelStats);
+				UBA_TRACE_FETCH_HINT_SCOPE("FetchFile")
+				TimerScope fts(scope.cacheStats.fetchOutput);
+				if (!FetchFile(*context.bucket, context.rootPaths, context.info, scope.cacheStats, scope.storageStats, *it))
+					fetchSuccess = false;
+
+			}, AsView(TC("CacheFetchOutput")));
+
+		success = fetchSuccess;
+
+		return success;
+	}
+
 	bool CacheClient::FetchFile(Bucket& bucket, const RootPaths& rootPaths, const ProcessStartInfo& info, CacheFetchStats& cacheStats, StorageStats& storageStats, u32 casKeyOffset)
 	{
 		StringBuffer<MaxPath> path;
@@ -1078,16 +1201,6 @@ namespace uba
 		m_session.GetTrace().CacheSummary(writer.GetData(), writer.GetPosition());
 	}
 
-	template<typename T>
-	void AtomicMax(Atomic<T>& target, T value)
-	{
-		T current = target.load(std::memory_order_relaxed);
-		while (current < value && !target.compare_exchange_weak(current, value, std::memory_order_relaxed, std::memory_order_relaxed))
-		{
-			// current is updated by compare_exchange_weak
-		}
-	}
-
 	bool CacheClient::SendPathTable(Bucket& bucket, CacheSendStats& stats, u32 requiredPathTableSize)
 	{
 		AtomicMax(bucket.pathTableSizeMaxToSend, requiredPathTableSize);
@@ -1349,38 +1462,38 @@ namespace uba
 
 	bool CacheClient::FetchCasTable(TrackWorkScope& tws, Bucket& bucket, CacheFetchStats& stats, u32 requiredCasTableOffset)
 	{
-		if (HasEnoughData(bucket, requiredCasTableOffset))
+		if (HasEnoughCasData(bucket, requiredCasTableOffset))
 			return true;
-		return FetchCasTable2(tws, bucket, stats, requiredCasTableOffset);
-	}
 
-	bool CacheClient::FetchCasTable2(TrackWorkScope& tws, Bucket& bucket, CacheFetchStats& stats, u32 requiredCasTableOffset)
-	{
 		UBA_TRACE_FETCH_HINT_SCOPE("FetchCasTable")
 
 		if (requiredCasTableOffset > 256*1024*1024)
 			return m_logger.Warning(TC("Cas table offset %u too large. Cache entry corrupt (Bucket %llu)"), bucket.id);
 
 		TimerScope ts2(stats.fetchCasTable);
-		u32 requiredPathTableOffset = 0;
+
+		SCOPED_FUTEX(bucket.serverCasKeyTableNetworkLock, lock);
+		while (!HasEnoughCasData(bucket, requiredCasTableOffset))
 		{
-			SCOPED_FUTEX(bucket.serverCasKeyTableNetworkLock, lock);
-			while (!HasEnoughCasData(bucket, requiredCasTableOffset, requiredPathTableOffset))
-			{
-				if (!FetchCompactTable(bucket.id, bucket.serverCasKeyTable, requiredCasTableOffset + sizeof(CasKey) + 8, CacheMessageType_FetchCasTable2))
-					return false;
-				bucket.serverCasKeyTableSize = bucket.serverCasKeyTable.GetSize();
-			}
+			if (!FetchCompactTable(bucket.id, bucket.serverCasKeyTable, requiredCasTableOffset + sizeof(CasKey) + 8, CacheMessageType_FetchCasTable2))
+				return false;
+			bucket.serverCasKeyTableSize = bucket.serverCasKeyTable.GetSize();
 		}
+		return true;
+	}
+
+	bool CacheClient::FetchPathTable(TrackWorkScope& tws, Bucket& bucket, CacheFetchStats& stats, u32 requiredPathTableOffset)
+	{
+		if (HasEnoughPathData(bucket, requiredPathTableOffset))
+			return true;
+
+		SCOPED_FUTEX(bucket.serverPathTableNetworkLock, lock);
+		while (!HasEnoughPathData(bucket, requiredPathTableOffset))
 		{
-			SCOPED_FUTEX(bucket.serverPathTableNetworkLock, lock);
-			while (!HasEnoughPathData(bucket, requiredPathTableOffset))
-			{
-				u32 targetSize = requiredPathTableOffset + 200;
-				if (!FetchCompactTable(bucket.id, bucket.serverPathTable, targetSize, CacheMessageType_FetchPathTable2))
-					return false;
-				bucket.serverPathTableSize = bucket.serverPathTable.GetSize();
-			}
+			u32 targetSize = requiredPathTableOffset + 200;
+			if (!FetchCompactTable(bucket.id, bucket.serverPathTable, targetSize, CacheMessageType_FetchPathTable2))
+				return false;
+			bucket.serverPathTableSize = bucket.serverPathTable.GetSize();
 		}
 		return true;
 	}
@@ -1439,17 +1552,7 @@ namespace uba
 		return true;
 	}
 
-	bool CacheClient::HasEnoughData(Bucket& bucket, u32 requiredCasTableOffset)
-	{
-		u32 pathOffset;
-		if (!HasEnoughCasData(bucket, requiredCasTableOffset, pathOffset))
-			return false;
-		if (!HasEnoughPathData(bucket, pathOffset))
-			return false;
-		return true;
-	}
-
-	bool CacheClient::HasEnoughCasData(Bucket& bucket, u32 requiredCasTableOffset, u32& outPathOffset)
+	bool CacheClient::HasEnoughCasData(Bucket& bucket, u32 requiredCasTableOffset)
 	{
 		u32 tableSize = bucket.serverCasKeyTableSize;
 		// CasKeyTable is 7bitEncoded(pathoffset) + CasKey... path table offset is minimum 1 byte
@@ -1457,7 +1560,7 @@ namespace uba
 		if (neededSizeMin > tableSize)
 			return false;
 		BinaryReader r(bucket.serverCasKeyTable.GetMemory(), requiredCasTableOffset, tableSize);
-		outPathOffset = u32(r.Read7BitEncoded());
+		r.Read7BitEncoded();
 		u32 neededSize = u32(r.GetPosition()) + sizeof(CasKey);
 		return neededSize <= tableSize;
 	}
@@ -1685,4 +1788,6 @@ namespace uba
 				m_storage.ReportFileInfoWeak(fileNameKey, e.lastWritten, e.size);
 			});
 	}
+
+	////////////////////////////////////////////////////////////////////////////////////////////////////
 }
diff --git a/Engine/Source/Programs/UnrealBuildAccelerator/Common/Private/UbaScheduler.cpp b/Engine/Source/Programs/UnrealBuildAccelerator/Common/Private/UbaScheduler.cpp
--- a/Engine/Source/Programs/UnrealBuildAccelerator/Common/Private/UbaScheduler.cpp
+++ b/Engine/Source/Programs/UnrealBuildAccelerator/Common/Private/UbaScheduler.cpp
@@ -4,6 +4,7 @@
 #include "UbaApplicationRules.h"
 #include "UbaCacheClient.h"
 #include "UbaConfig.h"
+#include "UbaNetworkMessage.h"
 #include "UbaNetworkServer.h"
 #include "UbaProcess.h"
 #include "UbaProcessStartInfoHolder.h"
@@ -11,8 +12,30 @@
 #include "UbaSessionServer.h"
 #include "UbaStringBuffer.h"
 
+
+// TODO. This is how we want scheduler to work.
+// 
+// 1. Cache should be queried from front of queue.
+//    We need to split up cache in fetching cache entry and fetching cas.
+//    .. if there are 10000 actions and the last is a cache miss, then we want to find this out as early as possible.
+// 2. Local processing should prioritize cache misses first, then query from the back of they queue
+// 
+// At some point cached processes and local will race which is fine.
+//
+
 namespace uba
 {
+	////////////////////////////////////////////////////////////////////////////////////////////////////
+
+	struct Scheduler::CacheFetchInfo : CacheClient::FetchContext
+	{
+		CacheFetchInfo(const ProcessStartInfo& i, Session& s, RootsHandle rh, MessagePriority mp) : CacheClient::FetchContext(i, s, rh, mp) {}
+		CacheResult result;
+		CacheClient* client;
+	};
+
+	////////////////////////////////////////////////////////////////////////////////////////////////////
+
 	struct Scheduler::ProcessStartInfo2 : ProcessStartInfoHolder
 	{
 		ProcessStartInfo2(const ProcessStartInfo& si, const u8* ki, u32 kic)
@@ -27,12 +50,15 @@ namespace uba
 			delete[] knownInputs;
 		}
 
+		CacheFetchInfo* cacheInfo = nullptr;
+
 		const u8* knownInputs;
 		u32 knownInputsCount;
 		float weight = 1.0f;
 		u32 cacheBucketId = 0;
 	};
 
+	////////////////////////////////////////////////////////////////////////////////////////////////////
 
 	struct Scheduler::ExitProcessInfo
 	{
@@ -43,6 +69,7 @@ namespace uba
 		bool isLocal = true;
 	};
 
+	////////////////////////////////////////////////////////////////////////////////////////////////////
 
 	class SkippedProcess : public Process
 	{
@@ -61,6 +88,8 @@ namespace uba
 		ProcessStartInfoHolder startInfo;
 	};
 
+	////////////////////////////////////////////////////////////////////////////////////////////////////
+
 	class CachedProcess : public Process
 	{
 	public:
@@ -79,6 +108,8 @@ namespace uba
 		Vector<ProcessLogLine> logLines;
 	};
 
+	////////////////////////////////////////////////////////////////////////////////////////////////////
+
 	void SchedulerCreateInfo::Apply(const Config& config)
 	{
 		if (const ConfigTable* table = config.GetTable(TC("Scheduler")))
@@ -92,12 +123,13 @@ namespace uba
 		}
 	}
 
+	////////////////////////////////////////////////////////////////////////////////////////////////////
 
 	Scheduler::Scheduler(const SchedulerCreateInfo& info)
 	:	m_session(info.session)
 	,	m_maxLocalProcessors(info.maxLocalProcessors != ~0u ? info.maxLocalProcessors : GetLogicalProcessorCount())
 	,	m_maxParallelCacheQueries(info.maxParallelCacheQueries)
-	,	m_updateThreadLoop(false)
+	,	m_updateThreadWakeup(false)
 	,	m_enableProcessReuse(info.enableProcessReuse)
 	,	m_forceRemote(info.forceRemote)
 	,	m_forceNative(info.forceNative)
@@ -126,17 +158,20 @@ namespace uba
 		m_session.SetRemoteProcessReturnedEvent([this](Process& process) { RemoteProcessReturned(process); });
 		m_session.SetRemoteProcessSlotAvailableEvent([this](bool isCrossArchitecture) { RemoteSlotAvailable(isCrossArchitecture); });
 
+		if (!m_cacheClients.empty())
+			m_cacheWorkManager = new WorkManagerImpl(m_maxParallelCacheQueries, TC("UbaSchedCach"));
+
 		m_loop = true;
-		m_thread.Start([this]() { ThreadLoop(); return 0; }, TC("UbaSchedLoop"));
+		m_updateThread.Start([this]() { ThreadLoop(); return 0; }, TC("UbaSchedLoop"));
 	}
 
 	void Scheduler::Stop()
 	{
 		m_loop = false;
 		m_cancelled = true;
-		m_updateThreadLoop.Set();
-		m_thread.Wait();
-		m_cacheDownloadThreads.clear();
+		m_updateThreadWakeup.Set();
+		m_updateThread.Wait();
+		delete m_cacheWorkManager;
 		m_session.WaitOnAllTasks();
 		SkipAllQueued();
 		Cleanup();
@@ -155,11 +190,14 @@ namespace uba
 		Vector<ProcessStartInfo2*> skipped;
 		SCOPED_FUTEX(m_processEntriesLock, lock);
 		for (auto& entry : m_processEntries)
-			if (entry.status == ProcessStatus_QueuedForCache || entry.status == ProcessStatus_QueuedForRun)
+		{
+			if (entry.runStatus == RunStatus_QueuedForRun && (entry.cacheStatus == CacheStatus_QueuedForTest || entry.cacheStatus == CacheStatus_Failed))
 			{
-				entry.status = ProcessStatus_Skipped;
+				entry.cacheStatus = CacheStatus_Failed;
+				entry.runStatus = RunStatus_Skipped;
 				skipped.push_back(entry.info);
 			}
+		}
 		lock.Leave();
 		for (auto pi : skipped)
 			SkipProcess(*pi);
@@ -170,19 +208,19 @@ namespace uba
 		SCOPED_FUTEX(m_processEntriesLock, lock);
 		for (auto& entry : m_processEntries)
 		{
-			UBA_ASSERTF(entry.status > ProcessStatus_Running, TC("Found processes in queue/running state when stopping scheduler."));
+			UBA_ASSERTF(entry.runStatus > RunStatus_Running, TC("Found processes in queue/running state when stopping scheduler."));
 			delete[] entry.dependencies;
 			delete entry.info;
 		}
 		m_processEntries.clear();
-		m_processEntriesStart = 0;
+		m_processEntriesRunStart = 0;
 		m_session.SetOuterScheduler(nullptr);
 	}
 
 	void Scheduler::SetMaxLocalProcessors(u32 maxLocalProcessors)
 	{
 		m_maxLocalProcessors = maxLocalProcessors;
-		m_updateThreadLoop.Set();
+		m_updateThreadWakeup.Set();
 	}
 
 	void Scheduler::SetAllowDisableRemoteExecution(bool allow)
@@ -239,7 +277,8 @@ namespace uba
 		entry.info = info2;
 		entry.dependencies = dep;
 		entry.dependencyCount = info.dependencyCount;
-		entry.status = useCache ? ProcessStatus_QueuedForCache : ProcessStatus_QueuedForRun;
+		entry.runStatus = RunStatus_QueuedForRun;
+		entry.cacheStatus = useCache ? CacheStatus_QueuedForTest : CacheStatus_Failed;
 		entry.canDetour = canDetour;
 		entry.canExecuteRemotely = canExecuteRemotely;
 		lock.Leave();
@@ -248,7 +287,7 @@ namespace uba
 
 		UpdateQueueCounter(1);
 
-		m_updateThreadLoop.Set();
+		m_updateThreadWakeup.Set();
 		return index;
 	}
 
@@ -282,7 +321,7 @@ namespace uba
 		{
 			if (!entry.canExecuteRemotely)
 				continue;
-			if (entry.status != ProcessStatus_QueuedForRun)
+			if (entry.runStatus != RunStatus_QueuedForRun || entry.cacheStatus != CacheStatus_Failed)
 				continue;
 			++count;
 		}
@@ -303,13 +342,13 @@ namespace uba
 		{
 			if (!entry.canExecuteRemotely)
 				continue;
-			if (entry.status != ProcessStatus_QueuedForRun)
+			if (entry.runStatus != RunStatus_QueuedForRun || entry.cacheStatus != CacheStatus_Failed)
 				continue;
 			bool canRun = true;
 			for (u32 j=0, je=entry.dependencyCount; j!=je; ++j)
 			{
 				auto depIndex = entry.dependencies[j];
-				if (m_processEntries[depIndex].status == ProcessStatus_Success)
+				if (m_processEntries[depIndex].runStatus == RunStatus_Success)
 					continue;
 				canRun = false;
 				break;
@@ -325,9 +364,12 @@ namespace uba
 	{
 		while (m_loop)
 		{
-			if (!m_updateThreadLoop.IsSet())
+			if (!m_updateThreadWakeup.IsSet())
 				break;
 
+			while (RunCacheQuery())
+				;
+
 			while (RunQueuedProcess(true))
 				;
 		}
@@ -346,41 +388,42 @@ namespace uba
 			return;
 
 		SCOPED_FUTEX(m_processEntriesLock, lock);
-		if (m_processEntries[processIndex].status != ProcessStatus_Running)
+		if (m_processEntries[processIndex].runStatus != RunStatus_Running)
 			return;
-		m_processEntries[processIndex].status = ProcessStatus_QueuedForRun;
-		m_processEntriesStart = Min(m_processEntriesStart, processIndex);
+		m_processEntries[processIndex].runStatus = RunStatus_QueuedForRun;
+		m_processEntriesRunStart = Min(m_processEntriesRunStart, processIndex);
 		lock.Leave();
 
 		UpdateQueueCounter(1);
 		UpdateActiveProcessCounter(false, -1);
-		m_updateThreadLoop.Set();
+		m_updateThreadWakeup.Set();
 	}
 
 	void Scheduler::HandleCacheMissed(u32 processIndex)
 	{
 		SCOPED_FUTEX(m_processEntriesLock, lock);
 		auto& entry = m_processEntries[processIndex];
-		UBA_ASSERT(entry.status == ProcessStatus_DownloadingFromCache);
+		UBA_ASSERT((entry.cacheStatus == CacheStatus_Testing || entry.cacheStatus == CacheStatus_Downloading) && entry.runStatus == RunStatus_QueuedForRun);
 
 		--m_activeCacheQueries;
 
+		entry.cacheStatus = CacheStatus_Failed;
+
 		if (m_cancelled)
 		{
-			entry.status = ProcessStatus_Skipped;
+			entry.runStatus = RunStatus_Skipped;
 			auto pi = entry.info;
 			lock.Leave();
 			SkipProcess(*pi);
 			return;
 		}
 
-		entry.status = ProcessStatus_QueuedForRun;
-		m_processEntriesStart = Min(m_processEntriesStart, processIndex);
+		m_processEntriesRunStart = Min(m_processEntriesRunStart, processIndex);
 		lock.Leave();
 
 		UpdateQueueCounter(1);
 		UpdateActiveProcessCounter(false, -1);
-		m_updateThreadLoop.Set();
+		m_updateThreadWakeup.Set();
 	}
 
 	void Scheduler::RemoteSlotAvailable(bool isCrossArchitecture)
@@ -395,7 +438,7 @@ namespace uba
 		u32 count = 0;
 		SCOPED_FUTEX_READ(m_processEntriesLock, lock);
 		for (auto& entry : m_processEntries)
-			count += u32(entry.canExecuteRemotely && entry.status <= ProcessStatus_QueuedForRun);
+			count += u32(entry.canExecuteRemotely && entry.runStatus <= RunStatus_QueuedForRun);
 		lock.Leave();
 		if (count < m_maxLocalProcessors)
 			m_session.DisableRemoteExecution();
@@ -434,24 +477,84 @@ namespace uba
 			}, TC("UbaScheRemFin"));
 	}
 
-	u32 Scheduler::PopProcess(bool isLocal, ProcessStatus& outPrevStatus)
+	u32 Scheduler::PopCache()
+	{
+		bool allFinished = true;
+		u32 entryToProcess = ~0u;
+		auto processEntries = m_processEntries.data();
+
+		for (u32 i=m_processEntriesCacheStart, e=u32(m_processEntries.size()); i!=e; ++i)
+		{
+			auto& entry = processEntries[i];
+			auto cacheStatus = entry.cacheStatus;
+
+			if (cacheStatus >= CacheStatus_Downloading)
+			{
+				if (allFinished)
+					m_processEntriesCacheStart = i;
+				continue;
+			}
+
+			allFinished = false;
+
+			if (cacheStatus == CacheStatus_QueuedForTest)
+			{
+				bool canRun = true;
+				for (u32 j=0, je=entry.dependencyCount; j!=je; ++j)
+				{
+					auto depIndex = entry.dependencies[j];
+					if (depIndex >= m_processEntries.size())
+					{
+						m_session.GetLogger().Error(TC("Found dependency on index %u but there are only %u processes registered"), depIndex, u32(m_processEntries.size()));
+						return ~0u;
+					}
+					auto depStatus = processEntries[depIndex].runStatus;
+					if (depStatus == RunStatus_Failed || depStatus == RunStatus_Skipped)
+					{
+						entry.cacheStatus = CacheStatus_Failed;
+						canRun = false;
+						break;
+					}
+					if (depStatus != RunStatus_Success)
+					{
+						canRun = false;
+						break;
+					}
+				}
+				if (!canRun)
+					continue;
+
+				return i;
+			}
+
+			if (cacheStatus == CacheStatus_Testing)
+				continue;
+
+			UBA_ASSERT(cacheStatus == CacheStatus_QueuedForDownload);
+			if (entryToProcess == ~0u)
+				entryToProcess = i;
+		}
+		return entryToProcess;
+	}
+
+	u32 Scheduler::PopProcess(bool isLocal, RunStatus& outPrevStatus)
 	{
 		bool atMaxLocalWeight = m_activeLocalProcessWeight >= float(m_maxLocalProcessors);
 		auto processEntries = m_processEntries.data();
 		bool allFinished = true;
 
-		for (u32 i=m_processEntriesStart, e=u32(m_processEntries.size()); i!=e; ++i)
+		for (u32 i=m_processEntriesRunStart, e=u32(m_processEntries.size()); i!=e; ++i)
 		{
 			auto& entry = processEntries[i];
-			auto status = entry.status;
+			auto runStatus = entry.runStatus;
 
 			// Check if we can move search starting point forward
-			if (status != ProcessStatus_QueuedForCache && status != ProcessStatus_QueuedForRun)
+			if (runStatus != RunStatus_QueuedForRun)
 			{
 				if (allFinished)
 				{
-					if (status != ProcessStatus_Running && status != ProcessStatus_DownloadingFromCache)
-						m_processEntriesStart = i;
+					if (runStatus != RunStatus_Running)
+						m_processEntriesRunStart = i;
 					else
 						allFinished = false;
 				}
@@ -459,20 +562,18 @@ namespace uba
 			}
 			allFinished = false;
 
+			auto cacheStatus = entry.cacheStatus;
 			
-			if (status == ProcessStatus_DownloadingFromCache)
+			if (cacheStatus != CacheStatus_Failed)
 				continue;
 
-			if (status == ProcessStatus_QueuedForCache)
-			{
-				if (m_activeCacheQueries >= m_maxParallelCacheQueries)
-					continue;
-			}
-			else if (isLocal)
+			UBA_ASSERT(cacheStatus != CacheStatus_Success);
+
+			if (isLocal)
 			{
 				if (m_forceRemote && entry.canExecuteRemotely)
 					continue;
-				if (status == ProcessStatus_QueuedForRun && atMaxLocalWeight)
+				if (runStatus == RunStatus_QueuedForRun && atMaxLocalWeight)
 					continue;
 			}
 			else
@@ -490,13 +591,13 @@ namespace uba
 					m_session.GetLogger().Error(TC("Found dependency on index %u but there are only %u processes registered"), depIndex, u32(m_processEntries.size()));
 					return ~0u;
 				}
-				auto depStatus = processEntries[depIndex].status;
-				if (depStatus == ProcessStatus_Failed || depStatus == ProcessStatus_Skipped)
+				auto depStatus = processEntries[depIndex].runStatus;
+				if (depStatus == RunStatus_Failed || depStatus == RunStatus_Skipped)
 				{
-					entry.status = ProcessStatus_Skipped;
+					entry.runStatus = RunStatus_Skipped;
 					return i;
 				}
-				if (depStatus != ProcessStatus_Success)
+				if (depStatus != RunStatus_Success)
 				{
 					canRun = false;
 					break;
@@ -506,74 +607,128 @@ namespace uba
 			if (!canRun)
 				continue;
 
-			if (status != ProcessStatus_QueuedForCache)
-			{
-				if (isLocal)
-					m_activeLocalProcessWeight += entry.info->weight;
+			if (isLocal)
+				m_activeLocalProcessWeight += entry.info->weight;
 
-				outPrevStatus = entry.status;
-				entry.status = ProcessStatus_Running;
-				return i;
-			}
+			outPrevStatus = entry.runStatus;
+			entry.runStatus = RunStatus_Running;
+			return i;
+		}
 
-			// Clear out existing download threads if they're done
-			FlushFinishedDownloadThreads();
+		return ~0u;
+	}
 
-			++m_activeCacheQueries;
+	bool Scheduler::RunCacheQuery()
+	{
+		SCOPED_FUTEX(m_processEntriesLock, lock);
+		
+		if (m_activeCacheQueries >= m_maxParallelCacheQueries)
+			return false;
+
+		u32 entryToProcess = PopCache();
+		if (entryToProcess == ~0u)
+			return false;
+
+		++m_activeCacheQueries;
 
-			// Start download
-			entry.status = ProcessStatus_DownloadingFromCache;
+		auto& entry = m_processEntries[entryToProcess];
+		auto info = entry.info;
+
+		if (entry.cacheStatus == CacheStatus_QueuedForTest)
+		{
+			entry.cacheStatus = CacheStatus_Testing;
+			lock.Leave();
 
-			m_cacheDownloadThreads.emplace_back([this, processIndex = i, info = entry.info]()
+			m_cacheWorkManager->AddWork([this, info, entryToProcess](const WorkContext&)
 				{
 					ProcessStartInfo2& si = *info;
 
-					CacheResult cacheResult;
+					MessagePriority priority = m_activeCacheDownloads != 0 ? HasPriority : NormalPriority;
+					auto cacheInfo = new CacheFetchInfo(si, m_session, si.rootsHandle, priority);
+
 					bool isHit = false;
 					for (auto cacheClient : m_cacheClients)
 					{
-						if (!cacheClient->FetchFromCache(cacheResult, si.rootsHandle, si.cacheBucketId, si) || !cacheResult.hit)
+						if (!cacheClient->FetchEntryFromCache(cacheInfo->result, *cacheInfo, si.cacheBucketId) || !cacheInfo->result.hit)
 							continue;
 
+						cacheInfo->client = cacheClient;
+						info->cacheInfo = cacheInfo;
+
+						if (m_cancelled)
+							break;
+
 						isHit = true;
+
+						SCOPED_FUTEX(m_processEntriesLock, lock);
+						m_processEntries[entryToProcess].cacheStatus = CacheStatus_QueuedForDownload;
+						--m_activeCacheQueries;
+						break;
+					}
+
+					if (isHit)
+					{
+						m_updateThreadWakeup.Set();
+						return;
+					}
+
+					delete cacheInfo;
+					HandleCacheMissed(entryToProcess);
+					++m_cacheMissCount;
+					UpdateStatusHitMissCount();
+
+				}, 1, TC("UbaSchedTest"));
+		}
+		else
+		{
+			UBA_ASSERT(entry.cacheStatus == CacheStatus_QueuedForDownload);
+			entry.cacheStatus = CacheStatus_Downloading;
+			lock.Leave();
+
+			++m_activeCacheDownloads;
+
+			m_cacheWorkManager->AddWork([this, info, entryToProcess](const WorkContext&)
+				{
+					auto cacheInfo = info->cacheInfo;
+
+					ProcessStartInfo2& si = *info;
+
+					bool success = cacheInfo->client->FetchFilesFromCache(cacheInfo->result, *cacheInfo);
+					--m_activeCacheDownloads;
+
+					if (success)
+					{
 						auto process = new CachedProcess(si);
-						process->logLines.swap(cacheResult.logLines);
+						process->logLines.swap(cacheInfo->result.logLines);
 						ProcessHandle ph(process);
 						ExitProcessInfo exitInfo;
 						exitInfo.scheduler = this;
 						exitInfo.startInfo = info;
 						exitInfo.isLocal = true;
-						exitInfo.processIndex = processIndex;
+						exitInfo.processIndex = entryToProcess;
 						ExitProcess(exitInfo, *process, 0, true);
 						++m_cacheHitCount;
-						break;
 					}
-
-					if (!isHit)
+					else
 					{
-						HandleCacheMissed(processIndex);
+						HandleCacheMissed(entryToProcess);
 						++m_cacheMissCount;
 					}
 
-					StringBuffer<> str;
-					str.Appendf(TC("Hits %u Misses %u"), m_cacheHitCount.load(), m_cacheMissCount.load());
-					m_session.GetTrace().StatusUpdate(1, 6, str.data, LogEntryType_Info, nullptr);
-					return 0;
+					delete cacheInfo;
+					UpdateStatusHitMissCount();
 
-				}, TC("UbaScheDlCach"));
-			continue;
+				}, 1, TC("UbaSchedDlwd"));
 		}
 
-		FlushFinishedDownloadThreads();
-		return ~0u;
+		return true;
 	}
 
-
 	bool Scheduler::RunQueuedProcess(bool isLocal)
 	{
 		while (true)
 		{
-			ProcessStatus prevStatus;
+			RunStatus prevStatus;
 			SCOPED_FUTEX(m_processEntriesLock, lock);
 			u32 indexToRun = PopProcess(isLocal, prevStatus);
 			if (indexToRun == ~0u)
@@ -582,7 +737,7 @@ namespace uba
 			auto& processEntry = m_processEntries[indexToRun];
 			auto info = processEntry.info;
 			bool canDetour = processEntry.canDetour && !m_forceNative;
-			bool wasSkipped = processEntry.status == ProcessStatus_Skipped;
+			bool wasSkipped = processEntry.runStatus == RunStatus_Skipped;
 			lock.Leave();
 
 			UpdateQueueCounter(-1);
@@ -640,15 +795,15 @@ namespace uba
 
 		while (true)
 		{
-			ProcessStatus prevStatus;
+			RunStatus prevStatus;
 			SCOPED_FUTEX(m_processEntriesLock, lock);
 			u32 indexToRun = PopProcess(isLocal, prevStatus);
 			if (indexToRun == ~0u)
 				return false;
-			UBA_ASSERT(prevStatus != ProcessStatus_QueuedForCache);
+			//UBA_ASSERT(prevStatus != RunStatus_QueuedForCache);
 			auto& processEntry = m_processEntries[indexToRun];
 			auto newInfo = processEntry.info;
-			bool wasSkipped = processEntry.status == ProcessStatus_Skipped;
+			bool wasSkipped = processEntry.runStatus == RunStatus_Skipped;
 			lock.Leave();
 
 			UpdateQueueCounter(-1);
@@ -696,6 +851,8 @@ namespace uba
 		if (!si)
 			return;
 
+		//m_session.GetLogger().Error(TC("PROCESS: %u"), info.processIndex);
+
 		ProcessHandle ph;
 		ph.m_process = &process;
 
@@ -712,7 +869,7 @@ namespace uba
 
 		if (isDone)
 		{
-			entry.status = exitCode == 0 ? ProcessStatus_Success : ProcessStatus_Failed;
+			entry.runStatus = exitCode == 0 ? RunStatus_Success : RunStatus_Failed;
 			entry.info = nullptr;
 			entry.dependencies = nullptr;
 		}
@@ -720,14 +877,17 @@ namespace uba
 		{
 			entry.canExecuteRemotely = false;
 			entry.canDetour = exitedResponse != ProcessExitedResponse_RerunNative;
-			entry.status = ProcessStatus_QueuedForRun;
-			m_processEntriesStart = Min(m_processEntriesStart, info.processIndex);
+			entry.runStatus = RunStatus_QueuedForRun;
+			m_processEntriesRunStart = Min(m_processEntriesRunStart, info.processIndex);
 		}
 
 		if (info.isLocal)
 		{
 			if (fromCache)
+			{
+				entry.cacheStatus = CacheStatus_Success;
 				--m_activeCacheQueries;
+			}
 			else
 				m_activeLocalProcessWeight -= si->weight;
 		}
@@ -735,7 +895,7 @@ namespace uba
 		lock.Leave();
 
 		UpdateActiveProcessCounter(info.isLocal, -1);
-		m_updateThreadLoop.Set();
+		m_updateThreadWakeup.Set();
 
 		if (isDone)
 		{
@@ -766,6 +926,13 @@ namespace uba
 		FinishProcess(ph);
 	}
 
+	void Scheduler::UpdateStatusHitMissCount()
+	{
+		StringBuffer<> str;
+		str.Appendf(TC("Hits %u Misses %u"), m_cacheHitCount.load(), m_cacheMissCount.load());
+		m_session.GetTrace().StatusUpdate(1, 6, str.data, LogEntryType_Info, nullptr);
+	}
+
 	void Scheduler::UpdateQueueCounter(int offset)
 	{
 		m_queuedProcesses += u32(offset);
@@ -788,16 +955,6 @@ namespace uba
 		m_session.UpdateProgress(m_totalProcesses, m_finishedProcesses, m_errorCount);
 	}
 
-	void Scheduler::FlushFinishedDownloadThreads()
-	{
-		while (!m_cacheDownloadThreads.empty())
-		{
-			if (!m_cacheDownloadThreads.front().Wait(0))
-				break;
-			m_cacheDownloadThreads.pop_front();
-		}
-	}
-
 	bool Scheduler::EnqueueFromFile(const tchar* yamlFilename, const Function<void(EnqueueProcessInfo&)>& enqueued)
 	{
 		auto& logger = m_session.GetLogger();
@@ -1131,4 +1288,6 @@ namespace uba
 
 		return true;
 	}
+
+	////////////////////////////////////////////////////////////////////////////////////////////////////
 }
diff --git a/Engine/Source/Programs/UnrealBuildAccelerator/Common/Public/UbaCacheClient.h b/Engine/Source/Programs/UnrealBuildAccelerator/Common/Public/UbaCacheClient.h
--- a/Engine/Source/Programs/UnrealBuildAccelerator/Common/Public/UbaCacheClient.h
+++ b/Engine/Source/Programs/UnrealBuildAccelerator/Common/Public/UbaCacheClient.h
@@ -5,6 +5,8 @@
 #include "UbaHash.h"
 #include "UbaLogger.h"
 #include "UbaProcessHandle.h"
+#include "UbaRootPaths.h"
+#include "UbaStats.h"
 
 namespace uba
 {
@@ -15,6 +17,7 @@ namespace uba
 	class RootPaths;
 	class Session;
 	class StorageImpl;
+	enum MessagePriority : u8;
 	struct CacheFetchStats;
 	struct CacheSendStats;
 	struct CasKey;
@@ -24,6 +27,8 @@ namespace uba
 	struct TrackWorkScope;
 	using RootsHandle = u64;
 
+	////////////////////////////////////////////////////////////////////////////////////////////////////
+
 	struct CacheClientCreateInfo
 	{
 		CacheClientCreateInfo(LogWriter& w, StorageImpl& st, NetworkClient& c, Session& se) : writer(w), storage(st), client(c), session(se) {}
@@ -44,12 +49,16 @@ namespace uba
 		const tchar* hint = TC(""); // Hint will show up in cache server log
 	};
 
+	////////////////////////////////////////////////////////////////////////////////////////////////////
+
 	struct CacheResult
 	{
 		bool hit = false;
 		Vector<ProcessLogLine> logLines;
 	};
 
+	////////////////////////////////////////////////////////////////////////////////////////////////////
+
 	class CacheClient
 	{
 	public:
@@ -75,8 +84,27 @@ namespace uba
 
 		void TraceSummary();
 
-	private:
 		struct Bucket;
+
+		struct FetchContext
+		{
+			FetchContext(const ProcessStartInfo& info, Session& session, RootsHandle rootsHandle, MessagePriority priority);
+			FetchContext(const ProcessStartInfo& info, const RootPaths& rootPaths);
+
+			const ProcessStartInfo& info;
+
+			RootPaths rootPaths;
+			Bucket* bucket = nullptr;
+			Vector<u32> casKeyOffsets;
+			MessagePriority priority;
+		};
+
+		bool FetchEntryFromCache(CacheResult& outResult, FetchContext& context, u32 bucketId);
+		bool FetchFilesFromCache(CacheResult& outResult, FetchContext& context);
+
+	private:
+		struct FetchScope;
+
 		u64 MakeId(u32 bucketId);
 
 		// No inlining because they use lots of stack space
@@ -84,12 +112,11 @@ namespace uba
 		UBA_NOINLINE bool SendCasTable(Bucket& bucket, CacheSendStats& stats, u32 requiredCasTableSize);
 		UBA_NOINLINE bool SendCacheEntry(TrackWorkScope& tws, Bucket& bucket, CacheSendStats& stats, const RootPaths& rootPaths, const CasKey& cmdKey, const Map<u32, u32>& inputsStringToCasKey, const Map<u32, u32>& outputsStringToCasKey, const u8* logLines, u64 logLinesSize);
 		UBA_NOINLINE bool SendCacheEntryMessage(BinaryReader& out, Bucket& bucket, CacheSendStats& stats, const CasKey& cmdKey, const Map<u32, u32>& inputsStringToCasKey, const Map<u32, u32>& outputsStringToCasKey, const u8* logLines, u64 logLinesSize);
-		bool FetchCasTable(TrackWorkScope& tws, Bucket& bucket, CacheFetchStats& stats, u32 requiredCasTableOffset);
-		UBA_NOINLINE bool FetchCasTable2(TrackWorkScope& tws, Bucket& bucket, CacheFetchStats& stats, u32 requiredCasTableOffset);
+		UBA_NOINLINE bool FetchCasTable(TrackWorkScope& tws, Bucket& bucket, CacheFetchStats& stats, u32 requiredCasTableOffset);
+		UBA_NOINLINE bool FetchPathTable(TrackWorkScope& tws, Bucket& bucket, CacheFetchStats& stats, u32 requiredPathTableOffset);
 		UBA_NOINLINE bool FetchFile(Bucket& bucket, const RootPaths& rootPaths, const ProcessStartInfo& info, CacheFetchStats& cacheStats, StorageStats& storageStats, u32 casKeyOffset);
 		template<typename TableType> bool FetchCompactTable(u32 bucketId, TableType& table, u32 requiredTableOffset, u8 messageType);
-		bool HasEnoughData(Bucket& bucket, u32 requiredCasTableOffset);
-		bool HasEnoughCasData(Bucket& bucket, u32 requiredCasTableOffset, u32& outPathOffset);
+		bool HasEnoughCasData(Bucket& bucket, u32 requiredCasTableOffset);
 		bool HasEnoughPathData(Bucket& bucket, u32 requiredPathTableOffset);
 		UBA_NOINLINE bool ReportUsedEntry(Vector<ProcessLogLine>& outLogLines, bool ownedLogLines, Bucket& bucket, const CasKey& cmdKey, u32 entryId);
 		bool PopulateLogLines(Vector<ProcessLogLine>& outLogLines, const u8* mem, u64 memLen);
@@ -131,4 +158,6 @@ namespace uba
 		CacheClient& operator=(const CacheClient&) = delete;
 		struct DowngradedLogger;
 	};
+
+	////////////////////////////////////////////////////////////////////////////////////////////////////
 }
\ No newline at end of file
diff --git a/Engine/Source/Programs/UnrealBuildAccelerator/Common/Public/UbaScheduler.h b/Engine/Source/Programs/UnrealBuildAccelerator/Common/Public/UbaScheduler.h
--- a/Engine/Source/Programs/UnrealBuildAccelerator/Common/Public/UbaScheduler.h
+++ b/Engine/Source/Programs/UnrealBuildAccelerator/Common/Public/UbaScheduler.h
@@ -14,9 +14,12 @@ namespace uba
 	class Process;
 	class RootPaths;
 	class SessionServer;
+	class WorkManagerImpl;
 	struct NextProcessInfo;
 	struct ProcessStartInfoHolder;
 
+	////////////////////////////////////////////////////////////////////////////////////////////////////
+
 	struct SchedulerCreateInfo
 	{
 		SchedulerCreateInfo(SessionServer& s) : session(s) {}
@@ -36,6 +39,8 @@ namespace uba
 		ConfigTable* processConfigs = nullptr;
 	};
 
+	////////////////////////////////////////////////////////////////////////////////////////////////////
+
 	struct EnqueueProcessInfo
 	{
 		EnqueueProcessInfo(const ProcessStartInfo& i) : info(i) {}
@@ -56,6 +61,8 @@ namespace uba
 		u32 cacheBucketId = 0; // Bucket that cache should be fetched from. Zero means that it will not fetch anything
 	};
 
+	////////////////////////////////////////////////////////////////////////////////////////////////////
+
 	class Scheduler
 	{
 	public:
@@ -85,18 +92,27 @@ namespace uba
 		float GetProcessWeightThatCanRunRemotelyNow();
 
 	private:
+		struct CacheFetchInfo;
 		struct ExitProcessInfo;
 		struct ProcessStartInfo2;
 
-		enum ProcessStatus : u8
+		enum RunStatus : u8
 		{
-			ProcessStatus_QueuedForCache,
-			ProcessStatus_DownloadingFromCache,
-			ProcessStatus_QueuedForRun,
-			ProcessStatus_Running,
-			ProcessStatus_Success,
-			ProcessStatus_Failed,
-			ProcessStatus_Skipped,
+			RunStatus_QueuedForRun,
+			RunStatus_Running,
+			RunStatus_Success,
+			RunStatus_Failed,
+			RunStatus_Skipped,
+		};
+
+		enum CacheStatus : u8
+		{
+			CacheStatus_QueuedForTest,
+			CacheStatus_Testing,
+			CacheStatus_QueuedForDownload,
+			CacheStatus_Downloading,
+			CacheStatus_Success,
+			CacheStatus_Failed,
 		};
 
 		UBA_NOINLINE void ThreadLoop();
@@ -106,15 +122,18 @@ namespace uba
 		void HandleCacheMissed(u32 processIndex);
 		void RemoteSlotAvailable(bool isCrossArchitecture);
 		void ProcessExited(ExitProcessInfo* info, const ProcessHandle& handle);
-		u32 PopProcess(bool isLocal, ProcessStatus& outPrevStatus);
+		bool CanRun();
+		u32 PopCache();
+		u32 PopProcess(bool isLocal, RunStatus& outPrevStatus);
+		bool RunCacheQuery();
 		bool RunQueuedProcess(bool isLocal);
 		bool HandleReuseMessage(Process& process, NextProcessInfo& outNextProcess, u32 prevExitCode);
 		void ExitProcess(ExitProcessInfo& info, Process& process, u32 exitCode, bool fromCache);
 		void SkipProcess(ProcessStartInfo2& info);
+		void UpdateStatusHitMissCount();
 		void UpdateQueueCounter(int offset);
 		void UpdateActiveProcessCounter(bool isLocal, int offset);
 		void FinishProcess(const ProcessHandle& handle);
-		void FlushFinishedDownloadThreads();
 
 		SessionServer& m_session;
 		u32 m_maxLocalProcessors;
@@ -125,19 +144,23 @@ namespace uba
 			ProcessStartInfo2* info;
 			u32* dependencies;
 			u32 dependencyCount;
-			ProcessStatus status;
+			CacheStatus cacheStatus;
+			RunStatus runStatus;
 			bool canDetour;
 			bool canExecuteRemotely;
 		};
 
 		Futex m_processEntriesLock;
 		Vector<ProcessEntry> m_processEntries;
-		u32 m_processEntriesStart = 0;
 
 		Function<void(const ProcessHandle&)> m_processFinished;
 
-		Event m_updateThreadLoop;
-		Thread m_thread;
+		Event m_updateThreadWakeup;
+		Thread m_updateThread;
+		u32 m_processEntriesRunStart = 0;
+
+		u32 m_processEntriesCacheStart = 0;
+
 		Atomic<bool> m_loop;
 		bool m_enableProcessReuse;
 		bool m_forceRemote;
@@ -149,6 +172,7 @@ namespace uba
 
 		float m_activeLocalProcessWeight = 0.0f;
 		u32 m_activeCacheQueries = 0;
+		Atomic<u32> m_activeCacheDownloads;
 
 		Atomic<u32> m_totalProcesses;
 		Atomic<u32> m_queuedProcesses;
@@ -163,12 +187,14 @@ namespace uba
 		Vector<RootPaths*> m_rootPaths;
 		bool m_writeToCache;
 
-		List<Thread> m_cacheDownloadThreads;
+		WorkManagerImpl* m_cacheWorkManager = nullptr;
 		List<Thread> m_finishedRemoteProcessThreads;
 
 		Scheduler(const Scheduler&) = delete;
 		void operator=(const Scheduler&) = delete;
 	};
+
+	////////////////////////////////////////////////////////////////////////////////////////////////////
 }
 
 #if 0
diff --git a/Engine/Source/Programs/UnrealBuildAccelerator/Core/Public/UbaSynchronization.h b/Engine/Source/Programs/UnrealBuildAccelerator/Core/Public/UbaSynchronization.h
--- a/Engine/Source/Programs/UnrealBuildAccelerator/Core/Public/UbaSynchronization.h
+++ b/Engine/Source/Programs/UnrealBuildAccelerator/Core/Public/UbaSynchronization.h
@@ -26,6 +26,9 @@ namespace uba
 		void operator=(const AtomicU64& o) { store(o); }
 	};
 
+	template<typename T>
+	void AtomicMax(Atomic<T>& target, T value);
+
 	////////////////////////////////////////////////////////////////////////////////////////////////////
 
 	class CriticalSection
@@ -182,4 +185,13 @@ namespace uba
 	ScopeGuard<Lambda> MakeGuard(Lambda&& lambda) { return ScopeGuard<Lambda>(std::move(lambda)); }
 
 	////////////////////////////////////////////////////////////////////////////////////////////////////
+
+	template<typename T>
+	void AtomicMax(Atomic<T>& target, T value)
+	{
+		T current = target.load(std::memory_order_relaxed);
+		while (current < value && !target.compare_exchange_weak(current, value, std::memory_order_relaxed, std::memory_order_relaxed)) {}
+	}
+
+	////////////////////////////////////////////////////////////////////////////////////////////////////
 }
