[UBA]
* Fixed bug in scheduler where it could hang during ctrl-c
* Fixed bug where log entry with error could be printed even though it exited successfully

diff --git a/Engine/Source/Programs/UnrealBuildAccelerator/Common/Private/UbaProcess.cpp b/Engine/Source/Programs/UnrealBuildAccelerator/Common/Private/UbaProcess.cpp
--- a/Engine/Source/Programs/UnrealBuildAccelerator/Common/Private/UbaProcess.cpp
+++ b/Engine/Source/Programs/UnrealBuildAccelerator/Common/Private/UbaProcess.cpp
@@ -329,7 +329,14 @@ namespace uba
 			}
 
 			if (err.IsEmpty())
+			{
+				// Ok, this is a bit hacky putting this here.. but theory is that process could have exited and exit message is still in the shared memory
+				// We've seen builds on the farm that has a successful process but also error below
+				if (m_readEvent.IsSet(0))
+					return true;
+
 				err.Appendf(TC("ERROR: Process %llu %s (%s) not active but did not get exit message. Received %u messages (GetExitCodeProcess returned 0x%x)"), u64(m_nativeProcessHandle), m_startInfo.GetDescription(), m_realApplication.c_str(), m_messageCount, exitCode);
+			}
 			LogLine(false, err.data, LogEntryType_Error);
 			m_nativeProcessExitCode = UBA_EXIT_CODE(666);
 		}
diff --git a/Engine/Source/Programs/UnrealBuildAccelerator/Common/Private/UbaScheduler.cpp b/Engine/Source/Programs/UnrealBuildAccelerator/Common/Private/UbaScheduler.cpp
--- a/Engine/Source/Programs/UnrealBuildAccelerator/Common/Private/UbaScheduler.cpp
+++ b/Engine/Source/Programs/UnrealBuildAccelerator/Common/Private/UbaScheduler.cpp
@@ -83,7 +83,6 @@ namespace uba
 		virtual const Vector<u8>& GetTrackedInputs() const override { static Vector<u8> v; return v;}
 		virtual const Vector<u8>& GetTrackedOutputs() const override { static Vector<u8> v; return v;}
 		virtual bool IsRemote() const override { return false; }
-		virtual bool IsCache() const { return false; }
 		virtual ProcessExecutionType GetExecutionType() const override { return ProcessExecutionType_Skipped; }
 		ProcessStartInfoHolder startInfo;
 	};
@@ -102,7 +101,6 @@ namespace uba
 		virtual const Vector<u8>& GetTrackedInputs() const override { static Vector<u8> v; return v;}
 		virtual const Vector<u8>& GetTrackedOutputs() const override { static Vector<u8> v; return v;}
 		virtual bool IsRemote() const override { return false; }
-		virtual bool IsCache() const { return false; }
 		virtual ProcessExecutionType GetExecutionType() const override { return ProcessExecutionType_FromCache; }
 		ProcessStartInfoHolder startInfo;
 		Vector<ProcessLogLine> logLines;
@@ -171,6 +169,7 @@ namespace uba
 		m_cancelled = true;
 		m_updateThreadWakeup.Set();
 		m_updateThread.Wait();
+		m_cacheWorkManager->FlushWork(60*1000);
 		delete m_cacheWorkManager;
 		m_session.WaitOnAllTasks();
 		SkipAllQueued();
@@ -302,7 +301,28 @@ namespace uba
 	bool Scheduler::IsEmpty()
 	{
 		SCOPED_FUTEX_READ(m_processEntriesLock, lock);
-		return m_processEntries.size() <= m_finishedProcesses;
+		bool isEmpty = m_processEntries.size() <= m_finishedProcesses;
+		if (isEmpty)
+			return true;
+			
+		#if 0
+		if (m_cancelled)
+		{
+			m_session.GetLogger().Error(TC("STATUS"));
+			u32 index = 0;
+			for (auto& entry : m_processEntries)
+			{
+				if (entry.runStatus != RunStatus_Skipped && entry.runStatus != RunStatus_Failed && entry.runStatus != RunStatus_Success)
+				{
+					m_session.GetLogger().Error(TC("Entry %u: Run: %u Cache: %u"), index, u32(entry.runStatus), u32(entry.cacheStatus));
+					break;
+				}
+				++index;
+			}
+		}
+		#endif
+		
+		return false;
 	}
 
 	void Scheduler::SetProcessFinishedCallback(const Function<void(const ProcessHandle&)>& processFinished)
@@ -512,8 +532,7 @@ namespace uba
 					if (depStatus == RunStatus_Failed || depStatus == RunStatus_Skipped)
 					{
 						entry.cacheStatus = CacheStatus_Failed;
-						canRun = false;
-						break;
+						return i;
 					}
 					if (depStatus != RunStatus_Success)
 					{
@@ -679,9 +698,8 @@ namespace uba
 
 				}, 1, TC("UbaSchedTest"));
 		}
-		else
+		else if (entry.cacheStatus == CacheStatus_QueuedForDownload)
 		{
-			UBA_ASSERT(entry.cacheStatus == CacheStatus_QueuedForDownload);
 			entry.cacheStatus = CacheStatus_Downloading;
 			lock.Leave();
 
@@ -720,6 +738,15 @@ namespace uba
 
 				}, 1, TC("UbaSchedDlwd"));
 		}
+		else
+		{
+			UBA_ASSERT(entry.cacheStatus == CacheStatus_Failed);
+			UBA_ASSERT(!info->cacheInfo);
+			lock.Leave();
+
+			HandleCacheMissed(entryToProcess);
+			++m_cacheMissCount;
+		}
 
 		return true;
 	}
@@ -738,6 +765,7 @@ namespace uba
 			auto info = processEntry.info;
 			bool canDetour = processEntry.canDetour && !m_forceNative;
 			bool wasSkipped = processEntry.runStatus == RunStatus_Skipped;
+			UBA_ASSERT(processEntry.cacheStatus == CacheStatus_Failed);
 			lock.Leave();
 
 			UpdateQueueCounter(-1);
@@ -851,8 +879,6 @@ namespace uba
 		if (!si)
 			return;
 
-		//m_session.GetLogger().Error(TC("PROCESS: %u"), info.processIndex);
-
 		ProcessHandle ph;
 		ph.m_process = &process;
 
@@ -875,6 +901,9 @@ namespace uba
 		}
 		else
 		{
+			UBA_ASSERT(!fromCache);
+			UBA_ASSERT(entry.cacheStatus == CacheStatus_Failed);
+
 			entry.canExecuteRemotely = false;
 			entry.canDetour = exitedResponse != ProcessExitedResponse_RerunNative;
 			entry.runStatus = RunStatus_QueuedForRun;
@@ -885,6 +914,7 @@ namespace uba
 		{
 			if (fromCache)
 			{
+				UBA_ASSERT(isDone);
 				entry.cacheStatus = CacheStatus_Success;
 				--m_activeCacheQueries;
 			}
@@ -924,6 +954,7 @@ namespace uba
 			func(info.userData, ph, exitedResponse);
 		UBA_ASSERT(exitedResponse == ProcessExitedResponse_None);
 		FinishProcess(ph);
+		m_updateThreadWakeup.Set();
 	}
 
 	void Scheduler::UpdateStatusHitMissCount()
