[UBA]
* Added more detouring of posix system calls to handle deeper vfs integration into linux linking (faccessat, fstatat, unlinkat)
* Added devirtualize call in opendir
* Fixed fchmodat supporting vfs

diff --git a/Engine/Source/Programs/UnrealBuildAccelerator/Detours/Private/PosixOS/UbaDetoursFunctionsPosix.cpp b/Engine/Source/Programs/UnrealBuildAccelerator/Detours/Private/PosixOS/UbaDetoursFunctionsPosix.cpp
--- a/Engine/Source/Programs/UnrealBuildAccelerator/Detours/Private/PosixOS/UbaDetoursFunctionsPosix.cpp
+++ b/Engine/Source/Programs/UnrealBuildAccelerator/Detours/Private/PosixOS/UbaDetoursFunctionsPosix.cpp
@@ -70,6 +70,8 @@
 	DETOURED_FUNCTION(fchmod) \
 	DETOURED_FUNCTION(fchmodat) \
 	DETOURED_FUNCTION(fstat) \
+	DETOURED_FUNCTION(faccessat) \
+	DETOURED_FUNCTION(fstatat) \
 	DETOURED_FUNCTION(futimens) \
 	DETOURED_FUNCTION(fclose) \
 	DETOURED_FUNCTION(opendir) \
@@ -92,6 +94,7 @@
 	DETOURED_FUNCTION(remove) \
 	DETOURED_FUNCTION(link) \
 	DETOURED_FUNCTION(unlink) \
+	DETOURED_FUNCTION(unlinkat) \
 	DETOURED_FUNCTION(symlink) \
 	DETOURED_FUNCTION(access) \
 	DETOURED_FUNCTION(posix_spawn) \
@@ -124,7 +127,6 @@
 
 // Used by cp and might need detour
 // fadvise64
-// newfstatat
 // copy_file_range
 
 #if UBA_DETOUR_DEBUG && PLATFORM_LINUX
@@ -766,7 +768,71 @@
 	return res;
 }
 
+//template<typename True_stat>
+int Shared_access(const char* funcName, const char* pathname, int mode)
+{
+	StringBuffer<> fixedPath;
+	if (!FixPath(fixedPath, pathname) || fixedPath.StartsWith("/proc") || !CanDetour2(fixedPath))
+	{
+		auto res = TRUE_WRAPPER(access)(pathname, mode);
+		DEBUG_LOG_TRUE("access", "%s %i -> %i (%s)", pathname, mode, res, StrError(res, errno));
+		return res;
+	}
+
+	bool checkIfDir = false;
+	StringBuffer<> temp;
+	if (g_runningRemote && fixedPath.StartsWith(g_exeDir.data))
+	{
+		temp.Append(g_virtualApplicationDir).Append(fixedPath.data + g_exeDir.count);
+
+		if (temp.count == g_virtualApplicationDir.count) // Remove last slash from temp if this is a access call for the actual applicaiton directory
+		{
+			checkIfDir = true;
+			temp.Resize(temp.count - 1);
+		}
+		fixedPath.Clear().Append(temp);
+	}
+
+	DevirtualizePath(fixedPath);
+
+	if (!CanDetour(pathname))
+	{
+		auto res = TRUE_WRAPPER(access)(pathname, mode);
+		DEBUG_LOG_TRUE(funcName, "%s %i -> %i (%s)", pathname, mode, res, StrError(res, errno));
+		return res;
+	}
+
+	FileAttributes attr;
+	const char* realName = Shared_GetFileAttributes(attr, fixedPath, checkIfDir);
+
+	if (!attr.useCache)
+	{
+		auto res = TRUE_WRAPPER(access)(realName, mode);
+		DEBUG_LOG_TRUE(funcName, "%s %i (%s) -> %i %s", pathname, mode, realName, res, StrError(res, errno));
+		return res;
+	}
+
+	int res = attr.lastError == 0 ? 0 : -1;
+	
+	DEBUG_LOG_DETOURED(funcName, "%s %i (%s) -> %i %s", pathname, mode, realName, res, StrError(res, errno));
+
+	#if UBA_DEBUG_VALIDATE
+	if (!g_runningRemote)
+	{
+		//errno = 0;
+		auto res2 = TRUE_WRAPPER(access)(realName, mode);
+		//int eo = errno;
+		
+		DEBUG_LOG_DETOURED(funcName, "%s %i (%s) -> %i %s", pathname, mode, realName, res, StrError(res, attr.lastError));
+		
+		UBA_ASSERTF(res2 == res, "MISMATCH OF RESULTS for %s - %i %i (err = %s) (exedir %s)", realName, res2, res, StrError(res, attr.lastError), g_exeDir.data);
+		//UBA_ASSERTF(eo == attr.lastError, "MISMATCH OF ERRORS for %s - %i %i", realName, eo, attr.lastError);
+	}
+	#endif
+	
+	errno = attr.lastError;
+	return res;
+}
 
 
-
 // Detoured functions
@@ -1253,8 +1319,14 @@
 UBA_EXPORT int UBA_WRAPPER(fchmodat)(int dirfd, const char* pathname, mode_t mode, int flags)
 {
 	UBA_INIT_DETOUR(fchmodat, dirfd, pathname, mode, flags);
-	DEBUG_LOG_TRUE("fchmodat", "%i %s %i %i", dirfd, pathname, mode, flags);
-	return TRUE_WRAPPER(fchmodat)(dirfd, pathname, mode, flags);
+	UBA_ASSERT(dirfd == AT_FDCWD);
+	StringBuffer<> fixedPath;
+	FixPath(fixedPath, pathname);
+	DevirtualizePath(fixedPath);
+
+	int res = TRUE_WRAPPER(fchmodat)(dirfd, fixedPath.data, mode, flags);
+	DEBUG_LOG_TRUE("fchmodat", "%i %s %i %i -> %i (%s)", dirfd, pathname, mode, flags, res, StrError(res, errno));
+	return res;
 }
 
 //UBA_EXPORT size_t UBA_WRAPPER(fwrite)(const void* ptr, size_t size, size_t nitems, FILE* stream)
@@ -1272,6 +1344,16 @@
 	return Shared_fstat("fstat", fd, buf, [](int fd, struct stat* buf) { return TRUE_WRAPPER(fstat)(fd, buf); });
 }
 
+UBA_EXPORT int UBA_WRAPPER(fstatat)(int dirfd, const char* pathname, struct stat* statbuf, int flags)
+{
+	UBA_INIT_DETOUR(fstatat, dirfd, pathname, statbuf, flags);
+	if (dirfd == AT_FDCWD)
+		return Shared_stat("fstatat", pathname, statbuf, [](const char* file, struct stat* attr) { return TRUE_WRAPPER(stat)(file, attr); });
+	UBA_ASSERT(!g_runningRemote);
+	DEBUG_LOG_TRUE("fstatat", "%s", pathname);
+	return TRUE_WRAPPER(fstatat)(dirfd, pathname, statbuf, flags);
+}
+
 UBA_EXPORT int UBA_WRAPPER(futimens)(int fd, const struct timespec* times)
 {
 	UBA_INIT_DETOUR(futimens, fd, times);
@@ -1313,6 +1395,8 @@
 		return res;
 	}
 
+	DevirtualizePath(dirName);
+
 	StringBuffer<> forHash(dirName);
 	if (forHash.count == 1)
 		forHash.Resize(0);
@@ -1545,66 +1629,16 @@
 UBA_EXPORT int UBA_WRAPPER(access)(const char* pathname, int mode)
 {
 	UBA_INIT_DETOUR(access, pathname, mode);
-
-	StringBuffer<> fixedPath;
-	if (!FixPath(fixedPath, pathname) || fixedPath.StartsWith("/proc") || !CanDetour2(fixedPath))
-	{
-		auto res = TRUE_WRAPPER(access)(pathname, mode);
-		DEBUG_LOG_TRUE("access", "%s %i -> %i (%s)", pathname, mode, res, StrError(res, errno));
-		return res;
-	}
-
-	bool checkIfDir = false;
-	StringBuffer<> temp;
-	if (g_runningRemote && fixedPath.StartsWith(g_exeDir.data))
-	{
-		temp.Append(g_virtualApplicationDir).Append(fixedPath.data + g_exeDir.count);
-
-		if (temp.count == g_virtualApplicationDir.count) // Remove last slash from temp if this is a access call for the actual applicaiton directory
-		{
-			checkIfDir = true;
-			temp.Resize(temp.count - 1);
-		}
-		fixedPath.Clear().Append(temp);
-	}
-
-	DevirtualizePath(fixedPath);
-
-	if (!CanDetour(pathname))
-	{
-		auto res = TRUE_WRAPPER(access)(pathname, mode);
-		DEBUG_LOG_TRUE("access", "%s %i -> %i (%s)", pathname, mode, res, StrError(res, errno));
-		return res;
-	}
-
-	FileAttributes attr;
-	const char* realName = Shared_GetFileAttributes(attr, fixedPath, checkIfDir);
-
-	if (!attr.useCache)
-	{
-		auto res = TRUE_WRAPPER(access)(realName, mode);
-		DEBUG_LOG_TRUE("access", "%s %i (%s) -> %i %s", pathname, mode, realName, res, StrError(res, errno));
-		return res;
-	}
-
-	int res = attr.lastError == 0 ? 0 : -1;
-	
-	#if UBA_DEBUG_VALIDATE
-	if (!g_runningRemote)
-	{
-		//errno = 0;
-		auto res2 = TRUE_WRAPPER(access)(realName, mode);
-		//int eo = errno;
-		
-		DEBUG_LOG_DETOURED("access", "%s %i (%s) -> %i %s", pathname, mode, realName, res, StrError(res, attr.lastError));
-		
-		UBA_ASSERTF(res2 == res, "MISMATCH OF RESULTS for %s - %i %i (err = %s) (exedir %s)", realName, res2, res, StrError(res, attr.lastError), g_exeDir.data);
-		//UBA_ASSERTF(eo == attr.lastError, "MISMATCH OF ERRORS for %s - %i %i", realName, eo, attr.lastError);
-	}
-	#endif
-	
-	errno = attr.lastError;
-	return res;
+	return Shared_access("access", pathname, mode);
 }
 
+UBA_EXPORT int UBA_WRAPPER(faccessat)(int dirfd, const char *pathname, int mode, int flags)
+{
+	UBA_INIT_DETOUR(faccessat, dirfd, pathname, mode, flags);
+	if (dirfd == AT_FDCWD)
+		return Shared_access("faccessat", pathname, mode);
+	UBA_ASSERT(!g_runningRemote);
+	return TRUE_WRAPPER(faccessat)(dirfd, pathname, mode, flags);
+}
+
 #if PLATFORM_LINUX
@@ -1936,6 +1970,21 @@
 	return Shared_DeleteFile("unlink", pathname);
 }
 
+UBA_EXPORT int UBA_WRAPPER(unlinkat)(int dirfd, const char *pathname, int flags)
+{
+	UBA_INIT_DETOUR(unlinkat, dirfd, pathname, flags);
+
+	UBA_ASSERT(!g_runningRemote);
+	UBA_ASSERT(dirfd == AT_FDCWD);
+
+	StringBuffer<> fixedPath;
+	FixPath(fixedPath, pathname);
+	DevirtualizePath(fixedPath);
+	int res = TRUE_WRAPPER(unlinkat)(dirfd, fixedPath.data, flags);
+	DEBUG_LOG_TRUE("unlinkat", "%s -> %i (%s)", pathname, res, StrError(res, errno));
+	return res;
+}
+
 thread_local int t_inVfork;
 
 void FlattenArgs(StringBufferBase& out, const char* const argv[])
