[UBA]
* Fixed another hang when pressing ctrl-c while being in certain download-from-cache states in scheduler

diff --git a/Engine/Source/Programs/UnrealBuildAccelerator/Common/Private/UbaScheduler.cpp b/Engine/Source/Programs/UnrealBuildAccelerator/Common/Private/UbaScheduler.cpp
--- a/Engine/Source/Programs/UnrealBuildAccelerator/Common/Private/UbaScheduler.cpp
+++ b/Engine/Source/Programs/UnrealBuildAccelerator/Common/Private/UbaScheduler.cpp
@@ -423,7 +423,7 @@ namespace uba
 	{
 		SCOPED_FUTEX(m_processEntriesLock, lock);
 		auto& entry = m_processEntries[processIndex];
-		UBA_ASSERT((entry.cacheStatus == CacheStatus_Testing || entry.cacheStatus == CacheStatus_Downloading) && entry.runStatus == RunStatus_QueuedForRun);
+		UBA_ASSERTF((entry.cacheStatus == CacheStatus_Testing || entry.cacheStatus == CacheStatus_Downloading) && entry.runStatus == RunStatus_QueuedForRun, TC("Unexpected entry state. Run: %u Cache: %u"), u32(entry.runStatus), u32(entry.cacheStatus));
 
 		--m_activeCacheQueries;
 
@@ -532,6 +532,7 @@ namespace uba
 					if (depStatus == RunStatus_Failed || depStatus == RunStatus_Skipped)
 					{
 						entry.cacheStatus = CacheStatus_Failed;
+						entry.runStatus = RunStatus_Skipped;
 						return i;
 					}
 					if (depStatus != RunStatus_Success)
@@ -741,11 +742,10 @@ namespace uba
 		else
 		{
 			UBA_ASSERT(entry.cacheStatus == CacheStatus_Failed);
+			UBA_ASSERT(entry.runStatus == RunStatus_Skipped);
 			UBA_ASSERT(!info->cacheInfo);
 			lock.Leave();
-
-			HandleCacheMissed(entryToProcess);
-			++m_cacheMissCount;
+			SkipProcess(*info);
 		}
 
 		return true;
