[UBA]
* Fixed crash when destroying scheduler without cache work manager

diff --git a/Engine/Source/Programs/UnrealBuildAccelerator/Common/Private/UbaScheduler.cpp b/Engine/Source/Programs/UnrealBuildAccelerator/Common/Private/UbaScheduler.cpp
--- a/Engine/Source/Programs/UnrealBuildAccelerator/Common/Private/UbaScheduler.cpp
+++ b/Engine/Source/Programs/UnrealBuildAccelerator/Common/Private/UbaScheduler.cpp
@@ -169,8 +169,12 @@ namespace uba
 		m_cancelled = true;
 		m_updateThreadWakeup.Set();
 		m_updateThread.Wait();
-		m_cacheWorkManager->FlushWork(60*1000);
-		delete m_cacheWorkManager;
+		if (m_cacheWorkManager)
+		{
+			m_cacheWorkManager->FlushWork(60*1000);
+			delete m_cacheWorkManager;
+			m_cacheWorkManager = nullptr;
+		}
 		m_session.WaitOnAllTasks();
 		SkipAllQueued();
 		Cleanup();
