[UBA]
* Moved GetMemoryInfo from UbaSession to UbaEnvironment

diff --git a/Engine/Source/Programs/UnrealBuildAccelerator/Common/Private/UbaEnvironment.cpp b/Engine/Source/Programs/UnrealBuildAccelerator/Common/Private/UbaEnvironment.cpp
new file mode 100644
--- /dev/null
+++ b/Engine/Source/Programs/UnrealBuildAccelerator/Common/Private/UbaEnvironment.cpp
@@ -0,0 +1,136 @@
+// Copyright Epic Games, Inc. All Rights Reserved.
+
+#include "UbaEnvironment.h"
+#include "UbaStringBuffer.h"
+
+namespace uba
+{
+#if PLATFORM_WINDOWS
+	u64 GetMaxPageFileSize(Logger& logger, MEMORYSTATUSEX& memStatus)
+	{
+		tchar systemDrive = 'c';
+		{
+			tchar temp[32];
+			if (GetEnvironmentVariableW(TC("SystemDrive"), temp, 32))
+				systemDrive = ToLower(temp[0]);
+		}
+
+		u64 maxPageSize = 0;
+		wchar_t str[1024];
+		DWORD strBytes = sizeof(str);
+		LSTATUS res = RegGetValueW(HKEY_LOCAL_MACHINE, TC("SYSTEM\\CurrentControlSet\\Control\\Session Manager\\Memory Management"), TC("PagingFiles"), RRF_RT_REG_MULTI_SZ, NULL, str, &strBytes);
+		if (res == ERROR_SUCCESS)
+		{
+			DWORD pagefileOnOsVolume = 0;
+			DWORD pagefileOnOsVolumeSize = 4;
+			RegGetValueW(HKEY_LOCAL_MACHINE, TC("SYSTEM\\CurrentControlSet\\Control\\Session Manager\\Memory Management"), TC("PagefileOnOsVolume"), RRF_RT_DWORD, NULL, &pagefileOnOsVolume, &pagefileOnOsVolumeSize);
+
+			wchar_t* line = str;
+			for (; size_t lineLen = wcslen(line); line += lineLen + 1)
+			{
+				if (lineLen < 3)
+					continue;
+
+				u64 maxSizeMb = 0;
+
+				StringBuffer<8> drive;
+				drive.Append(line, 3); // Get drive root path
+
+				if (drive[0] == '?') // Drive '?' can exist when "Automatically manage paging file size for all drives".. 
+				{
+					// We can use ExistingPageFiles registry key to figure out which drive...
+					// This key can contain multiple page files normally.. have no idea if it can contain multiple when drive is '?'.. but for now, just use the first
+					wchar_t str2[1024];
+					DWORD str2Bytes = sizeof(str2);
+					res = RegGetValueW(HKEY_LOCAL_MACHINE, TC("SYSTEM\\CurrentControlSet\\Control\\Session Manager\\Memory Management"), TC("ExistingPageFiles"), RRF_RT_REG_MULTI_SZ, NULL, str2, &str2Bytes);
+					if (res != ERROR_SUCCESS)
+						continue;
+
+					auto colon = wcschr(str2, ':'); // Path is something like \??\C:\pagefile.sys or similar.. let's search for : and use character in front of it.
+					if (colon == nullptr || colon == str2)
+						continue;
+
+					drive[0] = colon[-1];
+
+					if (pagefileOnOsVolume && ToLower(drive[0]) != systemDrive)
+						continue;
+				}
+				else if (!pagefileOnOsVolume || ToLower(drive[0]) == systemDrive)
+				{
+					const wchar_t* maxSizeStr = wcsrchr(line, ' ');
+						
+					if (!maxSizeStr || !StringBuffer<32>(maxSizeStr + 1).Parse(maxSizeMb))
+					{
+						logger.Warning(TC("Unrecognized page file information format (please report): %s"), line);
+						continue;
+					}
+
+					if (maxSizeMb) // Custom set page file size
+					{
+						maxPageSize += maxSizeMb * 1024 * 1024;
+						continue;
+					}
+				}
+				else
+				{
+					logger.Warning(TC("Page file is set on drive %c: but registry key value PagefileOnOsVolume is set to 1. Fix registry"), drive[0]);
+				}
+
+				// Max possible system-managed page file
+				maxSizeMb = Max(u64(memStatus.ullTotalPhys) * 3, 4ull * 1024 * 1024 * 1024);
+
+				// Check if disk is limiting factor of system-managed page file
+				// Page file can be max 1/8 of volume size and ofc not more than free space
+				ULARGE_INTEGER totalNumberOfBytes;
+				ULARGE_INTEGER totalNumberOfFreeBytes;
+				if (!GetDiskFreeSpaceExW(drive.data, NULL, &totalNumberOfBytes, &totalNumberOfFreeBytes))
+					return logger.Error(TC("GetDiskFreeSpaceExW failed to get information about %s (%s)"), drive.data, LastErrorToText().data);
+
+				u64 maxDiskPageFileSize = Min(totalNumberOfBytes.QuadPart / 8, totalNumberOfFreeBytes.QuadPart);
+				maxPageSize += Min(maxDiskPageFileSize, maxSizeMb);
+			}
+		}
+		return maxPageSize;
+	}
+#endif
+
+	bool GetMemoryInfo(Logger& logger, u64& outAvailable, u64& outTotal, u64* outMaxPageFileSize)
+	{
+		if (outMaxPageFileSize)
+			*outMaxPageFileSize = 0;
+
+#if PLATFORM_WINDOWS
+		MEMORYSTATUSEX memStatus = { sizeof(memStatus) };
+		if (!GlobalMemoryStatusEx(&memStatus))
+		{
+			outAvailable = 0;
+			outTotal = 0;
+			return logger.Error(TC("Failed to get global memory status (%s)"), LastErrorToText().data);
+		}
+
+		// Page file can grow and we want to use the absolute max size to figure out when we need to wait to start new processes
+		static u64 maxPageSize = GetMaxPageFileSize(logger, memStatus);
+
+		if (outMaxPageFileSize)
+			*outMaxPageFileSize = maxPageSize;
+
+		u64 currentPageSize = memStatus.ullTotalPageFile - memStatus.ullTotalPhys;
+		if (currentPageSize < maxPageSize)
+		{
+			outTotal = memStatus.ullTotalPhys + maxPageSize;
+			outAvailable = memStatus.ullAvailPageFile + (maxPageSize - currentPageSize);
+		}
+		else
+		{
+			outTotal = memStatus.ullTotalPageFile;
+			outAvailable = memStatus.ullAvailPageFile;
+		}
+#else
+		u64 memKb;
+		GetPhysicallyInstalledSystemMemory(memKb);
+		outTotal = memKb*1024*1024;
+		outAvailable = outTotal;
+#endif
+		return true;
+	}
+}
diff --git a/Engine/Source/Programs/UnrealBuildAccelerator/Common/Private/UbaSession.cpp b/Engine/Source/Programs/UnrealBuildAccelerator/Common/Private/UbaSession.cpp
--- a/Engine/Source/Programs/UnrealBuildAccelerator/Common/Private/UbaSession.cpp
+++ b/Engine/Source/Programs/UnrealBuildAccelerator/Common/Private/UbaSession.cpp
@@ -4,6 +4,7 @@
 #include "UbaBinaryParser.h"
 #include "UbaCompressedFileHeader.h"
 #include "UbaConfig.h"
+#include "UbaEnvironment.h"
 #include "UbaFileAccessor.h"
 #include "UbaObjectFile.h"
 #include "UbaProcess.h"
@@ -3690,7 +3691,7 @@ namespace uba
 
 #else
 		u64 throwAway;
-		GetMemoryInfo(throwAway, totalMemoryInKilobytes);
+		GetMemoryInfo(m_logger, throwAway, totalMemoryInKilobytes);
 		totalMemoryInKilobytes /= 1024;
 		double processorMhz = 0.0;
 	#if PLATFORM_LINUX
@@ -3765,125 +3766,6 @@ namespace uba
 		#endif
 	}
 
-	bool Session::GetMemoryInfo(u64& outAvailable, u64& outTotal)
-	{
-#if PLATFORM_WINDOWS
-
-		MEMORYSTATUSEX memStatus = { sizeof(memStatus) };
-		if (!GlobalMemoryStatusEx(&memStatus))
-		{
-			outAvailable = 0;
-			outTotal = 0;
-			return m_logger.Error(TC("Failed to get global memory status (%s)"), LastErrorToText().data);
-		}
-
-		// Page file can grow and we want to use the absolute max size to figure out when we need to wait to start new processes
-		if (m_maxPageSize == ~u64(0))
-		{
-			tchar systemDrive = 'c';
-			{
-				tchar temp[32];
-				if (GetEnvironmentVariableW(TC("SystemDrive"), temp, 32))
-					systemDrive = ToLower(temp[0]);
-			}
-
-			m_maxPageSize = 0;
-			wchar_t str[1024];
-			DWORD strBytes = sizeof(str);
-			LSTATUS res = RegGetValueW(HKEY_LOCAL_MACHINE, TC("SYSTEM\\CurrentControlSet\\Control\\Session Manager\\Memory Management"), TC("PagingFiles"), RRF_RT_REG_MULTI_SZ, NULL, str, &strBytes);
-			if (res == ERROR_SUCCESS)
-			{
-				DWORD pagefileOnOsVolume = 0;
-				DWORD pagefileOnOsVolumeSize = 4;
-				RegGetValueW(HKEY_LOCAL_MACHINE, TC("SYSTEM\\CurrentControlSet\\Control\\Session Manager\\Memory Management"), TC("PagefileOnOsVolume"), RRF_RT_DWORD, NULL, &pagefileOnOsVolume, &pagefileOnOsVolumeSize);
-
-				wchar_t* line = str;
-				for (; size_t lineLen = wcslen(line); line += lineLen + 1)
-				{
-					if (lineLen < 3)
-						continue;
-
-					u64 maxSizeMb = 0;
-
-					StringBuffer<8> drive;
-					drive.Append(line, 3); // Get drive root path
-
-					if (drive[0] == '?') // Drive '?' can exist when "Automatically manage paging file size for all drives".. 
-					{
-						// We can use ExistingPageFiles registry key to figure out which drive...
-						// This key can contain multiple page files normally.. have no idea if it can contain multiple when drive is '?'.. but for now, just use the first
-						wchar_t str2[1024];
-						DWORD str2Bytes = sizeof(str2);
-						res = RegGetValueW(HKEY_LOCAL_MACHINE, TC("SYSTEM\\CurrentControlSet\\Control\\Session Manager\\Memory Management"), TC("ExistingPageFiles"), RRF_RT_REG_MULTI_SZ, NULL, str2, &str2Bytes);
-						if (res != ERROR_SUCCESS)
-							continue;
-
-						auto colon = wcschr(str2, ':'); // Path is something like \??\C:\pagefile.sys or similar.. let's search for : and use character in front of it.
-						if (colon == nullptr || colon == str2)
-							continue;
-
-						drive[0] = colon[-1];
-
-						if (pagefileOnOsVolume && ToLower(drive[0]) != systemDrive)
-							continue;
-					}
-					else if (!pagefileOnOsVolume || ToLower(drive[0]) == systemDrive)
-					{
-						const wchar_t* maxSizeStr = wcsrchr(line, ' ');
-						
-						if (!maxSizeStr || !StringBuffer<32>(maxSizeStr + 1).Parse(maxSizeMb))
-						{
-							m_logger.Warning(TC("Unrecognized page file information format (please report): %s"), line);
-							continue;
-						}
-
-						if (maxSizeMb) // Custom set page file size
-						{
-							m_maxPageSize += maxSizeMb * 1024 * 1024;
-							continue;
-						}
-					}
-					else
-					{
-						m_logger.Warning(TC("Page file is set on drive %c: but registry key value PagefileOnOsVolume is set to 1. Fix registry"), drive[0]);
-					}
-
-					// Max possible system-managed page file
-					maxSizeMb = Max(u64(memStatus.ullTotalPhys) * 3, 4ull * 1024 * 1024 * 1024);
-
-					// Check if disk is limiting factor of system-managed page file
-					// Page file can be max 1/8 of volume size and ofc not more than free space
-					ULARGE_INTEGER totalNumberOfBytes;
-					ULARGE_INTEGER totalNumberOfFreeBytes;
-					if (!GetDiskFreeSpaceExW(drive.data, NULL, &totalNumberOfBytes, &totalNumberOfFreeBytes))
-						return m_logger.Error(TC("GetDiskFreeSpaceExW failed to get information about %s (%s)"), drive.data, LastErrorToText().data);
-
-					u64 maxDiskPageFileSize = Min(totalNumberOfBytes.QuadPart / 8, totalNumberOfFreeBytes.QuadPart);
-					m_maxPageSize += Min(maxDiskPageFileSize, maxSizeMb);
-				}
-			}
-		}
-
-		u64 currentPageSize = memStatus.ullTotalPageFile - memStatus.ullTotalPhys;
-		if (currentPageSize < m_maxPageSize)
-		{
-			outTotal = memStatus.ullTotalPhys + m_maxPageSize;
-			outAvailable = memStatus.ullAvailPageFile + (m_maxPageSize - currentPageSize);
-		}
-		else
-		{
-			outTotal = memStatus.ullTotalPageFile;
-			outAvailable = memStatus.ullAvailPageFile;
-		}
-#else
-		u64 memKb;
-		GetPhysicallyInstalledSystemMemory(memKb);
-		outTotal = memKb*1024*1024;
-		outAvailable = outTotal;
-#endif
-		return true;
-	}
-
 	bool GetCpuTime(u64& outTotalTime, u64& outIdleTime)
 	{
 		outTotalTime = 0;
diff --git a/Engine/Source/Programs/UnrealBuildAccelerator/Common/Private/UbaSessionClient.cpp b/Engine/Source/Programs/UnrealBuildAccelerator/Common/Private/UbaSessionClient.cpp
--- a/Engine/Source/Programs/UnrealBuildAccelerator/Common/Private/UbaSessionClient.cpp
+++ b/Engine/Source/Programs/UnrealBuildAccelerator/Common/Private/UbaSessionClient.cpp
@@ -3,6 +3,7 @@
 #include "UbaSessionClient.h"
 #include "UbaApplicationRules.h"
 #include "UbaConfig.h"
+#include "UbaEnvironment.h"
 #include "UbaNetworkClient.h"
 #include "UbaNetworkMessage.h"
 #include "UbaProcess.h"
@@ -1314,7 +1315,7 @@ namespace uba
 
 		u64 memAvail;
 		u64 memTotal;
-		GetMemoryInfo(memAvail, memTotal);
+		GetMemoryInfo(m_logger, memAvail, memTotal);
 		float cpuLoad = UpdateCpuLoad();
 
 		m_cpuUsage = cpuLoad;
@@ -1914,7 +1915,7 @@ namespace uba
 
 			FlushDeadProcesses();
 
-			GetMemoryInfo(memAvail, memTotal);
+			GetMemoryInfo(m_logger, memAvail, memTotal);
 			m_memAvail = memAvail;
 			m_memTotal = memTotal;
 
@@ -1956,7 +1957,7 @@ namespace uba
 
 				if (!firstCall)
 				{
-					GetMemoryInfo(memAvail, memTotal);
+					GetMemoryInfo(m_logger, memAvail, memTotal);
 					m_memAvail = memAvail;
 					m_memTotal = memTotal;
 				}
diff --git a/Engine/Source/Programs/UnrealBuildAccelerator/Common/Private/UbaSessionServer.cpp b/Engine/Source/Programs/UnrealBuildAccelerator/Common/Private/UbaSessionServer.cpp
--- a/Engine/Source/Programs/UnrealBuildAccelerator/Common/Private/UbaSessionServer.cpp
+++ b/Engine/Source/Programs/UnrealBuildAccelerator/Common/Private/UbaSessionServer.cpp
@@ -3,6 +3,7 @@
 #include "UbaSessionServer.h"	
 #include "UbaApplicationRules.h"
 #include "UbaConfig.h"
+#include "UbaEnvironment.h"
 #include "UbaNetworkServer.h"
 #include "UbaProcess.h"
 #include "UbaProcessStartInfoHolder.h"
@@ -209,7 +210,7 @@ namespace uba
 
 			u64 memAvail;
 			u64 memTotal;
-			if (GetMemoryInfo(memAvail, memTotal))
+			if (GetMemoryInfo(m_logger, memAvail, memTotal))
 			{
 				m_memAvail = memAvail;
 				m_memTotal = memTotal;
@@ -1958,7 +1959,8 @@ namespace uba
 
 			u64 memAvail;
 			u64 memTotal;
-			if (!GetMemoryInfo(memAvail, memTotal))
+			u64 maxPageFile;
+			if (!GetMemoryInfo(m_logger, memAvail, memTotal, &maxPageFile))
 				m_memRequiredToSpawn = 0;
 			m_memAvail = memAvail;
 
@@ -2006,7 +2008,7 @@ namespace uba
 						m_logger.Info(TC("       Also note, this is \"committed\" memory. Not memory in use. So you necessarily don't need more physical memory"));
 						MEMORYSTATUSEX memStatus = { sizeof(memStatus) };
 						GlobalMemoryStatusEx(&memStatus);
-						m_logger.Info(TC("  MaxPage:   %s"), BytesToText(m_maxPageSize));
+						m_logger.Info(TC("  MaxPage:   %s"), BytesToText(maxPageFile));
 						m_logger.Info(TC("  TotalPhys: %s"), BytesToText(memStatus.ullTotalPhys));
 						m_logger.Info(TC("  AvailPhys: %s"), BytesToText(memStatus.ullAvailPhys));
 						m_logger.Info(TC("  TotalPage: %s"), BytesToText(memStatus.ullTotalPageFile));
diff --git a/Engine/Source/Programs/UnrealBuildAccelerator/Common/Public/UbaEnvironment.h b/Engine/Source/Programs/UnrealBuildAccelerator/Common/Public/UbaEnvironment.h
new file mode 100644
--- /dev/null
+++ b/Engine/Source/Programs/UnrealBuildAccelerator/Common/Public/UbaEnvironment.h
@@ -0,0 +1,10 @@
+// Copyright Epic Games, Inc. All Rights Reserved.
+
+#pragma once
+
+#include "UbaLogger.h"
+
+namespace uba
+{
+	bool GetMemoryInfo(Logger& logger, u64& outAvailable, u64& outTotal, u64* outMaxPageFileSize = nullptr);
+}
\ No newline at end of file
diff --git a/Engine/Source/Programs/UnrealBuildAccelerator/Common/Public/UbaSession.h b/Engine/Source/Programs/UnrealBuildAccelerator/Common/Public/UbaSession.h
--- a/Engine/Source/Programs/UnrealBuildAccelerator/Common/Public/UbaSession.h
+++ b/Engine/Source/Programs/UnrealBuildAccelerator/Common/Public/UbaSession.h
@@ -202,7 +202,6 @@ namespace uba
 		bool CreateProcessJobObject();
 		void EnsureDirectoryTableMemory(u64 neededSize);
 		void GetSystemInfo(StringBufferBase& out);
-		bool GetMemoryInfo(u64& outAvailable, u64& outTotal);
 		float UpdateCpuLoad();
 
 		bool HasVfs(RootsHandle handle) const;
@@ -335,7 +334,6 @@ namespace uba
 		StringBuffer<256> m_traceOutputFile;
 		u32 m_traceIntervalMs;
 		TString m_extraInfo;
-		u64 m_maxPageSize = ~u64(0);
 		u64 m_previousTotalCpuTime = 0;
 		u64 m_previousIdleCpuTime = 0;
 		float m_cpuLoad = 0;
