[UBA]
* Moved code from UbaPlatform to UbaEnvironment since it is only used on non-detour side

diff --git a/Engine/Source/Programs/UnrealBuildAccelerator/Common/Private/UbaEnvironment.cpp b/Engine/Source/Programs/UnrealBuildAccelerator/Common/Private/UbaEnvironment.cpp
--- a/Engine/Source/Programs/UnrealBuildAccelerator/Common/Private/UbaEnvironment.cpp
+++ b/Engine/Source/Programs/UnrealBuildAccelerator/Common/Private/UbaEnvironment.cpp
@@ -5,6 +5,112 @@
 
 namespace uba
 {
+	MutexHandle CreateMutexW(bool bInitialOwner, const tchar* lpName)
+	{
+		#if PLATFORM_WINDOWS
+		return (MutexHandle)(u64)::CreateMutexW(NULL, bInitialOwner, lpName);
+		#else
+		// TODO: This is used to check for exclusivity and also for trace streams (only created by host and read by visualizer)
+		SetLastError(ERROR_SUCCESS);
+		return ((MutexHandle)(u64)1337); // Just some random value
+		#endif
+	}
+
+	void ReleaseMutex(MutexHandle mutex)
+	{
+		if (mutex == InvalidMutexHandle)
+			return;
+		#if PLATFORM_WINDOWS
+		::ReleaseMutex((HANDLE)mutex);
+		#else
+		#endif
+	}
+
+	void CloseMutex(MutexHandle mutex)
+	{
+		if (mutex == InvalidMutexHandle)
+			return;
+		#if PLATFORM_WINDOWS
+		::CloseHandle((HANDLE)mutex);
+		#else
+		#endif
+	}
+
+	u32 GetEnvironmentVariableW(const tchar* name, tchar* buffer, u32 nSize)
+	{
+		#if PLATFORM_WINDOWS
+		return ::GetEnvironmentVariableW(name, buffer, nSize);
+		#else
+		const char* env = getenv(name);
+		if (!env)
+		{
+			SetLastError(203); // ERROR_ENVVAR_NOT_FOUND
+			return 0;
+		}
+
+		auto envLen = strlen(env);
+		if (nSize <= envLen)
+			return envLen + 1;
+		memcpy(buffer, env, envLen + 1);
+		return envLen;
+		#endif
+	}
+
+	bool SetEnvironmentVariableW(const tchar* name, const tchar* value)
+	{
+		#if PLATFORM_WINDOWS
+		return ::SetEnvironmentVariableW(name, value);
+		#else
+		return setenv(name, value, 1) == 0;
+		#endif
+	}
+
+	u32 ExpandEnvironmentStringsW(const tchar* lpSrc, tchar* lpDst, u32 nSize)
+	{
+		#if PLATFORM_WINDOWS
+		return ::ExpandEnvironmentStringsW(lpSrc, lpDst, nSize);
+		#else
+		UBA_ASSERTF(false, TC("ExpandEnvironmentStringsW not implemented (%s)"), lpSrc);
+		return 0;
+		#endif
+	}
+
+	ProcHandle GetCurrentProcessHandle()
+	{
+		#if PLATFORM_WINDOWS
+		return (ProcHandle)(u64)GetCurrentProcess();
+		#else
+		UBA_ASSERTF(false, TC("GetCurrentProcessHandle not implemented"));
+		return (ProcHandle)0;
+		#endif
+	}
+
+	u32 GetLogicalProcessorCount()
+	{
+		#if PLATFORM_WINDOWS
+		return ::GetActiveProcessorCount(ALL_PROCESSOR_GROUPS);
+		#else
+		return (u32)sysconf(_SC_NPROCESSORS_ONLN);
+		#endif
+	}
+
+	u32 GetProcessorGroupCount()
+	{
+		#if PLATFORM_WINDOWS
+		static u32 s_processorGroupCount = u32(GetActiveProcessorGroupCount());
+		if (s_processorGroupCount)
+			return s_processorGroupCount;
+		#endif
+		return 1u;
+	}
+
+	void ElevateCurrentThreadPriority()
+	{
+		#if PLATFORM_WINDOWS
+		SetThreadPriority(GetCurrentThread(), THREAD_PRIORITY_HIGHEST);
+		#endif
+	}
+
 #if PLATFORM_WINDOWS
 	u64 GetMaxPageFileSize(Logger& logger, MEMORYSTATUSEX& memStatus)
 	{
diff --git a/Engine/Source/Programs/UnrealBuildAccelerator/Common/Private/UbaNetworkClient.cpp b/Engine/Source/Programs/UnrealBuildAccelerator/Common/Private/UbaNetworkClient.cpp
--- a/Engine/Source/Programs/UnrealBuildAccelerator/Common/Private/UbaNetworkClient.cpp
+++ b/Engine/Source/Programs/UnrealBuildAccelerator/Common/Private/UbaNetworkClient.cpp
@@ -4,6 +4,7 @@
 #include "UbaConfig.h"
 #include "UbaCrypto.h"
 #include "UbaBinaryReaderWriter.h"
+#include "UbaEnvironment.h"
 #include "UbaNetworkBackendTcp.h"
 #include "UbaNetworkMessage.h"
 #include <stdlib.h>
diff --git a/Engine/Source/Programs/UnrealBuildAccelerator/Common/Private/UbaNetworkServer.cpp b/Engine/Source/Programs/UnrealBuildAccelerator/Common/Private/UbaNetworkServer.cpp
--- a/Engine/Source/Programs/UnrealBuildAccelerator/Common/Private/UbaNetworkServer.cpp
+++ b/Engine/Source/Programs/UnrealBuildAccelerator/Common/Private/UbaNetworkServer.cpp
@@ -4,6 +4,7 @@
 #include "UbaConfig.h"
 #include "UbaCrypto.h"
 #include "UbaBinaryReaderWriter.h"
+#include "UbaEnvironment.h"
 #include "UbaPlatform.h"
 
 namespace uba
diff --git a/Engine/Source/Programs/UnrealBuildAccelerator/Common/Private/UbaThread.cpp b/Engine/Source/Programs/UnrealBuildAccelerator/Common/Private/UbaThread.cpp
--- a/Engine/Source/Programs/UnrealBuildAccelerator/Common/Private/UbaThread.cpp
+++ b/Engine/Source/Programs/UnrealBuildAccelerator/Common/Private/UbaThread.cpp
@@ -1,6 +1,7 @@
 // Copyright Epic Games, Inc. All Rights Reserved.
 
 #include "UbaThread.h"
+#include "UbaEnvironment.h"
 #include "UbaPlatform.h"
 
 #if PLATFORM_WINDOWS
diff --git a/Engine/Source/Programs/UnrealBuildAccelerator/Common/Public/UbaEnvironment.h b/Engine/Source/Programs/UnrealBuildAccelerator/Common/Public/UbaEnvironment.h
--- a/Engine/Source/Programs/UnrealBuildAccelerator/Common/Public/UbaEnvironment.h
+++ b/Engine/Source/Programs/UnrealBuildAccelerator/Common/Public/UbaEnvironment.h
@@ -6,5 +6,24 @@
 
 namespace uba
 {
+	enum MutexHandle : u64 {};
+	inline constexpr MutexHandle InvalidMutexHandle = MutexHandle(0);
+	MutexHandle CreateMutexW(bool bInitialOwner, const tchar* name);
+	void ReleaseMutex(MutexHandle mutex);
+	void CloseMutex(MutexHandle mutex);
+
+	enum ProcHandle : u64 {};
+	inline constexpr ProcHandle InvalidProcHandle = ((ProcHandle)(u64)-1);
+	ProcHandle GetCurrentProcessHandle();
+
+	u32 GetEnvironmentVariableW(const tchar* name, tchar* buffer, u32 nSize);
+	bool SetEnvironmentVariableW(const tchar* name, const tchar* value);
+	u32 ExpandEnvironmentStringsW(const tchar* lpSrc, tchar* lpDst, u32 nSize);
+	u32 GetLogicalProcessorCount();
+	u32 GetProcessorGroupCount();
+
+	void ElevateCurrentThreadPriority();
+	void PrefetchVirtualMemory(const void* mem, u64 size);
+
 	bool GetMemoryInfo(Logger& logger, u64& outAvailable, u64& outTotal, u64* outMaxPageFileSize = nullptr);
 }
\ No newline at end of file
diff --git a/Engine/Source/Programs/UnrealBuildAccelerator/Common/Public/UbaFileMapping.h b/Engine/Source/Programs/UnrealBuildAccelerator/Common/Public/UbaFileMapping.h
--- a/Engine/Source/Programs/UnrealBuildAccelerator/Common/Public/UbaFileMapping.h
+++ b/Engine/Source/Programs/UnrealBuildAccelerator/Common/Public/UbaFileMapping.h
@@ -2,6 +2,7 @@
 
 #pragma once
 
+#include "UbaEnvironment.h"
 #include "UbaEvent.h"
 #include "UbaHash.h"
 #include "UbaSynchronization.h"
diff --git a/Engine/Source/Programs/UnrealBuildAccelerator/Core/Private/UbaPlatform.cpp b/Engine/Source/Programs/UnrealBuildAccelerator/Core/Private/UbaPlatform.cpp
--- a/Engine/Source/Programs/UnrealBuildAccelerator/Core/Private/UbaPlatform.cpp
+++ b/Engine/Source/Programs/UnrealBuildAccelerator/Core/Private/UbaPlatform.cpp
@@ -508,112 +508,6 @@ namespace uba
 		#endif
 	}
 
-	MutexHandle CreateMutexW(bool bInitialOwner, const tchar* lpName)
-	{
-		#if PLATFORM_WINDOWS
-		return (MutexHandle)(u64)::CreateMutexW(NULL, bInitialOwner, lpName);
-		#else
-		// TODO: This is used to check for exclusivity and also for trace streams (only created by host and read by visualizer)
-		SetLastError(ERROR_SUCCESS);
-		return ((MutexHandle)(u64)1337); // Just some random value
-		#endif
-	}
-
-	void ReleaseMutex(MutexHandle mutex)
-	{
-		if (mutex == InvalidMutexHandle)
-			return;
-		#if PLATFORM_WINDOWS
-		::ReleaseMutex((HANDLE)mutex);
-		#else
-		#endif
-	}
-
-	void CloseMutex(MutexHandle mutex)
-	{
-		if (mutex == InvalidMutexHandle)
-			return;
-		#if PLATFORM_WINDOWS
-		::CloseHandle((HANDLE)mutex);
-		#else
-		#endif
-	}
-
-	u32 GetEnvironmentVariableW(const tchar* name, tchar* buffer, u32 nSize)
-	{
-		#if PLATFORM_WINDOWS
-		return ::GetEnvironmentVariableW(name, buffer, nSize);
-		#else
-		const char* env = getenv(name);
-		if (!env)
-		{
-			SetLastError(203); // ERROR_ENVVAR_NOT_FOUND
-			return 0;
-		}
-
-		auto envLen = strlen(env);
-		if (nSize <= envLen)
-			return envLen + 1;
-		memcpy(buffer, env, envLen + 1);
-		return envLen;
-		#endif
-	}
-
-	bool SetEnvironmentVariableW(const tchar* name, const tchar* value)
-	{
-		#if PLATFORM_WINDOWS
-		return ::SetEnvironmentVariableW(name, value);
-		#else
-		return setenv(name, value, 1) == 0;
-		#endif
-	}
-
-	u32 ExpandEnvironmentStringsW(const tchar* lpSrc, tchar* lpDst, u32 nSize)
-	{
-		#if PLATFORM_WINDOWS
-		return ::ExpandEnvironmentStringsW(lpSrc, lpDst, nSize);
-		#else
-		UBA_ASSERTF(false, TC("ExpandEnvironmentStringsW not implemented (%s)"), lpSrc);
-		return 0;
-		#endif
-	}
-
-	ProcHandle GetCurrentProcessHandle()
-	{
-		#if PLATFORM_WINDOWS
-		return (ProcHandle)(u64)GetCurrentProcess();
-		#else
-		UBA_ASSERTF(false, TC("GetCurrentProcessHandle not implemented"));
-		return (ProcHandle)0;
-		#endif
-	}
-
-	u32 GetLogicalProcessorCount()
-	{
-		#if PLATFORM_WINDOWS
-		return ::GetActiveProcessorCount(ALL_PROCESSOR_GROUPS);
-		#else
-		return (u32)sysconf(_SC_NPROCESSORS_ONLN);
-		#endif
-	}
-
-	u32 GetProcessorGroupCount()
-	{
-		#if PLATFORM_WINDOWS
-		static u32 s_processorGroupCount = u32(GetActiveProcessorGroupCount());
-		if (s_processorGroupCount)
-			return s_processorGroupCount;
-		#endif
-		return 1u;
-	}
-
-	void ElevateCurrentThreadPriority()
-	{
-		#if PLATFORM_WINDOWS
-		SetThreadPriority(GetCurrentThread(), THREAD_PRIORITY_HIGHEST);
-		#endif
-	}
-
 	void PrefetchVirtualMemory(const void* mem, u64 size)
 	{
 		#if PLATFORM_WINDOWS
diff --git a/Engine/Source/Programs/UnrealBuildAccelerator/Core/Public/UbaPlatform.h b/Engine/Source/Programs/UnrealBuildAccelerator/Core/Public/UbaPlatform.h
--- a/Engine/Source/Programs/UnrealBuildAccelerator/Core/Public/UbaPlatform.h
+++ b/Engine/Source/Programs/UnrealBuildAccelerator/Core/Public/UbaPlatform.h
@@ -109,24 +109,6 @@ namespace uba
 	enum FileHandle : u64 {};
 	inline constexpr FileHandle InvalidFileHandle = (FileHandle)-1;
 
-	enum MutexHandle : u64 {};
-	inline constexpr MutexHandle InvalidMutexHandle = MutexHandle(0);
-	MutexHandle CreateMutexW(bool bInitialOwner, const tchar* name);
-	void ReleaseMutex(MutexHandle mutex);
-	void CloseMutex(MutexHandle mutex);
-
-	enum ProcHandle : u64 {};
-	inline constexpr ProcHandle InvalidProcHandle = ((ProcHandle)(u64)-1);
-	ProcHandle GetCurrentProcessHandle();
-
-	u32 GetEnvironmentVariableW(const tchar* name, tchar* buffer, u32 nSize);
-	bool SetEnvironmentVariableW(const tchar* name, const tchar* value);
-	u32 ExpandEnvironmentStringsW(const tchar* lpSrc, tchar* lpDst, u32 nSize);
-	u32 GetLogicalProcessorCount();
-	u32 GetProcessorGroupCount();
-
-	void ElevateCurrentThreadPriority();
-	void PrefetchVirtualMemory(const void* mem, u64 size);
 	void AddExceptionHandler();
 	void InitMemory();
 
