[UBA]
* Changed so WaitForRead always uses local m_cancelled instead of event. This  to speed up termination of processes

diff --git a/Engine/Source/Programs/UnrealBuildAccelerator/Common/Private/UbaProcess.cpp b/Engine/Source/Programs/UnrealBuildAccelerator/Common/Private/UbaProcess.cpp
--- a/Engine/Source/Programs/UnrealBuildAccelerator/Common/Private/UbaProcess.cpp
+++ b/Engine/Source/Programs/UnrealBuildAccelerator/Common/Private/UbaProcess.cpp
@@ -433,11 +433,7 @@ namespace uba
 
 	bool ProcessImpl::IsCancelled()
 	{
-		#if PLATFORM_WINDOWS
-		return m_cancelEvent.IsSet(0);
-		#else
-		return m_cancelled; // can't use cancel event since memory might have been returned
-		#endif
+		return m_cancelled;
 	}
 
 	bool ProcessImpl::HasFailedMessage()
@@ -462,10 +458,10 @@ namespace uba
 
 	void ProcessImpl::Cancel(bool terminate)
 	{
+		m_cancelled = true;
 		#if PLATFORM_WINDOWS
 		m_cancelEvent.Set();
 		#else
-		m_cancelled = true;
 		SCOPED_FUTEX(m_comMemoryLock, lock);
 		if (m_comMemory.memory)
 			m_cancelEvent.Set();
@@ -492,7 +488,7 @@ namespace uba
 		while (true)
 		{
 			if (m_readEvent.IsSet(readEventTimeout))
-				break;
+				return !m_cancelled;
 
 			#if !PLATFORM_WINDOWS
 			PollStdPipes(outReader, errReader, 0);
@@ -501,10 +497,9 @@ namespace uba
 			if (!IsActive())
 				return m_readEvent.IsSet(0); // Do one more check
 
-			if (IsCancelled())
+			if (m_cancelled)
 				return false;
 		}
-		return true;
 	}
 
 	void ProcessImpl::SetWritten()
