[UBA]
* Fixed bug that would cause shutdown hang if there is an error in a process which has lots of dependencies that are cacheable. m_activeCacheQueries was increased but never decreased if cache request was skipped because of failed dependency.. and would eventually get a number high enough to prevent any cache queries from happening

diff --git a/Engine/Source/Programs/UnrealBuildAccelerator/Common/Private/UbaScheduler.cpp b/Engine/Source/Programs/UnrealBuildAccelerator/Common/Private/UbaScheduler.cpp
--- a/Engine/Source/Programs/UnrealBuildAccelerator/Common/Private/UbaScheduler.cpp
+++ b/Engine/Source/Programs/UnrealBuildAccelerator/Common/Private/UbaScheduler.cpp
@@ -653,14 +653,13 @@ namespace uba
 		if (entryToProcess == ~0u)
 			return false;
 
-		++m_activeCacheQueries;
-
 		auto& entry = m_processEntries[entryToProcess];
 		auto info = entry.info;
 
 		if (entry.cacheStatus == CacheStatus_QueuedForTest)
 		{
 			entry.cacheStatus = CacheStatus_Testing;
+			++m_activeCacheQueries;
 			lock.Leave();
 
 			m_cacheWorkManager->AddWork([this, info, entryToProcess](const WorkContext&)
@@ -706,6 +705,7 @@ namespace uba
 		else if (entry.cacheStatus == CacheStatus_QueuedForDownload)
 		{
 			entry.cacheStatus = CacheStatus_Downloading;
+			++m_activeCacheQueries;
 			lock.Leave();
 
 			++m_activeCacheDownloads;
