[UBA]
* Added memory tracking support in scheduler which can be used to delay spawning of processes during memory pressure or killing processes if running out of memory
* Fixed process cancel code to be more robust (fixed a bunch of bugs)
* Changed how process peakMemory is calculated in order to get more accurate numbers (previous numbers were very off)
* Added tracking of process createTime which is needed to identify unique process (pid+createtime)
* Moved TimeToTick to UbaExport.cpp
* Changed so visualizer show different color between cachefetchtest and cachefetchdownload
* Added visualizer badge showing when process spawning is paused or processes are killed

diff --git a/Engine/Source/Programs/UnrealBuildAccelerator/Cli/Private/UbaCli.cpp b/Engine/Source/Programs/UnrealBuildAccelerator/Cli/Private/UbaCli.cpp
--- a/Engine/Source/Programs/UnrealBuildAccelerator/Cli/Private/UbaCli.cpp
+++ b/Engine/Source/Programs/UnrealBuildAccelerator/Cli/Private/UbaCli.cpp
@@ -1037,7 +1037,7 @@ namespace uba
 		if (logToFile)
 		{
 			logFile.count = GetFullPathNameW(g_rootDir.data, logFile.capacity, logFile.data, nullptr);
-			logFile.EnsureEndsWithSlash().Append(TCV("DebugLog.log"));
+			logFile.EnsureEndsWithSlash().Append(TCV("debuglog.log"));
 			logger.Info(TC("Logging to file: %s"), logFile.data);
 		}
 
@@ -1318,6 +1318,7 @@ namespace uba
 
 			CacheClient* cacheClients[] = { cacheClient };
 			SchedulerCreateInfo info(sessionServer);
+			info.Apply(config);
 			info.forceRemote = isRemote;
 			info.forceNative = commandType == CommandType_Native;
 			info.maxLocalProcessors = maxProcessCount;
diff --git a/Engine/Source/Programs/UnrealBuildAccelerator/Common/Private/UbaProcess.cpp b/Engine/Source/Programs/UnrealBuildAccelerator/Common/Private/UbaProcess.cpp
--- a/Engine/Source/Programs/UnrealBuildAccelerator/Common/Private/UbaProcess.cpp
+++ b/Engine/Source/Programs/UnrealBuildAccelerator/Common/Private/UbaProcess.cpp
@@ -456,8 +456,11 @@ namespace uba
 		return m_processStats.cpuTime;
 	}
 
-	void ProcessImpl::Cancel(bool terminate)
+	bool ProcessImpl::Cancel()
 	{
+		SCOPED_FUTEX(m_logLinesLock, l); // Abuse exiting mutex to not waste memory
+		if (!m_cancelAllowed)
+			return false;
 		m_cancelled = true;
 		#if PLATFORM_WINDOWS
 		m_cancelEvent.Set();
@@ -466,6 +469,7 @@ namespace uba
 		if (m_comMemory.memory)
 			m_cancelEvent.Set();
 		#endif
+		return true;
 	}
 
 	void ProcessImpl::TraverseOutputFiles(const Function<void(StringView file)>& func)
@@ -539,6 +543,8 @@ namespace uba
 
 			if (exitCode == 0)
 			{
+				m_session.ProcessNativeCreated(*this);
+
 				PipeReader outReader(*this, LogEntryType_Info);
 				PipeReader errReader(*this, LogEntryType_Error);
 
@@ -562,7 +568,12 @@ namespace uba
 
 			m_processStats.exitTime = GetTime();
 
+			// Process can't be cancelled after this.
+			SCOPED_FUTEX(m_logLinesLock, cancelLock); // Abuse exiting mutex to not waste memory
+			m_cancelAllowed = false;
 			bool cancelled = IsCancelled();
+			cancelLock.Leave();
+
 			if (exitCode == 0)
 				exitCode = InternalExitProcess(cancelled);
 
@@ -570,19 +581,33 @@ namespace uba
 			// We need to wait also because we need all written files to be added in shared file system
 			WaitForChildrenExit();
 
-			if (!cancelled)
+			if (cancelled)
+			{
+				m_session.ProcessCancelled(*this);
+			}
+			else
+			{
 				if (m_startInfo.writeOutputFilesOnFail || m_startInfo.rules->IsExitCodeSuccess(m_nativeProcessExitCode))
 					if (!WriteFilesToDisk(true))
 						exitCode = UBA_EXIT_CODE(1);
+			}
 
-			if (m_parentProcess)
+			u64 childMaxPeakMemory = 0;
+			u64 childMaxDetoursMemory = 0;
+			for (auto& cph : m_childProcesses)
 			{
-				m_parentProcess->m_processStats.Add(m_processStats);
-				m_parentProcess->m_sessionStats.Add(m_sessionStats);
-				m_parentProcess->m_storageStats.Add(m_storageStats);
-				m_parentProcess->m_kernelStats.Add(m_kernelStats);
+				auto& child = *(ProcessImpl*)cph.m_process;
+				m_processStats.Add(child.m_processStats);
+				m_sessionStats.Add(child.m_sessionStats);
+				m_storageStats.Add(child.m_storageStats);
+				m_kernelStats.Add(child.m_kernelStats);
+
+				// We assume children are running one at the time.. this is an so-so estimation
+				childMaxPeakMemory = Max(childMaxPeakMemory, child.m_processStats.peakMemory.load());
+				childMaxDetoursMemory = Max(childMaxDetoursMemory, child.m_processStats.detoursMemory.load());
 			}
-
+			AtomicMax(m_processStats.peakMemory, m_processStats.peakMemory + childMaxPeakMemory);
+			AtomicMax(m_processStats.detoursMemory, m_processStats.detoursMemory + childMaxDetoursMemory);
 
 			if (exitCode == 0 && m_firstMessageError != 0)
 				exitCode = UBA_EXIT_CODE(1);
@@ -634,7 +659,7 @@ namespace uba
 		m_session.m_debugLogger->Info(TC("ProcessExitedStart (%u)"), m_id);
 		#endif
 
-		UBA_ASSERT(IsCancelled() || !m_parentProcess || !m_parentProcess->m_hasExited);
+		UBA_ASSERTF(IsCancelled() || !m_parentProcess || !m_parentProcess->m_hasExited || m_parentProcess->m_cancelled, TC("Process %s exited after parent exited"), m_startInfo.GetDescription());
 
 		m_hasExited = true;
 
@@ -1337,6 +1362,9 @@ namespace uba
 		m_processStats.Add(stats);
 		m_kernelStats.Add(kernelStats);
 
+		AtomicMax(m_processStats.detoursMemory, stats.detoursMemory.load());
+		AtomicMax(m_processStats.peakMemory, stats.peakMemory.load());
+
 		return false;
 	}
 
@@ -1622,7 +1650,7 @@ namespace uba
 			SCOPED_FUTEX_READ(m_session.m_processJobObjectLock, jobObjectLock);
 			if (!m_session.m_processJobObject)
 			{
-				m_cancelEvent.Set();
+				m_cancelled = true;
 				return ProcessCancelExitCode;
 			}
 
@@ -1648,8 +1676,11 @@ namespace uba
 				creationFlags |= CREATE_NO_WINDOW;
 
 			u32 retryCount = 0;
-			while (!IsCancelled())
+			while (true)
 			{
+				if (IsCancelled())
+					return ProcessCancelExitCode;
+
 				LPCWSTR workingDir = *m_realWorkingDir ? m_realWorkingDir : NULL;
 
 				if (m_detourEnabled)
@@ -1741,6 +1772,10 @@ namespace uba
 			m_nativeProcessId = processInfo.dwProcessId;
 			m_nativeThreadHandle = processInfo.hThread;
 
+			FILETIME dummy;
+			if (processInfo.hProcess && !GetProcessTimes(processInfo.hProcess, (FILETIME*)&m_nativeProcessCreationTime, &dummy, &dummy, &dummy))
+				logger.Error(TC("GetProcessTimes failed (%s)"), LastErrorToText().data);
+
 			#if UBA_DEBUG_TRACK_PROCESS
 			m_session.m_debugLogger->Info(TC("CreateRealProcess  (%u) %s %s"), m_id, m_realApplication.c_str(), m_startInfo.arguments);
 			#endif
@@ -2127,10 +2162,6 @@ namespace uba
 				CloseHandle((HANDLE)handle);
 			});
 
-		PROCESS_MEMORY_COUNTERS mem;
-		if (GetProcessMemoryInfo((HANDLE)handle, &mem, sizeof(mem)))
-			m_processStats.peakMemory = mem.PeakWorkingSetSize;
-
 		bool hadTimeout = false;
 		if (cancel)
 			TerminateProcess((HANDLE)handle, ProcessCancelExitCode);
@@ -2150,8 +2181,7 @@ namespace uba
 					{
 						hadTimeout = true;
 						const tchar* gotMessage = m_gotExitMessage ? TC("Got") : TC("Did not get");
-						const tchar* isCancelledNewCheck = IsCancelled() ? TC("true") : TC("false");
-						logger.Info(TC("WaitForSingleObject timed out after 120 seconds waiting for process %s to exit (Exit code %u, %s ExitMessage and wrote %u files. Cancelled: %s. Runtime: %s). Will terminate and wait again"), m_startInfo.GetDescription(), m_nativeProcessExitCode, gotMessage, u32(m_shared.writtenFiles.size()), isCancelledNewCheck, TimeToText(GetTime() - m_startTime).str);
+						logger.Info(TC("WaitForSingleObject timed out after 120 seconds waiting for process %s to exit (Exit code %u, %s ExitMessage and wrote %u files. Runtime: %s). Will terminate and wait again"), m_startInfo.GetDescription(), m_nativeProcessExitCode, gotMessage, u32(m_shared.writtenFiles.size()), TimeToText(GetTime() - m_startTime).str);
 						TerminateProcess((HANDLE)handle, m_nativeProcessExitCode);
 						continue;
 					}
@@ -2296,18 +2326,22 @@ namespace uba
 
 	void ProcessImpl::WaitForChildrenExit()
 	{
+		if (IsCancelled())
+			for (auto& child : m_childProcesses)
+				child.Cancel();
+
 		for (auto& child : m_childProcesses)
 		{
 			auto& childProcess = *(ProcessImpl*)child.m_process;
 			childProcess.m_waitForParent.Set();
-			while (!childProcess.m_hasExited && !IsCancelled())
+			while (!childProcess.m_hasExited)
 				Sleep(10);
 		}
 	}
 
 	bool ProcessImpl::CancelWithError(bool dummy)
 	{
-		m_cancelEvent.Set();
+		Cancel();
 		return false;
 	}
 
diff --git a/Engine/Source/Programs/UnrealBuildAccelerator/Common/Private/UbaScheduler.cpp b/Engine/Source/Programs/UnrealBuildAccelerator/Common/Private/UbaScheduler.cpp
--- a/Engine/Source/Programs/UnrealBuildAccelerator/Common/Private/UbaScheduler.cpp
+++ b/Engine/Source/Programs/UnrealBuildAccelerator/Common/Private/UbaScheduler.cpp
@@ -4,6 +4,7 @@
 #include "UbaApplicationRules.h"
 #include "UbaCacheClient.h"
 #include "UbaConfig.h"
+#include "UbaEnvironment.h"
 #include "UbaNetworkMessage.h"
 #include "UbaNetworkServer.h"
 #include "UbaProcess.h"
@@ -23,10 +24,28 @@
 // At some point cached processes and local will race which is fine.
 //
 
+#if PLATFORM_WINDOWS
+extern "C" NTSTATUS NTAPI NtQuerySystemInformation(int SystemInformationClass, PVOID SystemInformation, ULONG SystemInformationLength, PULONG ReturnLength);
+#endif
+
 namespace uba
 {
 	////////////////////////////////////////////////////////////////////////////////////////////////////
 
+	#if PLATFORM_WINDOWS
+	constexpr int SystemProcessInformation = 5;
+	constexpr u32 CreateTimeOffset = 32;
+	constexpr u32 PidOffset = 80;
+	constexpr u32 ParentPidOffset = 88;
+	constexpr u32 PagefileUsageOffset = 184;
+	//static_assert(offsetof(SYSTEM_PROCESS_INFORMATION, UniqueProcessId) == PidOffset);
+	//static_assert(offsetof(SYSTEM_PROCESS_INFORMATION, InheritedFromUniqueProcessId) == ParentPidOffset);
+	//static_assert(offsetof(SYSTEM_PROCESS_INFORMATION, CreateTime) == CreateTimeOffset);
+	//static_assert(offsetof(SYSTEM_PROCESS_INFORMATION, PagefileUsage) == PagefileUsageOffset);
+	#endif
+
+	////////////////////////////////////////////////////////////////////////////////////////////////////
+
 	struct Scheduler::CacheFetchInfo : CacheClient::FetchContext
 	{
 		CacheFetchInfo(const ProcessStartInfo& i, Session& s, RootsHandle rh, MessagePriority mp) : CacheClient::FetchContext(i, s, rh, mp) {}
@@ -55,7 +74,9 @@ namespace uba
 		const u8* knownInputs;
 		u32 knownInputsCount;
 		float weight = 1.0f;
-		u32 cacheBucketId = 0;
+		u32 cacheBucketId = 0; // Zero means this process will not check cache
+		u32 memoryGroupId = 0; // Zero means no memory group so this process will not be taken into account in memory allocations
+		u64 predictedMemoryUsage = 0;
 	};
 
 	////////////////////////////////////////////////////////////////////////////////////////////////////
@@ -66,11 +87,29 @@ namespace uba
 		ProcessStartInfo2* startInfo;
 		u32 processIndex = ~0u;
 		bool wasReturned = false;
+		bool hasCheckedWasReturned = false;
 		bool isLocal = true;
 	};
 
 	////////////////////////////////////////////////////////////////////////////////////////////////////
 
+	struct Scheduler::MemGroupStats
+	{
+		u64 baseline = 0;
+		u64 average = 0;
+		u64 history[10] = { 0 };
+		u32 historyCounter = 0;
+		u32 activeProcessCount = 0;
+	};
+
+	struct Scheduler::MemGroupEntry
+	{
+		u64 baseline = 0;
+		u64 average = 0;
+	};
+
+	////////////////////////////////////////////////////////////////////////////////////////////////////
+
 	class SkippedProcess : public Process
 	{
 	public:
@@ -118,6 +157,15 @@ namespace uba
 			table->GetValueAsU32(maxLocalProcessors, TC("MaxLocalProcessors"));
 			table->GetValueAsU32(maxParallelCacheQueries, TC("MaxParallelCacheQueries"));
 			table->GetValueAsBool(useThreadToExitRemoteProcess, TC("UseThreadToExitRemoteProcess"));
+			
+			table->GetValueAsBool(trackMemory, TC("MemTrack"));
+			u32 temp = 0;
+			if (table->GetValueAsU32(temp, TC("MemTracingLevel")))
+				memTracingLevel = u8(temp);
+			if (table->GetValueAsU32(temp, TC("MemStartWaitPercent")))
+				memStartWaitPercent = u8(temp);
+			if (table->GetValueAsU32(temp, TC("MemStartKillPercent")))
+				memStartKillPercent = u8(temp);
 		}
 	}
 
@@ -141,6 +189,19 @@ namespace uba
 				return HandleReuseMessage(process, outNextProcess, prevExitCode);
 			});
 
+		if (info.trackMemory && IsWindows && !IsRunningWine()) // Only implemented for normal windows for now
+		{
+			u64 memAvail;
+			u64 memTotal;
+			if (GetMemoryInfo(m_session.GetLogger(), memAvail, memTotal))
+			{
+				m_memTrack = true;
+				m_memTracingLevel = info.memTracingLevel;
+				m_memWaitThreshold = u64(double(memTotal) * double(info.memStartWaitPercent) / 100.0);
+				m_memKillThreshold = u64(double(memTotal) * double(info.memStartKillPercent) / 100.0);
+			}
+		}
+
 		m_session.SetOuterScheduler(this);
 	}
 
@@ -153,12 +214,20 @@ namespace uba
 
 	void Scheduler::Start()
 	{
-		m_session.SetRemoteProcessReturnedEvent([this](Process& process) { RemoteProcessReturned(process); });
+		m_session.SetRemoteProcessReturnedEvent([this](Process& process) { ProcessReturned(process, false); });
 		m_session.SetRemoteProcessSlotAvailableEvent([this](bool isCrossArchitecture) { RemoteSlotAvailable(isCrossArchitecture); });
 
 		if (!m_cacheClients.empty())
 			m_cacheWorkManager = new WorkManagerImpl(m_maxParallelCacheQueries, TC("UbaSchedCach"));
 
+		if (m_memTrack)
+		{
+			LoadMemTrackTable();
+			m_session.SetNativeProcessCreatedFunc([this](Process& process) { NativeProcessCreated(process); });
+			m_memThreadEvent.Create(true);
+			m_memThread.Start([this]() { ThreadMemoryCheckLoop();  return 0; }, TC("UbaMemTrackLoop"));
+		}
+
 		m_loop = true;
 		m_updateThread.Start([this]() { ThreadLoop(); return 0; }, TC("UbaSchedLoop"));
 	}
@@ -167,6 +236,7 @@ namespace uba
 	{
 		m_loop = false;
 		m_cancelled = true;
+
 		m_updateThreadWakeup.Set();
 		m_updateThread.Wait();
 		if (m_cacheWorkManager)
@@ -177,6 +247,14 @@ namespace uba
 		}
 		m_session.WaitOnAllTasks();
 		SkipAllQueued();
+
+		if (m_memTrack)
+		{
+			m_memThreadEvent.Set();
+			m_memThread.Wait();
+			SaveMemTrackTable();
+		}
+
 		Cleanup();
 	}
 
@@ -251,6 +329,8 @@ namespace uba
 		info2->Expand();
 		info2->weight = info.weight;
 		info2->cacheBucketId = info.cacheBucketId;
+		info2->memoryGroupId = info.memoryGroupId;
+		info2->predictedMemoryUsage = info.predictedMemoryUsage;
 
 		const ApplicationRules* rules = m_session.GetRules(*info2);
 		info2->rules = rules;
@@ -309,7 +389,7 @@ namespace uba
 		if (isEmpty)
 			return true;
 			
-		#if 0
+		#if UBA_DEBUG
 		if (m_cancelled)
 		{
 			m_session.GetLogger().Error(TC("STATUS"));
@@ -399,28 +479,62 @@ namespace uba
 		}
 	}
 
-	void Scheduler::RemoteProcessReturned(Process& process)
+	bool Scheduler::ProcessReturned(Process& process, bool isLocal)
 	{
-		auto& ei = *(ExitProcessInfo*)process.GetStartInfo().userData;
-
-		ei.wasReturned = true;
+		auto& si = process.GetStartInfo();
+		auto& ei = *(ExitProcessInfo*)si.userData;
 		u32 processIndex = ei.processIndex;
+		float processWeight = ei.startInfo->weight;
 
-		process.Cancel(true); // Cancel will call ProcessExited
+		if (isLocal)
+		{
+			SCOPED_CRITICAL_SECTION(m_wasReturnedLock, lock);
+			if (ei.hasCheckedWasReturned)
+				return false;
+			if (!process.Cancel())
+				return false;
+			ei.wasReturned = true;
+		}
+		else
+		{
+			process.Cancel(); // Cancel will call ProcessExited
+		}
 		
 		if (processIndex == ~0u)
-			return;
+			return m_session.GetLogger().Error(TC("Process has no process index. This should not happen"));
+
+		if (isLocal)
+			if (!process.WaitForExit(10*1000))
+				return m_session.GetLogger().Error(TC("Took more than 10 seconds to wait for process to exit (%s)"), si.GetDescription());
 
 		SCOPED_FUTEX(m_processEntriesLock, lock);
+
+		if (isLocal)
+			m_activeLocalProcessWeight -= processWeight;
+
 		if (m_processEntries[processIndex].runStatus != RunStatus_Running)
-			return;
-		m_processEntries[processIndex].runStatus = RunStatus_QueuedForRun;
-		m_processEntriesRunStart = Min(m_processEntriesRunStart, processIndex);
+			return m_session.GetLogger().Error(TC("This should not happen (%u)"), m_processEntries[processIndex].runStatus);
+
+		bool success = !m_cancelled;
+
+		if (success)
+		{
+			m_processEntries[processIndex].runStatus = RunStatus_QueuedForRun;
+			m_processEntriesRunStart = Min(m_processEntriesRunStart, processIndex);
+		}
+		else
+		{
+			m_session.GetLogger().Error(TC("Here to check this path works"));
+			m_processEntries[processIndex].runStatus = RunStatus_Skipped;
+			ProcessHandle ph(&process);
+			FinishProcess(ph);
+		}
 		lock.Leave();
 
 		UpdateQueueCounter(1);
-		UpdateActiveProcessCounter(false, -1);
+		UpdateActiveProcessCounter(isLocal, -1);
 		m_updateThreadWakeup.Set();
+		return true;
 	}
 
 	void Scheduler::HandleCacheMissed(u32 processIndex)
@@ -474,17 +588,27 @@ namespace uba
 	{
 		auto ig = MakeGuard([info]() { delete info; });
 
-		if (info->wasReturned)
-			return;
+		bool isLocal = !handle.IsRemote();
 
 		auto si = info->startInfo;
+
+		if (isLocal && si)
+			LocalProcessExit(*si, handle);
+
+		{
+			SCOPED_CRITICAL_SECTION(m_wasReturnedLock, lock);
+			if (info->wasReturned)
+				return;
+			info->hasCheckedWasReturned = true;
+		}
+
 		if (!si) // Can be a process that was reused but didn't get a new process
 		{
 			UBA_ASSERT(info->processIndex == ~0u);
 			return;
 		}
 
-		if (!m_useThreadToExitRemoteProcess || !handle.IsRemote())
+		if (!m_useThreadToExitRemoteProcess || isLocal)
 		{
 			ExitProcess(*info, *handle.m_process, handle.m_process->GetExitCode(), false);
 			return;
@@ -632,7 +756,12 @@ namespace uba
 				continue;
 
 			if (isLocal)
+			{
+				if (!LocalProcessOkToSpawn())
+					return ~0u;
+
 				m_activeLocalProcessWeight += entry.info->weight;
+			}
 
 			outPrevStatus = entry.runStatus;
 			entry.runStatus = RunStatus_Running;
@@ -780,6 +909,9 @@ namespace uba
 				continue;
 			}
 
+			if (isLocal)
+				LocalProcessStart(*info);
+
 			UpdateActiveProcessCounter(isLocal, 1);
 	
 			auto exitInfo = new ExitProcessInfo();
@@ -820,8 +952,13 @@ namespace uba
 
 		ei->startInfo = nullptr;
 		ei->processIndex = ~0u;
-		if (ei->wasReturned)
-			return false;
+
+		{
+			SCOPED_CRITICAL_SECTION(m_wasReturnedLock, lock);
+			if (ei->wasReturned)
+				return false;
+			ei->hasCheckedWasReturned = true;
+		}
 
 		bool isLocal = !process.IsRemote();
 
@@ -990,6 +1127,371 @@ namespace uba
 		m_session.UpdateProgress(m_totalProcesses, m_finishedProcesses, m_errorCount);
 	}
 
+
+	void Scheduler::KillNewestLocalProcess()
+	{
+		ProcessHandle ph = m_session.GetNewestLocalProcess();
+		if (!ph.IsValid())
+			return;
+
+		m_session.GetTrace().SchedulerKillProcess(ph.GetId());
+
+		u64 startTime = GetTime();
+		if (ProcessReturned(*ph.m_process, true))
+			m_session.GetLogger().Info(TC("Killed process %s (%s) due to memory pressure"), ph.GetStartInfo().GetDescription(), TimeToText(GetTime() - startTime).str);
+	}
+
+	void Scheduler::NativeProcessCreated(Process& process)
+	{
+		if (!m_memTrack)
+			return;
+
+		auto& impl = (ProcessImpl&)process;
+		auto ei = (ExitProcessInfo*)impl.m_startInfo.userData;
+		StringKey key(impl.m_nativeProcessId, impl.m_nativeProcessCreationTime);
+		SCOPED_FUTEX(m_processesLock, l);
+		m_processes[key] = ei;
+	}
+
+	void Scheduler::LocalProcessStart(ProcessStartInfo2& info)
+	{
+		if (!m_memTrack || !info.memoryGroupId)
+			return;
+
+		SCOPED_FUTEX(m_memGroupLookupLock, statsLock);
+		auto insres = m_memGroupLookup.try_emplace(info.memoryGroupId);
+		MemGroupStats& stats = insres.first->second;
+		if (insres.second)
+		{
+			auto findIt = m_memGroupDb.find(info.memoryGroupId);
+			if (findIt != m_memGroupDb.end())
+			{
+				stats.baseline = findIt->second.baseline;
+				stats.average = findIt->second.average;
+				stats.history[0] = stats.average;
+				stats.historyCounter = 1;
+			}
+		}
+		m_memEstimatedUsage += stats.average;
+		++stats.activeProcessCount;
+	}
+
+	bool Scheduler::LocalProcessOkToSpawn()
+	{
+		if (!m_memWaitThreshold)
+			return true;
+
+		if (m_memEstimatedUsage < m_memWaitThreshold)
+		{
+			if (m_memDelaySpawning)
+			{
+				m_session.GetTrace().SchedulerUpdate(false);
+				m_memDelaySpawning = false;
+			}
+			return true;
+		}
+
+		if (m_memDelaySpawning)
+			return false;
+
+		m_memDelaySpawning = true;
+		m_session.GetTrace().SchedulerUpdate(true);
+
+#if 0//PLATFORM_WINDOWS
+		static bool hasBeenRunOnce;
+		if (!hasBeenRunOnce)
+		{
+			hasBeenRunOnce = true;
+			auto& logger = m_session.GetLogger();
+			logger.BeginScope();
+			logger.Info(TC("NOTE - To mitigate this spawn delay it is recommended to make page file larger until you don't see these messages again (Or reduce number of max parallel processes)"));
+			logger.Info(TC("       Set max page file to a large number (like 128gb). It will not use disk space unless you actually start using that amount of committed memory"));
+			logger.Info(TC("       Also note, this is \"committed\" memory. Not memory in use. So you necessarily don't need more physical memory"));
+			MEMORYSTATUSEX memStatus = { sizeof(memStatus) };
+			GlobalMemoryStatusEx(&memStatus);
+			logger.Info(TC("  TotalPhys: %s"), BytesToText(memStatus.ullTotalPhys));
+			logger.Info(TC("  AvailPhys: %s"), BytesToText(memStatus.ullAvailPhys));
+			logger.Info(TC("  TotalPage: %s"), BytesToText(memStatus.ullTotalPageFile));
+			logger.Info(TC("  AvailPage: %s"), BytesToText(memStatus.ullAvailPageFile));
+			logger.EndScope();
+		}
+#endif
+		return false;
+	}
+
+	void Scheduler::LocalProcessExit(ProcessStartInfo2& info, const ProcessHandle& handle)
+	{
+		if (!m_memTrack)
+			return;
+
+		auto& impl = *(ProcessImpl*)handle.m_process;
+
+		if (info.memoryGroupId)
+		{
+			u64 processPeakMemory = impl.m_processStats.peakMemory;
+
+			SCOPED_FUTEX(m_memGroupLookupLock, statsLock);
+			MemGroupStats& stats = m_memGroupLookup[info.memoryGroupId];
+
+			if (!stats.historyCounter) // First process in memory group could be pch
+			{
+				for (auto& kv : impl.m_shared.writtenFiles)
+				{
+					auto& writtenFile = kv.second;
+					if (StringView(writtenFile.name).EndsWith(TC(".pch")))
+					{
+						stats.baseline = writtenFile.mappingWritten;
+						processPeakMemory = stats.baseline;
+					}
+				}
+			}			
+
+			--stats.activeProcessCount;
+
+			if (processPeakMemory >= stats.baseline) // Don't know if this test is too aggressive.. but 
+			{
+				u64& peakMemory = stats.history[stats.historyCounter % sizeof_array(stats.history)];
+				peakMemory = processPeakMemory;
+				++stats.historyCounter;
+
+				u32 count = Min(u32(sizeof_array(stats.history)), stats.historyCounter);
+				u64 average = 0;
+				for (u32 i=0;i!=count;++i)
+					average += stats.history[i];
+				average /= count;
+
+				if (average < stats.baseline)
+					average = stats.baseline;
+				stats.average = average;
+			}
+
+			m_memEstimatedUsage -= stats.average;
+		}
+
+		if (impl.m_nativeProcessCreationTime)
+		{
+			StringKey key(impl.m_nativeProcessId, impl.m_nativeProcessCreationTime);
+			SCOPED_FUTEX(m_processesLock, l);
+			bool res = m_processes.erase(key) == 1;
+			UBA_ASSERT(res);(void)res;
+		}
+	}
+
+	void Scheduler::LoadMemTrackTable()
+	{
+		StringBuffer<> fileName;
+		fileName.Append(m_session.GetRootDir()).Append("memgroups");
+		FileAccessor file(m_session.GetLogger(), fileName.data);
+		if (!file.OpenMemoryRead(0, false) || !file.GetSize())
+			return;
+		BinaryReader reader(file.GetData(), 0, file.GetSize());
+		
+		u16 version = reader.ReadU16();
+		if (version != 1)
+			return;
+
+		while (reader.GetLeft() >= 20)
+		{
+			u32 key = reader.ReadU32();
+			u64 baseline = reader.ReadU64();
+			u64 average = reader.ReadU64();
+			m_memGroupDb[key] = MemGroupEntry{ baseline, average };
+		}
+	}
+
+	void Scheduler::SaveMemTrackTable()
+	{
+		for (auto& kv : m_memGroupLookup)
+			m_memGroupDb[kv.first] = MemGroupEntry{ kv.second.baseline, kv.second.average };
+
+		StringBuffer<> fileName;
+		fileName.Append(m_session.GetRootDir()).Append("memgroups");
+		FileAccessor file(m_session.GetLogger(), fileName.data);
+
+		u64 fileSize = 2 + m_memGroupDb.size()*20;
+
+		if (!file.CreateMemoryWrite(false, DefaultAttributes(), fileSize))
+			return;
+
+		BinaryWriter writer(file.GetData(), 0, fileSize);
+
+		writer.WriteU16(1); // Version
+
+		for (auto& kv : m_memGroupDb)
+		{
+			writer.WriteU32(kv.first);
+			writer.WriteU64(kv.second.baseline);
+			writer.WriteU64(kv.second.average);
+		}
+
+		file.Close();
+	}
+
+	void Scheduler::ThreadMemoryCheckLoop()
+	{
+		auto& logger = m_session.GetLogger();
+
+		//u64 randomKillTime = GetTime() + MsToTime(5000);
+
+		u32 timeoutMs = 0;
+
+		Vector<u8> processInfoBuffer(1024 * 1024);
+
+		constexpr u32 StartRow = 20;
+
+		Trace& trace = m_session.GetTrace();
+
+		if (m_memTracingLevel >= 1)
+		{
+			trace.StatusUpdate(StartRow, 1, TC("MemTot"), LogEntryType_Info, nullptr);
+			trace.StatusUpdate(StartRow+1, 1, TC("MemTot Est"), LogEntryType_Info, nullptr);
+			trace.StatusUpdate(StartRow, 10, TC("MemProc"), LogEntryType_Info, nullptr);
+			trace.StatusUpdate(StartRow+1, 10, TC("MemProc Est"), LogEntryType_Info, nullptr);
+			trace.StatusUpdate(StartRow+2, 1, TC("MemBaseline"), LogEntryType_Info, nullptr);
+			trace.StatusUpdate(StartRow+3, 1, TC("MemUntrack"), LogEntryType_Info, nullptr);
+			trace.StatusUpdate(StartRow+4, 1, TC("MemStartWait"), LogEntryType_Info, nullptr);
+			trace.StatusUpdate(StartRow+4, 6, BytesToText(m_memWaitThreshold).str, LogEntryType_Info, nullptr);
+			trace.StatusUpdate(StartRow+5, 1, TC("MemStartKill"), LogEntryType_Info, nullptr);
+			trace.StatusUpdate(StartRow+5, 6, BytesToText(m_memKillThreshold).str, LogEntryType_Info, nullptr);
+		}
+
+		UnorderedMap<u32, u32> memoryGroupToIndex;
+
+		while (!m_memThreadEvent.IsSet(timeoutMs))
+		{
+			timeoutMs = 1000;
+
+			u64 memAvail;
+			u64 memTotal;
+			u64 maxPageFile;
+			if (!GetMemoryInfo(logger, memAvail, memTotal, &maxPageFile))
+			{
+				m_memTrack = false;
+				logger.Warning(TC("GetMemoryInfo failed. Memory tracking disabled."));
+				m_updateThreadWakeup.Set();
+				break;
+			}
+
+			u64 untrackedMemory = 0;
+			u64 processMemTotalPageCount = 0;
+
+			#if PLATFORM_WINDOWS
+			ULONG returnedSize;
+			while (true)
+			{
+				NTSTATUS res = NtQuerySystemInformation(SystemProcessInformation, processInfoBuffer.data(), (u32)processInfoBuffer.size(), &returnedSize);
+				if (res == STATUS_SUCCESS)
+					break;
+				if (res == STATUS_INFO_LENGTH_MISMATCH)
+				{
+					processInfoBuffer.resize(returnedSize + (10*PagefileUsageOffset));
+					continue;
+				}
+
+				m_memTrack = false;
+				logger.Warning(TC("NtQuerySystemInformation failed. Memory tracking disabled."));
+				m_updateThreadWakeup.Set();
+				break;
+			}
+
+			u8* it = processInfoBuffer.data();
+			ULONG nextEntryOffset = *(ULONG*)it;
+			if (nextEntryOffset < PagefileUsageOffset + 8)
+			{
+				m_memTrack = false;
+				logger.Warning(TC("NtQuerySystemInformation does not contain PageFileUsage. Memory tracking disabled."));
+				m_updateThreadWakeup.Set();
+				break;
+			}
+
+			{
+				SCOPED_FUTEX(m_processesLock, l);
+				while (nextEntryOffset)
+				{
+					u64 createTime = *(u64*)(it + CreateTimeOffset);
+					u32 pid = *(u32*)(it + PidOffset);
+					auto procIt = m_processes.find(StringKey(pid, createTime));
+					if (procIt != m_processes.end())
+					{
+						u64 procMem = *(u64*)(it + PagefileUsageOffset);
+						if (auto ei = procIt->second)
+							if (ei->startInfo->memoryGroupId == 0)
+								untrackedMemory += procMem;
+						processMemTotalPageCount += procMem;
+						//AtomicMax(((ProcessImpl*)procIt->second.process)->m_processStats.peakMemory, spi2->PagefileUsage);
+					}
+
+					nextEntryOffset = *(ULONG*)it;
+					it += nextEntryOffset;
+				}
+			}
+			#endif
+
+			u64 memoryUsage = memTotal - memAvail;
+			u64 baselineMemory = memoryUsage - processMemTotalPageCount;
+			u64 estimatedProcessMemoryUsage = 0;
+			u64 estimatedMemoryUsage = 0;
+
+			{
+				SCOPED_FUTEX(m_memGroupLookupLock, statsLock);
+
+				for (auto& kv : m_memGroupLookup)
+				{
+					MemGroupStats& stats = kv.second;
+					estimatedProcessMemoryUsage += stats.activeProcessCount * stats.average;
+
+					if (m_memTracingLevel >= 2)
+					{
+						constexpr u32 colIndex = 25;
+						auto res = memoryGroupToIndex.emplace(kv.first, 0u);
+						u32& index = res.first->second;
+						if (res.second)
+						{
+							index = StartRow + u32(memoryGroupToIndex.size() - 1);
+							StringBuffer<> b;
+							b.AppendValue(kv.first);
+							trace.StatusUpdate(index, colIndex, b.data, LogEntryType_Info, nullptr);
+						}
+
+						StringBuffer<> b;
+						b.Append(BytesToText(stats.average).str);
+						trace.StatusUpdate(index, colIndex + 5, b.data, LogEntryType_Info, nullptr);
+					}
+				}
+				estimatedMemoryUsage = baselineMemory + untrackedMemory + estimatedProcessMemoryUsage;
+				m_memEstimatedUsage = estimatedMemoryUsage < memoryUsage ? memoryUsage : estimatedMemoryUsage;
+			}
+
+			if (m_memTracingLevel >= 1)
+			{
+				trace.StatusUpdate(StartRow, 6, BytesToText(memoryUsage).str, LogEntryType_Info, nullptr);
+				trace.StatusUpdate(StartRow+1, 6, BytesToText(estimatedMemoryUsage).str, LogEntryType_Info, nullptr);
+				trace.StatusUpdate(StartRow, 15, BytesToText(processMemTotalPageCount).str, LogEntryType_Info, nullptr);
+				trace.StatusUpdate(StartRow+1, 15, BytesToText(estimatedProcessMemoryUsage).str, LogEntryType_Info, nullptr);
+				trace.StatusUpdate(StartRow+2, 6, BytesToText(baselineMemory).str, LogEntryType_Info, nullptr);
+				trace.StatusUpdate(StartRow+3, 6, BytesToText(untrackedMemory).str, LogEntryType_Info, nullptr);
+			}
+
+			if (m_memKillThreshold)
+			{
+				while (memoryUsage > m_memKillThreshold)
+				{
+					KillNewestLocalProcess();
+					if (!GetMemoryInfo(logger, memAvail, memTotal, &maxPageFile))
+						break;
+					memoryUsage = memTotal - memAvail;
+				}
+			}
+
+			//if (randomKillTime < GetTime())
+			//{
+			//	randomKillTime = GetTime() + MsToTime(500);
+			//	if (m_activeLocalProcesses > 4)
+			//		KillNewestLocalProcess();
+			//}
+		}
+	}
+
 	bool Scheduler::EnqueueFromFile(const tchar* yamlFilename, const Function<void(EnqueueProcessInfo&)>& enqueued)
 	{
 		auto& logger = m_session.GetLogger();
@@ -1001,6 +1503,8 @@ namespace uba
 		bool allowDetour = true;
 		bool allowRemote = true;
 		float weight = 1.0f;
+		u32 cacheBucket = 0;
+		u32 memoryGroup = 0;
 		Vector<u32> deps;
 
 		ProcessStartInfo si;
@@ -1029,6 +1533,8 @@ namespace uba
 				info.canDetour = allowDetour;
 				info.canExecuteRemotely = allowRemote;
 				info.weight = weight;
+				info.cacheBucketId = cacheBucket;
+				info.memoryGroupId = memoryGroup;
 				if (enqueued)
 					enqueued(info);
 				EnqueueProcess(info);
@@ -1040,6 +1546,8 @@ namespace uba
 				allowDetour = true;
 				allowRemote = true;
 				weight = 1.0f;
+				cacheBucket = 0;
+				memoryGroup = 0;
 			};
 
 		enum InsideArray
@@ -1137,6 +1645,10 @@ namespace uba
 						allowRemote = !Equals(valueStart, TC("false"));
 					else if (key.Equals(TCV("weight")))
 						StringBuffer<32>(valueStart).Parse(weight);
+					else if (key.Equals(TCV("cache")))
+						StringBuffer<32>(valueStart).Parse(cacheBucket);
+					else if (key.Equals(TCV("memgroup")))
+						StringBuffer<32>(valueStart).Parse(memoryGroup);
 					else if (key.Equals(TCV("dep")))
 					{
 						const tchar* depStart = TStrchr(valueStart, '[');
diff --git a/Engine/Source/Programs/UnrealBuildAccelerator/Common/Private/UbaSession.cpp b/Engine/Source/Programs/UnrealBuildAccelerator/Common/Private/UbaSession.cpp
--- a/Engine/Source/Programs/UnrealBuildAccelerator/Common/Private/UbaSession.cpp
+++ b/Engine/Source/Programs/UnrealBuildAccelerator/Common/Private/UbaSession.cpp
@@ -175,10 +175,10 @@ namespace uba
 		UBA_ASSERT(m_process);
 		return m_process->GetTotalWallTime();
 	}
-	void ProcessHandle::Cancel(bool terminate) const
+	bool ProcessHandle::Cancel() const
 	{
 		UBA_ASSERT(m_process);
-		return m_process->Cancel(terminate);
+		return m_process->Cancel();
 	}
 	const tchar* ProcessHandle::GetExecutingHost() const
 	{
@@ -1829,7 +1829,7 @@ namespace uba
 			}
 
 			for (auto& process : processes)
-				process.Cancel(true);
+				process.Cancel();
 
 			#if PLATFORM_WINDOWS
 			if (m_processJobObject != NULL)
@@ -1856,7 +1856,7 @@ namespace uba
 			processes.push_back(pair.second);
 		lock.Leave();
 		for (auto& process : processes)
-			process.Cancel(true);
+			process.Cancel();
 	}
 
 	ProcessHandle Session::RunProcess(const ProcessStartInfo& startInfo, bool async, bool enableDetour)
@@ -2466,6 +2466,33 @@ namespace uba
 		return true;
 	}
 
+	bool Session::ProcessNativeCreated(ProcessImpl& process)
+	{
+		return true;
+	}
+
+	bool Session::ProcessCancelled(ProcessImpl& process)
+	{
+		Vector<u32> closeIds;
+		SCOPED_FUTEX(m_activeFilesLock, lock);
+		for (auto& kv : m_activeFiles)
+			if (kv.second.owner == &process)
+				closeIds.push_back(kv.first);
+		lock.Leave();
+
+		CloseFileResponse out;
+		CloseFileMessage msg { process };
+
+		for (u32 closeId : closeIds)
+		{
+			// TODO: Should file be deleted from disk if it is there?
+			msg.closeId = closeId;
+			msg.success = false;
+			CloseFile(out, msg);
+		}
+		return true;
+	}
+
 	bool Session::GetInitResponse(InitResponse& out, const InitMessage& msg)
 	{
 		out.directoryTableHandle = m_directoryTableHandle.ToU64();
@@ -2855,9 +2882,10 @@ namespace uba
 		auto insres = m_activeFiles.try_emplace(wantsOnCloseId);
 		if (!insres.second)
 			return m_logger.Error(TC("TRYING TO ADD %s twice!"), out.fileName.data);
-
-		insres.first->second.name = fileName.data;
-		insres.first->second.nameKey = fileNameKey;
+		ActiveFile& activeFile = insres.first->second;
+		activeFile.name = fileName.data;
+		activeFile.nameKey = fileNameKey;
+		activeFile.owner = &msg.process;
 		return true;
 	}
 
@@ -3159,7 +3187,7 @@ namespace uba
 		UBA_ASSERT(msg.toKey != StringKeyZero);
 		SCOPED_FUTEX(m_activeFilesLock, lock);
 		u32 closeId = m_wantsOnCloseIdCounter++;
-		if (!m_activeFiles.try_emplace(closeId, ActiveFile{ msg.toName.data, msg.toKey }).second)
+		if (!m_activeFiles.try_emplace(closeId, ActiveFile{ msg.toName.data, msg.toKey, &msg.process }).second)
 		{
 			m_logger.Error(TC("SHOULD NOT HAPPEN"));
 		}
diff --git a/Engine/Source/Programs/UnrealBuildAccelerator/Common/Private/UbaSessionClient.cpp b/Engine/Source/Programs/UnrealBuildAccelerator/Common/Private/UbaSessionClient.cpp
--- a/Engine/Source/Programs/UnrealBuildAccelerator/Common/Private/UbaSessionClient.cpp
+++ b/Engine/Source/Programs/UnrealBuildAccelerator/Common/Private/UbaSessionClient.cpp
@@ -1929,7 +1929,7 @@ namespace uba
 					SCOPED_FUTEX(rec.lock, lock);
 					if (rec.isDone)
 						continue;
-					rec.handle.Cancel(true);
+					rec.handle.Cancel();
 					rec.isKilled = true;
 					SendReturnProcess(rec.handle.GetId(), TC("Running out of memory"));
 					++m_stats.killCount;
diff --git a/Engine/Source/Programs/UnrealBuildAccelerator/Common/Private/UbaSessionServer.cpp b/Engine/Source/Programs/UnrealBuildAccelerator/Common/Private/UbaSessionServer.cpp
--- a/Engine/Source/Programs/UnrealBuildAccelerator/Common/Private/UbaSessionServer.cpp
+++ b/Engine/Source/Programs/UnrealBuildAccelerator/Common/Private/UbaSessionServer.cpp
@@ -41,10 +41,10 @@ namespace uba
 		virtual const Vector<ProcessLogLine>& GetLogLines() const override { return m_logLines; }
 		virtual const Vector<u8>& GetTrackedInputs() const override { return m_trackedInputs; }
 		virtual const Vector<u8>& GetTrackedOutputs() const override { return m_trackedOutputs; };
-		virtual void Cancel(bool terminate) override
+		virtual bool Cancel() override
 		{
 			if (m_cancelled)
-				return;
+				return true;
 			m_cancelled = true;
 			m_exitCode = ProcessCancelExitCode;
 			if (auto s = m_server)
@@ -56,6 +56,7 @@ namespace uba
 			h.m_process = this;
 			CallProcessExit(h);
 			h.m_process = nullptr;
+			return true;
 		}
 
 		virtual const tchar* GetExecutingHost() const override { return m_executingHost.c_str(); }
@@ -119,15 +120,6 @@ namespace uba
 			table->GetValueAsBool(remoteTraceEnabled, TC("RemoteTraceEnabled"));
 			table->GetValueAsBool(nameToHashTableEnabled, TC("NameToHashTableEnabled"));
 			table->GetValueAsBool(traceIOEnabled, TC("TraceIOEnabled"));
-
-			table->GetValueAsBool(allowWaitOnMem, TC("AllowWaitOnMem"));
-			table->GetValueAsBool(allowKillOnMem, TC("AllowKillOnMem"));
-
-			u32 temp = 0;
-			if (table->GetValueAsU32(temp, TC("MemWaitLoadPercent")))
-				memWaitLoadPercent = u8(temp);
-			if (table->GetValueAsU32(temp, TC("MemKillLoadPercent")))
-				memKillLoadPercent = u8(temp);
 		}
 	}
 
@@ -189,7 +181,6 @@ namespace uba
 		m_resetCas = info.resetCas;
 		m_remoteExecutionEnabled = info.remoteExecutionEnabled;
 		m_nameToHashTableEnabled = info.nameToHashTableEnabled;
-		m_memKillLoadPercent = info.memKillLoadPercent;
 		m_remoteLogEnabled = info.remoteLogEnabled;
 		m_remoteTraceEnabled = info.remoteTraceEnabled;
 		m_traceIOEnabled = info.traceIOEnabled;
@@ -202,24 +193,6 @@ namespace uba
 		if (m_detailedTrace)
 			m_server.SetWorkTracker(&m_trace);
 
-		m_memoryThreadEvent.Create(true);
-		if (info.checkMemory)
-		{
-			m_allowWaitOnMem = info.allowWaitOnMem;
-			m_allowKillOnMem = info.allowKillOnMem;
-
-			u64 memAvail;
-			u64 memTotal;
-			if (GetMemoryInfo(m_logger, memAvail, memTotal))
-			{
-				m_memAvail = memAvail;
-				m_memTotal = memTotal;
-				m_memRequiredToSpawn = Min(u64(double(m_memTotal) * double(100 - info.memWaitLoadPercent) / 100.0), 35ull * 1024 * 1024 * 1024);
-			}
-
-			m_memoryThread.Start([this]() { ThreadMemoryCheckLoop();  return 0; }, TC("UbaMemTrackLoop"));
-		}
-
 		#if PLATFORM_WINDOWS
 		m_localEnvironmentVariables.insert(TC("TMP"));
 		m_localEnvironmentVariables.insert(TC("TEMP"));
@@ -259,9 +232,6 @@ namespace uba
 
 	SessionServer::~SessionServer()
 	{
-		m_memoryThreadEvent.Set();
-		m_memoryThread.Wait();
-
 		StopTraceThread();
 
 		m_server.SetWorkTracker(nullptr);
@@ -272,13 +242,13 @@ namespace uba
 		for (ProcessHandle& p : m_queuedRemoteProcesses)
 		{
 			((RemoteProcess*)p.m_process)->m_server = nullptr;
-			p.Cancel(true);
+			p.Cancel();
 		}
 		m_queuedRemoteProcesses.clear();
 		for (const ProcessHandle& p : m_activeRemoteProcesses)
 		{
 			((RemoteProcess*)p.m_process)->m_server = nullptr;
-			p.Cancel(true);
+			p.Cancel();
 		}
 		m_activeRemoteProcesses.clear();
 
@@ -555,6 +525,11 @@ namespace uba
 		m_remoteProcessReturnedEvent = remoteProcessReturnedEvent;
 	}
 
+	void SessionServer::SetNativeProcessCreatedFunc(const Function<void(Process&)>& func)
+	{
+		m_nativeProcessCreatedFunc = func;
+	}
+
 	void SessionServer::WaitOnAllTasks()
 	{
 		while (true)
@@ -659,6 +634,33 @@ namespace uba
 		return m_outerScheduler;
 	}
 
+	ProcessHandle SessionServer::GetNewestLocalProcess()
+	{
+		u64 newestTime = 0;
+		ProcessImpl* newestProcess = nullptr;
+		SCOPED_FUTEX(m_processesLock, lock);
+		for (auto& kv : m_processes)
+		{
+			ProcessHandle& h = kv.second;
+			if (h.IsRemote())
+				continue;
+			auto& p = *(ProcessImpl*)h.m_process;
+			if (p.m_startTime <= newestTime)
+				continue;
+			if (p.m_parentProcess) // Only get root processes
+				continue;
+			if (p.WaitForExit(0)) // Process already exited, let it finish its work
+				continue;
+
+			newestTime = p.m_startTime;
+			newestProcess = &p;
+		}
+		if (!newestProcess)
+			return {};
+		ProcessHandle ph(newestProcess);
+		return ph;
+	}
+
 	void SessionServer::OnDisconnected(const Guid& clientUid, u32 clientId)
 	{
 		u32 returnCount = 0;
@@ -1934,132 +1936,6 @@ namespace uba
 		return true;
 	}
 
-	void SessionServer::ThreadMemoryCheckLoop()
-	{
-		u64 lastMessageTime = 0;
-
-		while (true)
-		{
-			if (m_memoryThreadEvent.IsSet(1000))
-				break;
-
-			#if 0
-			SCOPED_CRITICAL_SECTION(m_remoteProcessAndSessionLock, queueLock);
-			m_logger.Info(TC("RemoteQueue: %llu Active: %llu ConnectionCount: %u"), m_queuedRemoteProcesses.size(), m_activeRemoteProcesses.size(), m_connectionCount);
-
-			if (!m_connectionCount && !m_activeRemoteProcesses.empty())
-			{
-				for (auto& i : m_activeRemoteProcesses)
-				{
-					m_logger.Info(TC("ACTIVE PROCESS: %s"), i.GetStartInfo().GetDescription());
-				}
-				break;
-			}
-			#endif
-
-			u64 memAvail;
-			u64 memTotal;
-			u64 maxPageFile;
-			if (!GetMemoryInfo(m_logger, memAvail, memTotal, &maxPageFile))
-				m_memRequiredToSpawn = 0;
-			m_memAvail = memAvail;
-
-			bool allGood = false;
-			while (memAvail >= m_memRequiredToSpawn)
-			{
-				SCOPED_FUTEX(m_waitingProcessesLock, lock);
-				WaitingProcess* wp = m_oldestWaitingProcess;
-				if (!wp)
-				{
-					allGood = true;
-					break;
-				}
-				m_oldestWaitingProcess = wp->next;
-				if (m_newestWaitingProcess == wp)
-					m_newestWaitingProcess = nullptr;
-				wp->event.Set();
-				memAvail -= m_memRequiredToSpawn;
-			}
-
-			if (allGood)
-				continue;
-
-			u64 time = GetTime();
-			if (TimeToMs(time - lastMessageTime) > 5*1000)
-			{
-				lastMessageTime = time;
-				u32 delayCount = 0;
-				SCOPED_FUTEX(m_waitingProcessesLock, lock);
-				for (auto it = m_oldestWaitingProcess; it; it = it->next)
-					++delayCount;
-				lock.Leave();
-				if (delayCount)
-				{
-					m_logger.BeginScope();
-					m_logger.Info(TC("Delaying %u processes from spawning due to memory pressure (Available: %s Total: %s)"), delayCount, BytesToText(m_memAvail).str, BytesToText(m_memTotal).str);
-
-#if PLATFORM_WINDOWS
-					static bool hasBeenRunOnce;
-					if (!hasBeenRunOnce)
-					{
-						hasBeenRunOnce = true;
-						m_logger.Info(TC("NOTE - To mitigate this spawn delay it is recommended to make page file larger until you don't see these messages again (Or reduce number of max parallel processes)"));
-						m_logger.Info(TC("       Set max page file to a large number (like 128gb). It will not use disk space unless you actually start using that amount of committed memory"));
-						m_logger.Info(TC("       Also note, this is \"committed\" memory. Not memory in use. So you necessarily don't need more physical memory"));
-						MEMORYSTATUSEX memStatus = { sizeof(memStatus) };
-						GlobalMemoryStatusEx(&memStatus);
-						m_logger.Info(TC("  MaxPage:   %s"), BytesToText(maxPageFile));
-						m_logger.Info(TC("  TotalPhys: %s"), BytesToText(memStatus.ullTotalPhys));
-						m_logger.Info(TC("  AvailPhys: %s"), BytesToText(memStatus.ullAvailPhys));
-						m_logger.Info(TC("  TotalPage: %s"), BytesToText(memStatus.ullTotalPageFile));
-						m_logger.Info(TC("  AvailPage: %s"), BytesToText(memStatus.ullAvailPageFile));
-					}
-#endif
-					m_logger.EndScope();
-				}
-			}
-
-			if (!m_allowKillOnMem)
-				continue;
-
-			// TODO: This code path is not implemented yet... the cancel need to end up in a Requeue call.
-			UBA_ASSERT(false);
-
-			u64 memRequiredFree = u64(double(memTotal) * double(100 - m_memKillLoadPercent) / 100.0);
-			if (m_memAvail < memRequiredFree)
-			{
-				u64 newestTime = 0;
-				ProcessImpl* newestProcess = nullptr;
-				SCOPED_FUTEX(m_processesLock, lock);
-				for (auto& kv : m_processes)
-				{
-					ProcessHandle& h = kv.second;
-					if (h.IsRemote())
-						continue;
-					auto& p = *(ProcessImpl*)h.m_process;
-					if (p.m_startTime <= newestTime)
-						continue;
-					newestTime = p.m_startTime;
-					newestProcess = &p;
-				}
-
-				if (newestProcess)
-				{
-					newestProcess->Cancel(true);
-					newestProcess->WaitForExit(3000);
-				}
-
-				m_logger.Info(TC("Killed process due to memory pressure (Available: %s Total: %s)"), BytesToText(m_memAvail).str, BytesToText(m_memTotal).str);
-			}
-		}
-
-		SCOPED_FUTEX(m_waitingProcessesLock, lock);
-		for (auto it = m_oldestWaitingProcess; it; it = it->next)
-			it->event.Set();
-		m_oldestWaitingProcess = nullptr;
-		m_newestWaitingProcess = nullptr;
-	}
-
 	SessionServer::RemoteProcess* SessionServer::DequeueProcess(ClientSession& session, u32 sessionId, u32 clientId)
 	{
 		//TrackWorkScope tws(m_trace, AsView(TC("DequeueProcess")), ColorWork);
@@ -2180,38 +2056,10 @@ namespace uba
 		return str.Appendf(TC("%s"), findIt->second.GetStartInfo().GetDescription()).data;
 	}
 
-	bool SessionServer::PrepareProcess(ProcessImpl& process, bool isChild, StringBufferBase& outRealApplication, const tchar*& outRealWorkingDir)
+	bool SessionServer::ProcessNativeCreated(ProcessImpl& process)
 	{
-		if (!Session::PrepareProcess(process, isChild, outRealApplication, outRealWorkingDir))
-			return false;
-
-		if (!m_memTotal || !m_allowWaitOnMem || isChild)
-			return true;
-
-		if (m_memAvail >= m_memRequiredToSpawn)
-			return true;
-
-		u64 startWait = GetTime();
-
-		WaitingProcess wp;
-		wp.event.Create(true);
-
-		SCOPED_FUTEX(m_waitingProcessesLock, lock);
-		if (m_memoryThreadEvent.IsSet(0))
-			return false;
-
-		if (!m_oldestWaitingProcess)
-			m_oldestWaitingProcess = &wp;
-		else
-			m_newestWaitingProcess->next = &wp;
-		m_newestWaitingProcess = &wp;
-		lock.Leave();
-
-		wp.event.IsSet();
-
-		u64 waitTime = GetTime() - startWait;
-		m_logger.Info(TC("Waited %s for memory pressure to go down (Available: %s Total: %s)"), TimeToText(waitTime).str, BytesToText(m_memAvail).str, BytesToText(m_memTotal).str);
-
+		if (m_nativeProcessCreatedFunc)
+			m_nativeProcessCreatedFunc(process);
 		return true;
 	}
 
@@ -2394,9 +2242,9 @@ namespace uba
 		lock.Leave();
 
 		float cpuLoad = UpdateCpuLoad();
-		u64 memAvail = m_memAvail;
-		u64 memTotal = m_memTotal;
-
+		u64 memAvail;
+		u64 memTotal;
+		GetMemoryInfo(m_logger, memAvail, memTotal);
 
 		if (m_traceIOEnabled)
 		{
diff --git a/Engine/Source/Programs/UnrealBuildAccelerator/Common/Private/UbaStats.cpp b/Engine/Source/Programs/UnrealBuildAccelerator/Common/Private/UbaStats.cpp
--- a/Engine/Source/Programs/UnrealBuildAccelerator/Common/Private/UbaStats.cpp
+++ b/Engine/Source/Programs/UnrealBuildAccelerator/Common/Private/UbaStats.cpp
@@ -136,8 +136,8 @@ namespace uba
 		exitTime += other.exitTime;
 		wallTime += other.wallTime;
 		cpuTime += other.cpuTime;
-		detoursMemory = Max(detoursMemory, other.detoursMemory);
-		peakMemory = Max(peakMemory, other.peakMemory);
+		//detoursMemory = Max(detoursMemory, other.detoursMemory); // Handled on the outside
+		//peakMemory = Max(peakMemory, other.peakMemory); // Handled on the outside
 		iopsRead += other.iopsRead;
 		iopsWrite += other.iopsWrite;
 		iopsOther += other.iopsOther;
diff --git a/Engine/Source/Programs/UnrealBuildAccelerator/Common/Private/UbaTrace.cpp b/Engine/Source/Programs/UnrealBuildAccelerator/Common/Private/UbaTrace.cpp
--- a/Engine/Source/Programs/UnrealBuildAccelerator/Common/Private/UbaTrace.cpp
+++ b/Engine/Source/Programs/UnrealBuildAccelerator/Common/Private/UbaTrace.cpp
@@ -508,6 +508,18 @@ namespace uba
 		BEGIN_TRACE_ENTRY(RemoteExecutionDisabled);
 	}
 
+	void Trace::SchedulerUpdate(bool spawnPaused)
+	{
+		BEGIN_TRACE_ENTRY(SchedulerUpdate);
+		writer.WriteBool(spawnPaused);
+	}
+
+	void Trace::SchedulerKillProcess(u32 processId)
+	{
+		BEGIN_TRACE_ENTRY(SchedulerKillProcess);
+		writer.WriteU32(processId);
+	}
+
 	void Trace::CacheBeginFetch(u32 fetchId, const tchar* description)
 	{
 		BEGIN_TRACE_ENTRY(CacheBeginFetch);
diff --git a/Engine/Source/Programs/UnrealBuildAccelerator/Common/Public/UbaApplication.h b/Engine/Source/Programs/UnrealBuildAccelerator/Common/Public/UbaApplication.h
--- a/Engine/Source/Programs/UnrealBuildAccelerator/Common/Public/UbaApplication.h
+++ b/Engine/Source/Programs/UnrealBuildAccelerator/Common/Public/UbaApplication.h
@@ -2,7 +2,7 @@
 
 #pragma once
 
-#include "UbaPlatform.h"
+#include "UbaEnvironment.h"
 #include "UbaStringBuffer.h"
 
 namespace uba
diff --git a/Engine/Source/Programs/UnrealBuildAccelerator/Common/Public/UbaProcess.h b/Engine/Source/Programs/UnrealBuildAccelerator/Common/Public/UbaProcess.h
--- a/Engine/Source/Programs/UnrealBuildAccelerator/Common/Public/UbaProcess.h
+++ b/Engine/Source/Programs/UnrealBuildAccelerator/Common/Public/UbaProcess.h
@@ -22,7 +22,7 @@ namespace uba
 		virtual const Vector<ProcessLogLine>& GetLogLines() const = 0;
 		virtual const Vector<u8>& GetTrackedInputs() const = 0;
 		virtual const Vector<u8>& GetTrackedOutputs() const = 0;
-		virtual void Cancel(bool terminate) { UBA_ASSERT(false); }
+		virtual bool Cancel() { UBA_ASSERT(false); return false; }
 		virtual bool IsCalledFromThis() const { return false; }
 		virtual const tchar* GetExecutingHost() const { UBA_ASSERT(false); return nullptr; }
 		virtual bool IsRemote() const { UBA_ASSERT(false); return false; }
@@ -53,7 +53,7 @@ namespace uba
 		virtual const Vector<ProcessLogLine>& GetLogLines() const override { return m_logLines; }
 		virtual const Vector<u8>& GetTrackedInputs() const override { return m_trackedInputs; }
 		virtual const Vector<u8>& GetTrackedOutputs() const override { return m_trackedOutputs; }
-		virtual void Cancel(bool terminate) override;
+		virtual bool Cancel() override;
 		virtual bool IsCalledFromThis() const override;
 		virtual const tchar* GetExecutingHost() const override { return TC(""); }
 		virtual bool IsRemote() const override { return false; }
@@ -147,8 +147,10 @@ namespace uba
 
 		u32 m_nativeProcessId = 0;
 		u32 m_nativeProcessExitCode = ~0u;
+		u64 m_nativeProcessCreationTime = 0;
 		u32 m_exitCode = ~0u;
 		u32 m_messageCount = 0;
+		Atomic<bool> m_cancelAllowed = true;
 		Atomic<bool> m_cancelled;
 		Atomic<bool> m_hasExited;
 		MessageType m_firstMessageError = MessageType_None;
diff --git a/Engine/Source/Programs/UnrealBuildAccelerator/Common/Public/UbaProcessHandle.h b/Engine/Source/Programs/UnrealBuildAccelerator/Common/Public/UbaProcessHandle.h
--- a/Engine/Source/Programs/UnrealBuildAccelerator/Common/Public/UbaProcessHandle.h
+++ b/Engine/Source/Programs/UnrealBuildAccelerator/Common/Public/UbaProcessHandle.h
@@ -44,7 +44,7 @@ namespace uba
 		const Vector<u8>& GetTrackedOutputs() const;		// If ProcessStartInfo.trackInputs was true, this will return a buffer to wchar/char strings 
 		u64 GetTotalProcessorTime() const;					// Total used cpu time (in time units. Use TimeToMs etc)
 		u64 GetTotalWallTime() const;						// Total wall time (in time units. Use TimeToMs etc)
-		void Cancel(bool terminate) const;					// Request to cancel process. Set terminate to true to quickly kill process
+		bool Cancel() const;								// Request to cancel process. Returns false if failed to cancel (might be at the end of exiting process)
 		const tchar* GetExecutingHost() const;				// Host that is executing process.
 		bool IsRemote() const;								// Returns true if process is a remote process
 		ProcessExecutionType GetExecutionType() const;		// Returns execution type
diff --git a/Engine/Source/Programs/UnrealBuildAccelerator/Common/Public/UbaScheduler.h b/Engine/Source/Programs/UnrealBuildAccelerator/Common/Public/UbaScheduler.h
--- a/Engine/Source/Programs/UnrealBuildAccelerator/Common/Public/UbaScheduler.h
+++ b/Engine/Source/Programs/UnrealBuildAccelerator/Common/Public/UbaScheduler.h
@@ -36,6 +36,10 @@ namespace uba
 		bool forceNative = false; // Force all processes to run native (not detoured)
 		bool writeToCache = false; // Set to true in combination with setting cacheClient to populate cache
 		bool useThreadToExitRemoteProcess = false; // Set this to true to move exit process call from job to thread. 
+		bool trackMemory = false;
+		u8 memTracingLevel = 0;
+		u8 memStartWaitPercent = 85; // When memory usage goes above this percent, no new processes will be spawned until back below
+		u8 memStartKillPercent = 95; // When memory usage goes above this percent, newest processes will be killed to bring it back below
 		ConfigTable* processConfigs = nullptr;
 	};
 
@@ -59,6 +63,9 @@ namespace uba
 		u32 dependencyCount = 0; // Number of elements in dependencies
 
 		u32 cacheBucketId = 0; // Bucket that cache should be fetched from. Zero means that it will not fetch anything
+		u32 memoryGroupId = 0; // Group that calculates average of memory usage.
+
+		u64 predictedMemoryUsage = 0; // If zero, this will not be used for memory predictions
 	};
 
 	////////////////////////////////////////////////////////////////////////////////////////////////////
@@ -116,9 +123,10 @@ namespace uba
 		};
 
 		UBA_NOINLINE void ThreadLoop();
+
 		void SkipAllQueued();
 		void Cleanup();
-		void RemoteProcessReturned(Process& process);
+		bool ProcessReturned(Process& process, bool isLocal);
 		void HandleCacheMissed(u32 processIndex);
 		void RemoteSlotAvailable(bool isCrossArchitecture);
 		void ProcessExited(ExitProcessInfo* info, const ProcessHandle& handle);
@@ -135,6 +143,16 @@ namespace uba
 		void UpdateActiveProcessCounter(bool isLocal, int offset);
 		void FinishProcess(const ProcessHandle& handle);
 
+		void KillNewestLocalProcess();
+		void NativeProcessCreated(Process& process);
+		void LocalProcessStart(ProcessStartInfo2& info);
+		bool LocalProcessOkToSpawn();
+		void LocalProcessExit(ProcessStartInfo2& info, const ProcessHandle& handle);
+		void LoadMemTrackTable();
+		void SaveMemTrackTable();
+		UBA_NOINLINE void ThreadMemoryCheckLoop();
+
+
 		SessionServer& m_session;
 		u32 m_maxLocalProcessors;
 		u32 m_maxParallelCacheQueries;
@@ -190,6 +208,26 @@ namespace uba
 		WorkManagerImpl* m_cacheWorkManager = nullptr;
 		List<Thread> m_finishedRemoteProcessThreads;
 
+		CriticalSection m_wasReturnedLock;
+
+		bool m_memTrack = false;
+		Event m_memThreadEvent;
+		Thread m_memThread;
+		bool m_memDelaySpawning = false;
+		Atomic<u64> m_memEstimatedUsage;
+		u64 m_memWaitThreshold = 0;
+		u64 m_memKillThreshold = 0;
+		u8 m_memTracingLevel = 0;
+
+		struct MemGroupStats;
+		struct MemGroupEntry;
+		Futex m_memGroupLookupLock;
+		UnorderedMap<u32, MemGroupStats> m_memGroupLookup;
+		UnorderedMap<u32, MemGroupEntry> m_memGroupDb;
+
+		Futex m_processesLock;
+		UnorderedMap<StringKey, ExitProcessInfo*> m_processes;
+
 		Scheduler(const Scheduler&) = delete;
 		void operator=(const Scheduler&) = delete;
 	};
diff --git a/Engine/Source/Programs/UnrealBuildAccelerator/Common/Public/UbaSession.h b/Engine/Source/Programs/UnrealBuildAccelerator/Common/Public/UbaSession.h
--- a/Engine/Source/Programs/UnrealBuildAccelerator/Common/Public/UbaSession.h
+++ b/Engine/Source/Programs/UnrealBuildAccelerator/Common/Public/UbaSession.h
@@ -162,6 +162,8 @@ namespace uba
 		virtual bool WriteOutputFileInternal(const StringKey& fileNameKey, StringView filePath, bool deleteInternalMapping);
 
 		virtual bool ProcessThreadStart(ProcessImpl& process);
+		virtual bool ProcessNativeCreated(ProcessImpl& process);
+		virtual bool ProcessCancelled(ProcessImpl& process);
 		virtual bool GetInitResponse(InitResponse& out, const InitMessage& msg);
 		virtual bool CreateFile(CreateFileResponse& out, const CreateFileMessage& msg);
 		virtual bool CreateFileForRead(CreateFileResponse& out, TrackWorkScope& tws, const StringView& fileName, const StringKey& fileNameKey, ProcessImpl& process, const ApplicationRules& rules);
@@ -315,7 +317,7 @@ namespace uba
 		UnorderedMap<TString, TString> m_outputFiles;
 
 		Futex m_activeFilesLock;
-		struct ActiveFile { TString name; StringKey nameKey; };
+		struct ActiveFile { TString name; StringKey nameKey; Process* owner; };
 		UnorderedMap<u32, ActiveFile> m_activeFiles;
 
 		Futex m_virtualSourceFilesLock;
diff --git a/Engine/Source/Programs/UnrealBuildAccelerator/Common/Public/UbaSessionServer.h b/Engine/Source/Programs/UnrealBuildAccelerator/Common/Public/UbaSessionServer.h
--- a/Engine/Source/Programs/UnrealBuildAccelerator/Common/Public/UbaSessionServer.h
+++ b/Engine/Source/Programs/UnrealBuildAccelerator/Common/Public/UbaSessionServer.h
@@ -48,6 +48,9 @@ namespace uba
 		// Callback that is called when process is returned. This could be because a client unexpectedly disconnected or is running out of memory
 		void SetRemoteProcessReturnedEvent(const Function<void(Process&)>& remoteProcessReturnedEvent);
 
+		// Callback that is called when process is returned. This could be because a client unexpectedly disconnected or is running out of memory
+		void SetNativeProcessCreatedFunc(const Function<void(Process&)>& func);
+
 		// Wait for all queued up/active tasks to finish
 		void WaitOnAllTasks();
 
@@ -83,6 +86,8 @@ namespace uba
 		void SetOuterScheduler(Scheduler* scheduler);
 		Scheduler* GetOuterScheduler();
 
+		ProcessHandle GetNewestLocalProcess();
+
 	protected:
 		struct ClientSession;
 
@@ -95,8 +100,6 @@ namespace uba
 		bool WriteDirectoryTable(ClientSession& session, BinaryReader& reader, BinaryWriter& writer);
 		bool WriteNameToHashTable(BinaryReader& reader, BinaryWriter& writer, u32 requestedSize);
 
-		void ThreadMemoryCheckLoop();
-
 		class RemoteProcess;
 
 		RemoteProcess* DequeueProcess(ClientSession& session, u32 sessionId, u32 clientId);
@@ -106,7 +109,7 @@ namespace uba
 		ProcessHandle GetProcess(u32 processId);
 		TString GetProcessDescription(u32 processId);
 
-		virtual bool PrepareProcess(ProcessImpl& process, bool isChild, StringBufferBase& outRealApplication, const tchar*& outRealWorkingDir) override final;
+		virtual bool ProcessNativeCreated(ProcessImpl& process) override;
 		virtual bool CreateFileForRead(CreateFileResponse& out, TrackWorkScope& tws, const StringView& fileName, const StringKey& fileNameKey, ProcessImpl& process, const ApplicationRules& rules) override final;
 		virtual bool GetOutputFileSizeInternal(u64& outSize, const StringKey& fileNameKey, StringView filePath) override;
 		virtual bool GetOutputFileDataInternal(void* outData, const StringKey& fileNameKey, StringView filePath, bool deleteInternalMapping) override;
@@ -137,6 +140,8 @@ namespace uba
 		ReaderWriterLock m_remoteProcessReturnedEventLock;
 		Function<void(Process&)> m_remoteProcessReturnedEvent;
 
+		Function<void(Process&)> m_nativeProcessCreatedFunc;
+
 		CriticalSection m_remoteProcessAndSessionLock; // Can be re-entrant.
 		List<ProcessHandle> m_queuedRemoteProcesses;
 		UnorderedSet<ProcessHandle> m_activeRemoteProcesses;
@@ -202,18 +207,6 @@ namespace uba
 		struct CrossArchitectureMapping { TString from; TString to; };
 		Vector<CrossArchitectureMapping> m_crossArchitectureMappings;
 
-		Event m_memoryThreadEvent;
-		Thread m_memoryThread;
-		Atomic<u64> m_memAvail;
-		u64 m_memTotal = 0;
-		u64 m_memRequiredToSpawn = 0;
-		u8 m_memKillLoadPercent = 0;
-		struct WaitingProcess { Event event; WaitingProcess* next = nullptr; };
-		WaitingProcess* m_oldestWaitingProcess = nullptr;
-		WaitingProcess* m_newestWaitingProcess = nullptr;
-		Futex m_waitingProcessesLock;
-		bool m_allowWaitOnMem = false;
-		bool m_allowKillOnMem = false;
 		bool m_remoteLogEnabled = false;
 		bool m_remoteTraceEnabled = false;
 		bool m_traceIOEnabled = false;
diff --git a/Engine/Source/Programs/UnrealBuildAccelerator/Common/Public/UbaSessionServerCreateInfo.h b/Engine/Source/Programs/UnrealBuildAccelerator/Common/Public/UbaSessionServerCreateInfo.h
--- a/Engine/Source/Programs/UnrealBuildAccelerator/Common/Public/UbaSessionServerCreateInfo.h
+++ b/Engine/Source/Programs/UnrealBuildAccelerator/Common/Public/UbaSessionServerCreateInfo.h
@@ -15,14 +15,9 @@ namespace uba
 		void Apply(const Config& config);
 
 		NetworkServer& server;
-		u8 memWaitLoadPercent = 85; // When memory usage goes above this percent, no new processes will be spawned until back below
-		u8 memKillLoadPercent = 95; // When memory usage goes above this percent, newest processes will be killed to bring it back below
 		bool resetCas = false;
 		bool remoteExecutionEnabled = true;
 		bool nameToHashTableEnabled = true;
-		bool checkMemory = true;
-		bool allowWaitOnMem = false;
-		bool allowKillOnMem = false;
 		bool remoteLogEnabled = false; // If Uba is built in debug, then the logs will be sent back to server
 		bool remoteTraceEnabled = false; // If this is true, the agents will run trace and send the .uba file back to server
 		bool traceIOEnabled = true; // Enable IO information for storage 
diff --git a/Engine/Source/Programs/UnrealBuildAccelerator/Common/Public/UbaTrace.h b/Engine/Source/Programs/UnrealBuildAccelerator/Common/Public/UbaTrace.h
--- a/Engine/Source/Programs/UnrealBuildAccelerator/Common/Public/UbaTrace.h
+++ b/Engine/Source/Programs/UnrealBuildAccelerator/Common/Public/UbaTrace.h
@@ -65,6 +65,8 @@ namespace uba
 		UBA_TRACE_TYPE(TaskBegin) \
 		UBA_TRACE_TYPE(TaskHint) \
 		UBA_TRACE_TYPE(TaskEnd) \
+		UBA_TRACE_TYPE(SchedulerUpdate) \
+		UBA_TRACE_TYPE(SchedulerKillProcess) \
 
 	enum TraceType : u8
 	{
@@ -73,7 +75,7 @@ namespace uba
 		#undef UBA_TRACE_TYPE
 	};
 
-	static constexpr u32 TraceVersion = 46;
+	static constexpr u32 TraceVersion = 47;
 	static constexpr u32 TraceReadCompatibilityVersion = 6;
 
 	class Trace : public WorkTracker
@@ -111,6 +113,8 @@ namespace uba
 		void StatusUpdate(u32 statusRow, u32 statusColumn, const tchar* statusText, LogEntryType statusType, const tchar* statusLink);
 		void DriveUpdate(const tchar drive, u8 busyPercent, u32 readCount, u64 readBytes, u32 writeCount, u64 writeBytes);
 		void RemoteExecutionDisabled();
+		void SchedulerUpdate(bool spawnPaused);
+		void SchedulerKillProcess(u32 processId);
 
 		void CacheBeginFetch(u32 fetchId, const tchar* description);
 		void CacheEndFetch(u32 fetchId, bool success, const u8* data, u64 dataSize);
diff --git a/Engine/Source/Programs/UnrealBuildAccelerator/Core/Public/UbaHash.h b/Engine/Source/Programs/UnrealBuildAccelerator/Core/Public/UbaHash.h
--- a/Engine/Source/Programs/UnrealBuildAccelerator/Core/Public/UbaHash.h
+++ b/Engine/Source/Programs/UnrealBuildAccelerator/Core/Public/UbaHash.h
@@ -10,7 +10,7 @@ namespace uba
 
 	struct StringKey
 	{
-		constexpr StringKey() : a(0), b(0) {}
+		constexpr StringKey(u64 a_ = 0, u64 b_ = 0) : a(a_), b(b_) {}
 		u64 a;
 		u64 b;
 		bool operator==(const StringKey& o) const { return a == o.a && b == o.b; }
diff --git a/Engine/Source/Programs/UnrealBuildAccelerator/Core/Public/UbaTimer.h b/Engine/Source/Programs/UnrealBuildAccelerator/Core/Public/UbaTimer.h
--- a/Engine/Source/Programs/UnrealBuildAccelerator/Core/Public/UbaTimer.h
+++ b/Engine/Source/Programs/UnrealBuildAccelerator/Core/Public/UbaTimer.h
@@ -12,12 +12,10 @@ namespace uba
 	u64 GetFrequency();						// Number of counts per second. GetTime() / GetFrequency() = seconds
 	u64 GetSystemTimeUs();					// Get time in microseconds since 1 Jan 1970 00:00 UTC
 
-	inline u64 TimeToTick(u64 time, u64 frequency)	{ return time * 10'000'000 / frequency; }
 	inline u64 TimeToUs(u64 time, u64 frequency)	{ return time * 1'000'000 / frequency; }
 	inline u64 TimeToMs(u64 time, u64 frequency)	{ return time * 1'000 / frequency; }
 	inline float TimeToS(u64 time, u64 frequency)	{ return float(double(time) / double(frequency)); }
 
-	inline u64 TimeToTick(u64 time)					{ return TimeToTick(time, GetFrequency()); }
 	inline u64 TimeToUs(u64 time)					{ return TimeToUs(time, GetFrequency()); }
 	inline u64 TimeToMs(u64 time)					{ return TimeToMs(time, GetFrequency()); }
 	inline float TimeToS(u64 time)					{ return TimeToS(time, GetFrequency()); }
diff --git a/Engine/Source/Programs/UnrealBuildAccelerator/Detours/Private/Windows/UbaDetoursFunctionsNtDll.inl b/Engine/Source/Programs/UnrealBuildAccelerator/Detours/Private/Windows/UbaDetoursFunctionsNtDll.inl
--- a/Engine/Source/Programs/UnrealBuildAccelerator/Detours/Private/Windows/UbaDetoursFunctionsNtDll.inl
+++ b/Engine/Source/Programs/UnrealBuildAccelerator/Detours/Private/Windows/UbaDetoursFunctionsNtDll.inl
@@ -1194,7 +1194,12 @@ NTSTATUS NTAPI Shared_NtCreateFile(bool IsCreateFunc, PHANDLE hFileHandle, ACCES
 
 	if (lpFileName[0] == '^') // It is a HANDLE from session process
 	{
-		UBA_ASSERTF(!isWrite, TC("Mapped file cant be open for write (%s)"), info.originalName);
+		if (isWrite)
+		{
+			Rpc_WriteLogf(TC("Mapped file %s cant be open for write (%s)"), lpFileName, info.originalName);
+			UBA_ASSERTF(false, TC("Mapped file %s cant be open for write (%s)"), lpFileName, info.originalName);
+			return STATUS_UNSUCCESSFUL;
+		}
 
 		const wchar_t* handleStr = lpFileName + 1;
 		const wchar_t* handleStrEnd = wcschr(handleStr, '-');
@@ -1603,6 +1608,11 @@ NTSTATUS NTAPI Detoured_NtClose(HANDLE handle)
 				Rpc_CreateFileW(fixedName, fileNameKey, AccessFlag_Write, temp, sizeof_array(temp), size, fo->closeId, true);
 			}
 
+			PROCESS_MEMORY_COUNTERS_EX pmc;
+			GetProcessMemoryInfo(GetCurrentProcess(), (PROCESS_MEMORY_COUNTERS*)&pmc, sizeof(pmc));
+			AtomicMax(g_stats.peakMemory, pmc.PagefileUsage);
+
+
 			if (!fo->newName.empty())
 			{
 				// It might be that same process will open it again, so we will need to update the mapping table
diff --git a/Engine/Source/Programs/UnrealBuildAccelerator/Detours/Private/Windows/UbaDetoursFunctionsWin.cpp b/Engine/Source/Programs/UnrealBuildAccelerator/Detours/Private/Windows/UbaDetoursFunctionsWin.cpp
--- a/Engine/Source/Programs/UnrealBuildAccelerator/Detours/Private/Windows/UbaDetoursFunctionsWin.cpp
+++ b/Engine/Source/Programs/UnrealBuildAccelerator/Detours/Private/Windows/UbaDetoursFunctionsWin.cpp
@@ -667,6 +667,10 @@ void SendExitMessage(DWORD exitCode, u64 startTime)
 		return;
 	g_exitMessageSent = true;
 
+	PROCESS_MEMORY_COUNTERS_EX pmc;
+	GetProcessMemoryInfo(GetCurrentProcess(), (PROCESS_MEMORY_COUNTERS*)&pmc, sizeof(pmc));
+	AtomicMax(g_stats.peakMemory, pmc.PagefileUsage);
+
 	if (g_consoleStringIndex)
 		Shared_WriteConsole(L"\n", 1, 0);
 
diff --git a/Engine/Source/Programs/UnrealBuildAccelerator/Host/Private/UbaExports.cpp b/Engine/Source/Programs/UnrealBuildAccelerator/Host/Private/UbaExports.cpp
--- a/Engine/Source/Programs/UnrealBuildAccelerator/Host/Private/UbaExports.cpp
+++ b/Engine/Source/Programs/UnrealBuildAccelerator/Host/Private/UbaExports.cpp
@@ -382,16 +382,21 @@ extern "C"
 		return handle->GetHash();
 	}
 
+	uba::u64 TimeToTick(uba::u64 time)
+	{
+		return time * 10'000'000 / uba::GetFrequency();
+	}
+
 	// 100ns ticks
 	uba::u64 ProcessHandle_GetTotalProcessorTime(const uba::ProcessHandle* handle)
 	{
-		return uba::TimeToTick(handle->GetTotalProcessorTime());
+		return TimeToTick(handle->GetTotalProcessorTime());
 	}
 
 	// 100ns ticks
 	uba::u64 ProcessHandle_GetTotalWallTime(const uba::ProcessHandle* handle)
 	{
-		return uba::TimeToTick(handle->GetTotalWallTime());
+		return TimeToTick(handle->GetTotalWallTime());
 	}
 
 	void ProcessHandle_TraverseOutputFiles(const uba::ProcessHandle* handle, ProcessHandle_OutputFileCallback* callback, void* userData)
@@ -406,7 +411,7 @@ extern "C"
 
 	void ProcessHandle_Cancel(const uba::ProcessHandle* handle, bool terminate)
 	{
-		handle->Cancel(terminate);
+		handle->Cancel();
 	}
 
 	void ProcessHandle_Destroy(const uba::ProcessHandle* handle)
@@ -435,8 +440,6 @@ extern "C"
 		info->resetCas = resetCas;
 		info->shouldWriteToDisk = writeToDisk;
 		info->detailedTrace = detailedTrace;
-		info->allowWaitOnMem = allowWaitOnMem;
-		info->allowKillOnMem = allowKillOnMem;
 		info->storeIntermediateFilesCompressed = storeIntermediateFilesCompressed;
 		//info->remoteTraceEnabled = true;
 		info->remoteLogEnabled = true;
@@ -736,15 +739,15 @@ extern "C"
 
 	uba::u32 Scheduler_EnqueueProcess(uba::Scheduler* scheduler, const uba::ProcessStartInfo& info, float weight, const void* knownInputs, uba::u32 knownInputsBytes, uba::u32 knownInputsCount)
 	{
-		uba::EnqueueProcessInfo epi(info);
-		epi.weight = weight;
-		epi.knownInputs = knownInputs;
-		epi.knownInputsBytes = knownInputsBytes;
-		epi.knownInputsCount = knownInputsCount;
-		return scheduler->EnqueueProcess(epi);
+		return Scheduler_EnqueueProcess2(scheduler, info, weight, true, true, nullptr, 0, knownInputs, knownInputsBytes, knownInputsCount, 0);
 	}
 
 	uba::u32 Scheduler_EnqueueProcess2(uba::Scheduler* scheduler, const uba::ProcessStartInfo& info, float weight, bool canDetour, bool canExecuteRemotely, const uba::u32* dependencies, uba::u32 dependencyCount, const void* knownInputs, uba::u32 knownInputsBytes, uba::u32 knownInputsCount, uba::u32 cacheBucket)
+	{
+		return Scheduler_EnqueueProcess3(scheduler, info, weight, canDetour, canExecuteRemotely, dependencies, dependencyCount, knownInputs, knownInputsBytes, knownInputsCount, cacheBucket, 0, 0);
+	}
+
+	uba::u32 Scheduler_EnqueueProcess3(uba::Scheduler* scheduler, const uba::ProcessStartInfo& info, float weight, bool canDetour, bool canExecuteRemotely, const uba::u32* dependencies, uba::u32 dependencyCount, const void* knownInputs, uba::u32 knownInputsBytes, uba::u32 knownInputsCount, uba::u32 cacheBucket, uba::u32 memoryGroup, uba::u64 predictedMemoryUsage)
 	{
 		uba::EnqueueProcessInfo epi(info);
 		epi.weight = weight;
@@ -756,6 +759,8 @@ extern "C"
 		epi.canDetour = canDetour;
 		epi.canExecuteRemotely = canExecuteRemotely;
 		epi.cacheBucketId = cacheBucket;
+		epi.memoryGroupId = memoryGroup;
+		epi.predictedMemoryUsage = predictedMemoryUsage;
 		return scheduler->EnqueueProcess(epi);
 	}
 
diff --git a/Engine/Source/Programs/UnrealBuildAccelerator/Host/Public/UbaExports.h b/Engine/Source/Programs/UnrealBuildAccelerator/Host/Public/UbaExports.h
--- a/Engine/Source/Programs/UnrealBuildAccelerator/Host/Public/UbaExports.h
+++ b/Engine/Source/Programs/UnrealBuildAccelerator/Host/Public/UbaExports.h
@@ -156,6 +156,7 @@ extern "C"
 	UBA_API void Scheduler_Start(uba::Scheduler* scheduler);
 	UBA_API uba::u32 Scheduler_EnqueueProcess(uba::Scheduler* scheduler, const uba::ProcessStartInfo& info, float weight = 1.0f, const void* knownInputs = nullptr, uba::u32 knownInputsBytes = 0, uba::u32 knownInputsCount = 0);
 	UBA_API uba::u32 Scheduler_EnqueueProcess2(uba::Scheduler* scheduler, const uba::ProcessStartInfo& info, float weight, bool canDetour, bool canExecuteRemotely, const uba::u32* dependencies, uba::u32 dependencyCount, const void* knownInputs, uba::u32 knownInputsBytes, uba::u32 knownInputsCount, uba::u32 cacheBucket);
+	UBA_API uba::u32 Scheduler_EnqueueProcess3(uba::Scheduler* scheduler, const uba::ProcessStartInfo& info, float weight, bool canDetour, bool canExecuteRemotely, const uba::u32* dependencies, uba::u32 dependencyCount, const void* knownInputs, uba::u32 knownInputsBytes, uba::u32 knownInputsCount, uba::u32 cacheBucket, uba::u32 memoryGroup, uba::u64 predictedMemoryUsage);
 	UBA_API void Scheduler_SetMaxLocalProcessors(uba::Scheduler* scheduler, uba::u32 maxLocalProcessors);
 	UBA_API void Scheduler_Stop(uba::Scheduler* scheduler);
 	UBA_API void Scheduler_Cancel(uba::Scheduler* scheduler);
diff --git a/Engine/Source/Programs/UnrealBuildAccelerator/Visualizer/Private/UbaTraceReader.cpp b/Engine/Source/Programs/UnrealBuildAccelerator/Visualizer/Private/UbaTraceReader.cpp
--- a/Engine/Source/Programs/UnrealBuildAccelerator/Visualizer/Private/UbaTraceReader.cpp
+++ b/Engine/Source/Programs/UnrealBuildAccelerator/Visualizer/Private/UbaTraceReader.cpp
@@ -50,6 +50,9 @@ namespace uba
 		progressProcessesTotal = 0;
 		progressProcessesDone = 0;
 		progressErrorCount = 0;
+		lastKillProcessTime = 0;
+		lastSpawningDelayStartTime = 0;
+		lastSpawningDelayEndTime = 0;
 		remoteExecutionDisabled = false;
 		finished = true;
 	};
@@ -296,6 +299,7 @@ namespace uba
 		m_memoryPos = memoryBegin + pos;
 		m_memoryEnd = memoryBegin + end;
 		m_memoryBegin = memoryBegin;
+		out.finished = true;
 
 		return res;
 	}
@@ -1292,7 +1296,7 @@ namespace uba
 		{
 			u32 id = u32(reader.Read7BitEncoded());
 			TString desc = reader.ReadString();
-			ProcessBegin(out, 0, id, time, std::move(desc), {})->type = TraceView::ProcessType_CacheFetch;
+			ProcessBegin(out, 0, id, time, std::move(desc), {})->type = TraceView::ProcessType_CacheFetchTest;
 			break;
 		}
 		case TraceType_CacheEndFetch:
@@ -1311,6 +1315,10 @@ namespace uba
 			StorageStats storageStats;
 			const u8* dataStart = reader.GetPositionData();
 			cacheStats.Read(reader, out.version);
+
+			if (!cacheStats.fetchEntries.count)
+				process.type = TraceView::ProcessType_CacheFetchDownload;
+
 			if (success || out.version >= 29)
 			{
 				storageStats.Read(reader, out.version);
@@ -1397,6 +1405,23 @@ namespace uba
 				cacheSummary.lines.push_back(reader.ReadString());
 			break;
 		}
+		case TraceType_SchedulerUpdate:
+		{
+			if (reader.ReadBool())
+			{
+				out.lastSpawningDelayStartTime = time;
+				out.lastSpawningDelayEndTime = ~0llu - MsToTime(10000);
+			}
+			else
+				out.lastSpawningDelayEndTime = time;
+			break;
+		}
+		case TraceType_SchedulerKillProcess:
+		{
+			u32 processId = reader.ReadU32();(void)processId;
+			out.lastKillProcessTime = time;
+			break;
+		}
 		default:
 		{
 			return m_logger.Error(TC("Unknown trace type found in stream. UbaVisualizer too old?"));
diff --git a/Engine/Source/Programs/UnrealBuildAccelerator/Visualizer/Private/UbaTraceReader.h b/Engine/Source/Programs/UnrealBuildAccelerator/Visualizer/Private/UbaTraceReader.h
--- a/Engine/Source/Programs/UnrealBuildAccelerator/Visualizer/Private/UbaTraceReader.h
+++ b/Engine/Source/Programs/UnrealBuildAccelerator/Visualizer/Private/UbaTraceReader.h
@@ -13,7 +13,7 @@ namespace uba
 	class TraceView
 	{
 	public:
-		enum ProcessType : u8 { ProcessType_Normal, ProcessType_CacheFetch, ProcessType_Task };
+		enum ProcessType : u8 { ProcessType_Normal, ProcessType_CacheFetchTest, ProcessType_CacheFetchDownload, ProcessType_Task };
 
 		struct Process
 		{
@@ -185,6 +185,9 @@ namespace uba
 		u64 traceSystemStartTimeUs = 0;
 		u64 startTime = 0;
 		u64 frequency = 0;
+		u64 lastKillProcessTime = 0;
+		u64 lastSpawningDelayStartTime = 0;
+		u64 lastSpawningDelayEndTime = 0;
 		u32 totalProcessActiveCount = 0;
 		u32 totalProcessExitedCount = 0;
 		u32 activeSessionCount = 0;
diff --git a/Engine/Source/Programs/UnrealBuildAccelerator/Visualizer/Private/UbaVisualizer.cpp b/Engine/Source/Programs/UnrealBuildAccelerator/Visualizer/Private/UbaVisualizer.cpp
--- a/Engine/Source/Programs/UnrealBuildAccelerator/Visualizer/Private/UbaVisualizer.cpp
+++ b/Engine/Source/Programs/UnrealBuildAccelerator/Visualizer/Private/UbaVisualizer.cpp
@@ -479,8 +479,12 @@ namespace uba
 			m_processBrushes[1].success = CreateSolidBrush(RGB(10, 140, 10));
 			m_processBrushes[0].send = CreateSolidBrush(RGB(10, 115, 10));
 			m_processBrushes[1].send = CreateSolidBrush(RGB(10, 145, 10));
-			m_processBrushes[0].cacheFetch = CreateSolidBrush(RGB(24, 112, 110));
-			m_processBrushes[1].cacheFetch = CreateSolidBrush(RGB(31, 143, 138));
+			m_processBrushes[0].cacheFetchTest = CreateSolidBrush(RGB(24, 120, 110));
+			m_processBrushes[1].cacheFetchTest = CreateSolidBrush(RGB(31, 150, 138));
+			m_processBrushes[0].cacheFetchDownload = CreateSolidBrush(RGB(24, 112, 110));
+			m_processBrushes[1].cacheFetchDownload = CreateSolidBrush(RGB(31, 143, 138));
+			m_processBrushes[0].warning = CreateSolidBrush(RGB(100, 100, 0));
+			m_processBrushes[1].warning = CreateSolidBrush(RGB(100, 100, 0));
 
 			m_backgroundBrush = CreateSolidBrush(0x00252526);
 			m_separatorPen = CreatePen(PS_SOLID, 1, RGB(50, 50, 50));
@@ -515,9 +519,12 @@ namespace uba
 			m_processBrushes[1].success = CreateSolidBrush(RGB(20, 220, 20));
 			m_processBrushes[0].send = CreateSolidBrush(RGB(80, 210, 80));
 			m_processBrushes[1].send = CreateSolidBrush(RGB(90, 250, 90));
-
-			m_processBrushes[0].cacheFetch = CreateSolidBrush(RGB(150, 150, 200));
-			m_processBrushes[1].cacheFetch = CreateSolidBrush(RGB(170, 170, 200));
+			m_processBrushes[0].cacheFetchTest = CreateSolidBrush(RGB(150, 150, 200));
+			m_processBrushes[1].cacheFetchTest = CreateSolidBrush(RGB(170, 170, 200));
+			m_processBrushes[0].cacheFetchDownload = CreateSolidBrush(RGB(150, 150, 200));
+			m_processBrushes[1].cacheFetchDownload = CreateSolidBrush(RGB(170, 170, 200));
+			m_processBrushes[0].warning = CreateSolidBrush(RGB(170, 170, 0));
+			m_processBrushes[1].warning = CreateSolidBrush(RGB(170, 170, 0));
 
 			m_backgroundBrush = GetSysColorBrush(0);
 			m_separatorPen = CreatePen(PS_SOLID, 1, RGB(180, 180, 180));
@@ -1049,7 +1056,7 @@ namespace uba
 
 		SetActiveFont(m_defaultFont);
 
-		auto drawStatusText = [&](const StringView& text, LogEntryType type, int posX, int endX, bool moveY, bool underlined = false)
+		auto drawStatusText = [&](StringView text, LogEntryType type, int posX, int endX, bool moveY, bool underlined = false, bool centered = false)
 			{
 				RECT rect;
 				rect.left = posX;
@@ -1058,11 +1065,19 @@ namespace uba
 				rect.bottom = posY + m_activeFont.height + 2;
 				SetTextColor(hdc, type == LogEntryType_Info ? m_textColor : (type == LogEntryType_Error ? m_textErrorColor : m_textWarningColor));
 
+				u32 align = 0;
+				if (centered)
+				{
+					posX += (endX - posX)/2;
+					align = SetTextAlign(hdc, TA_CENTER);
+				}
 				if (underlined)
 					SelectObject(m_activeHdc, m_activeFont.handleUnderlined);
-				ExtTextOutW(hdc, rect.left, posY, ETO_CLIPPED, &rect, text.data, text.count, NULL);
+				ExtTextOutW(hdc, posX, posY, ETO_CLIPPED, &rect, text.data, text.count, NULL);
 				if (underlined)
 					SelectObject(m_activeHdc, m_activeFont.handle);
+				if (centered)
+					SetTextAlign(hdc, align);
 				if (moveY)
 					posY = rect.bottom;
 			};
@@ -1077,11 +1092,13 @@ namespace uba
 		{
 			drawIndentedText(AsView(L"Progress"), LogEntryType_Info, 1, false);
 
+
 			float progress = float(m_traceView.progressProcessesDone) / float(m_traceView.progressProcessesTotal);
+			u32 start = 3 + 6*m_activeFont.height;
 			u32 width = m_activeFont.height * 18;
 			RECT rect;
-			rect.left = 3 + 6*m_activeFont.height;
-			rect.right = rect.left + width;
+			rect.left = start;
+			rect.right = start + width;
 			rect.top = posY;
 			rect.bottom = posY + m_activeFont.height;
 			fillRect(rect, m_processBrushes[0].inProgress);
@@ -1098,6 +1115,22 @@ namespace uba
 			if (u32 active = m_traceView.totalProcessActiveCount)
 				str.Appendf(TC("   (%u active)"), active);
 
+			if (m_traceView.lastKillProcessTime && playTime >= m_traceView.lastKillProcessTime && playTime < m_traceView.lastKillProcessTime + MsToTime(800))
+			{
+				rect.left = start + width + 20;
+				rect.right = rect.left + m_activeFont.height * 7;
+				fillRect(rect, m_processBrushes[0].error);
+				drawStatusText(TCV("Killing (mem)"), LogEntryType_Info, rect.left, rect.right, 0, false, true);
+			}
+			else if (m_traceView.lastSpawningDelayStartTime && playTime >= m_traceView.lastSpawningDelayStartTime && playTime < m_traceView.lastSpawningDelayEndTime + MsToTime(500))
+			{
+				rect.left = start + width + 20;
+				rect.right = rect.left + m_activeFont.height * 7;
+				fillRect(rect, m_processBrushes[0].warning);
+				drawStatusText(TCV("Throttling (mem)"), LogEntryType_Info, rect.left, rect.right, 0, false, true);
+			}
+
+
 			drawIndentedText(str, LogEntryType_Info, 6, true);
 		}
 
@@ -1167,7 +1200,7 @@ namespace uba
 				float playTimeS = TimeToS(playTime);
 				u32 activeProcessCountIndex = 0;
 				u16 count = 0;
-				while (time < playTimeS)
+				while (time < playTimeS && activeProcessCountIndex < m_traceView.activeProcessCounts.size())
 				{
 					int posX = progressRect.left + int(startOffset + time*timeScale);
 					if (posX >= clientRect.right)
@@ -1290,7 +1323,7 @@ namespace uba
 
 						if (process.isRemote)
 							str.Append(L" [").Append(session.name).Append(']');
-						else if (process.type == TraceView::ProcessType_CacheFetch)
+						else if (process.type == TraceView::ProcessType_CacheFetchTest || process.type == TraceView::ProcessType_CacheFetchDownload)
 							str.Append(L" [cache]");
 						if (boxRect.left < 0)
 							str.Appendf(L"   %s", TimeToText(playTime - process.start, true).str);
@@ -2163,7 +2196,7 @@ namespace uba
 				KernelStats kernelStats;
 				CacheFetchStats cacheStats;
 
-				if (process.type == TraceView::ProcessType_CacheFetch)
+				if (process.type == TraceView::ProcessType_CacheFetchTest || process.type == TraceView::ProcessType_CacheFetchDownload)
 				{
 					if (!process.returnedReason.empty())
 						logger.Info(L"  Cache:       Miss");
@@ -2528,8 +2561,10 @@ namespace uba
 			brush = m_processBrushes[selected].returned;
 		else if (!done)
 			brush = m_processBrushes[selected].inProgress;
-		else if (process.type == TraceView::ProcessType_CacheFetch)
-			brush = m_processBrushes[selected].cacheFetch;
+		else if (process.type == TraceView::ProcessType_CacheFetchTest)
+			brush = m_processBrushes[selected].cacheFetchTest;
+		else if (process.type == TraceView::ProcessType_CacheFetchDownload)
+			brush = m_processBrushes[selected].cacheFetchDownload;
 		else if (process.exitCode == ProcessCancelExitCode)
 			brush = m_processBrushes[selected].returned;
 		else if (process.exitCode != 0)
diff --git a/Engine/Source/Programs/UnrealBuildAccelerator/Visualizer/Private/UbaVisualizer.h b/Engine/Source/Programs/UnrealBuildAccelerator/Visualizer/Private/UbaVisualizer.h
--- a/Engine/Source/Programs/UnrealBuildAccelerator/Visualizer/Private/UbaVisualizer.h
+++ b/Engine/Source/Programs/UnrealBuildAccelerator/Visualizer/Private/UbaVisualizer.h
@@ -188,7 +188,9 @@ namespace uba
 			HBRUSH returned = 0;
 			HBRUSH recv = 0;
 			HBRUSH send = 0;
-			HBRUSH cacheFetch = 0;
+			HBRUSH cacheFetchTest = 0;
+			HBRUSH cacheFetchDownload = 0;
+			HBRUSH warning = 0;
 		};
 
 		ProcessBrushes m_processBrushes[2]; // Non-selected and selected
