[UBA]
* Changed so ubascheduler.cpp includes winternl.h instead of forward declaring function because when building with unity files it got pulled in and code collided
* Removed test code that did not compile

diff --git a/Engine/Source/Programs/UnrealBuildAccelerator/Common/Private/UbaScheduler.cpp b/Engine/Source/Programs/UnrealBuildAccelerator/Common/Private/UbaScheduler.cpp
--- a/Engine/Source/Programs/UnrealBuildAccelerator/Common/Private/UbaScheduler.cpp
+++ b/Engine/Source/Programs/UnrealBuildAccelerator/Common/Private/UbaScheduler.cpp
@@ -13,6 +13,9 @@
 #include "UbaSessionServer.h"
 #include "UbaStringBuffer.h"
 
+#if PLATFORM_WINDOWS
+#include <winternl.h>
+#endif
 
 // TODO. This is how we want scheduler to work.
 // 
@@ -24,21 +27,16 @@
 // At some point cached processes and local will race which is fine.
 //
 
-#if PLATFORM_WINDOWS
-extern "C" NTSTATUS NTAPI NtQuerySystemInformation(int SystemInformationClass, PVOID SystemInformation, ULONG SystemInformationLength, PULONG ReturnLength);
-#endif
-
 namespace uba
 {
 	////////////////////////////////////////////////////////////////////////////////////////////////////
 
 	#if PLATFORM_WINDOWS
-	constexpr int SystemProcessInformation = 5;
 	constexpr u32 CreateTimeOffset = 32;
 	constexpr u32 PidOffset = 80;
 	constexpr u32 ParentPidOffset = 88;
 	constexpr u32 PagefileUsageOffset = 184;
-	//static_assert(offsetof(SYSTEM_PROCESS_INFORMATION, UniqueProcessId) == PidOffset);
+	static_assert(offsetof(SYSTEM_PROCESS_INFORMATION, UniqueProcessId) == PidOffset);
 	//static_assert(offsetof(SYSTEM_PROCESS_INFORMATION, InheritedFromUniqueProcessId) == ParentPidOffset);
 	//static_assert(offsetof(SYSTEM_PROCESS_INFORMATION, CreateTime) == CreateTimeOffset);
 	//static_assert(offsetof(SYSTEM_PROCESS_INFORMATION, PagefileUsage) == PagefileUsageOffset);
diff --git a/Engine/Source/Programs/UnrealBuildAccelerator/Test/Private/UbaTestSession.cpp b/Engine/Source/Programs/UnrealBuildAccelerator/Test/Private/UbaTestSession.cpp
--- a/Engine/Source/Programs/UnrealBuildAccelerator/Test/Private/UbaTestSession.cpp
+++ b/Engine/Source/Programs/UnrealBuildAccelerator/Test/Private/UbaTestSession.cpp
@@ -39,7 +39,6 @@ namespace uba
 		StorageImpl storage(storageInfo);
 
 		SessionServerCreateInfo sessionServerInfo(storage, server, logWriter);
-		sessionServerInfo.checkMemory = false;
 		sessionServerInfo.rootDir = rootDir.data;
 		sessionServerInfo.Apply(config);
 
@@ -128,7 +127,6 @@ namespace uba
 		auto ssg = MakeGuard([&]() { delete &storageServer; });
 
 		SessionServerCreateInfo sessionServerInfo(storageServer, server, logWriter);
-		sessionServerInfo.checkMemory = false;
 		sessionServerInfo.rootDir = rootDir.data;
 		sessionServerInfo.useUniqueId = false;
 
@@ -143,7 +141,7 @@ namespace uba
 
 		auto sg = MakeGuard([&]() { server.DisconnectClients(); });
 
-		sessionServer.SetRemoteProcessReturnedEvent([](Process& p) { p.Cancel(true); });
+		sessionServer.SetRemoteProcessReturnedEvent([](Process& p) { p.Cancel(); });
 
 		Config clientConfig;
 		clientConfig.AddTable(TC("Storage")).AddValue(TC("CheckExistsOnServer"), true);
diff --git a/Engine/Source/Programs/UnrealBuildAccelerator/Test/Private/UbaTestStdOut.cpp b/Engine/Source/Programs/UnrealBuildAccelerator/Test/Private/UbaTestStdOut.cpp
--- a/Engine/Source/Programs/UnrealBuildAccelerator/Test/Private/UbaTestStdOut.cpp
+++ b/Engine/Source/Programs/UnrealBuildAccelerator/Test/Private/UbaTestStdOut.cpp
@@ -26,7 +26,6 @@ namespace uba
 		StorageServer storageServer(storageServerInfo);
 
 		SessionServerCreateInfo sessionServerInfo(storageServer, networkServer, logWriter);
-		sessionServerInfo.checkMemory = false;
 		sessionServerInfo.rootDir = rootDir.data;
 		sessionServerInfo.traceEnabled = true;
 		//sessionServerInfo.remoteLogEnabled = true;
diff --git a/Engine/Source/Programs/UnrealBuildAccelerator/Test/Private/UbaTestStress.cpp b/Engine/Source/Programs/UnrealBuildAccelerator/Test/Private/UbaTestStress.cpp
--- a/Engine/Source/Programs/UnrealBuildAccelerator/Test/Private/UbaTestStress.cpp
+++ b/Engine/Source/Programs/UnrealBuildAccelerator/Test/Private/UbaTestStress.cpp
@@ -29,7 +29,6 @@ namespace uba
 		StorageServer storageServer(storageServerInfo);
 
 		SessionServerCreateInfo sessionServerInfo(storageServer, networkServer, logWriter);
-		sessionServerInfo.checkMemory = false;
 		sessionServerInfo.rootDir = rootDir.data;
 		sessionServerInfo.traceEnabled = true;
 		//sessionServerInfo.remoteLogEnabled = true;
