[UBA]
* Fixed bug in posix rename path when vfs is enabled
* Added detour of __lxstat to fix issue with temp file not being deleted

diff --git a/Engine/Source/Programs/UnrealBuildAccelerator/Detours/Private/PosixOS/UbaDetoursFunctionsPosix.cpp b/Engine/Source/Programs/UnrealBuildAccelerator/Detours/Private/PosixOS/UbaDetoursFunctionsPosix.cpp
--- a/Engine/Source/Programs/UnrealBuildAccelerator/Detours/Private/PosixOS/UbaDetoursFunctionsPosix.cpp
+++ b/Engine/Source/Programs/UnrealBuildAccelerator/Detours/Private/PosixOS/UbaDetoursFunctionsPosix.cpp
@@ -146,6 +146,7 @@ using namespace uba;
 	DETOURED_FUNCTION(fcntl) \
 	DETOURED_FUNCTION(__xstat) \
 	DETOURED_FUNCTION(__xstat64) \
+	DETOURED_FUNCTION(__lxstat) \
 	DETOURED_FUNCTION(__fxstat) \
 	DETOURED_FUNCTION(__fxstat64) \
 	DETOURED_FUNCTION(__fxstatat) \
@@ -814,7 +815,7 @@ int Shared_access(const char* funcName, const char* pathname, int mode)
 
 	int res = attr.lastError == 0 ? 0 : -1;
 	
-	DEBUG_LOG_DETOURED(funcName, "%s %i (%s) -> %i %s", pathname, mode, realName, res, StrError(res, errno));
+	DEBUG_LOG_DETOURED(funcName, "%s %i (%s) -> %i %s", pathname, mode, realName, res, StrError(res, attr.lastError));
 
 	#if UBA_DEBUG_VALIDATE
 	if (!g_runningRemote)
@@ -1662,6 +1663,12 @@ UBA_EXPORT int UBA_WRAPPER(__xstat)(int ver, const char* file, struct stat* attr
 	return Shared_stat("__xstat", file, attr, [ver](const char* file, struct stat* attr) { return TRUE_WRAPPER(__xstat)(ver, file, attr); });
 }
 
+UBA_EXPORT int UBA_WRAPPER(__lxstat)(int ver, const char* file, struct stat* attr)
+{
+	UBA_INIT_DETOUR(__lxstat, ver, file, attr);
+	return Shared_stat("__lxstat", file, attr, [ver](const char* file, struct stat* attr) { return TRUE_WRAPPER(__lxstat)(ver, file, attr); });
+}
+
 UBA_EXPORT int UBA_WRAPPER(__fxstat)(int ver, int fd, struct stat* attr)
 {
 	UBA_INIT_DETOUR(__fxstat, ver, fd, attr);
@@ -1713,8 +1720,8 @@ UBA_EXPORT int UBA_WRAPPER(rename)(const char* oldpath, const char* newpath)
 				bool isTempFile = fixedOldPath.StartsWith(g_systemTemp.data);
 				UBA_ASSERTF(wasTempFile == isTempFile, "File changing from temp to not or vice versa not implemented");
 
-				int res = TRUE_WRAPPER(rename)(oldpath, newpath);
-				DEBUG_LOG_DETOURED("rename", "IS_OPEN (%i) (from %s to %s) -> %i (%s)", kv.first, fixedOldPath.data, fixedNewPath.data, 0, StrError(0, errno));
+				int res = TRUE_WRAPPER(rename)(fixedOldPath.data, fixedNewPath.data);
+				DEBUG_LOG_DETOURED("rename", "IS_OPEN (%i) (from %s to %s) -> %i (%s)", kv.first, fixedOldPath.data, fixedNewPath.data, 0, StrError(res, errno));
 				return res;
 			}
 		}
@@ -1928,7 +1935,7 @@ int Shared_DeleteFile(const char* funcName, const char* pathname)
 	g_mappedFileTable.SetDeleted(fileNameKey, fixedName.data, true);
 
 	int res = result != 0 ? 0 : -1;
-	DEBUG_LOG_DETOURED(funcName, "(%s) -> %i (%s)", pathname, res, StrError(res, errorCode));
+	DEBUG_LOG_DETOURED(funcName, "%s (%s) -> %i (%s)", pathname, fixedName.data, res, StrError(res, errorCode));
 	errno = errorCode;
 	return res;
 }
