[UBA]
* Fixed crash in return process code when doing process reuse

diff --git a/Engine/Source/Programs/UnrealBuildAccelerator/Common/Private/UbaScheduler.cpp b/Engine/Source/Programs/UnrealBuildAccelerator/Common/Private/UbaScheduler.cpp
--- a/Engine/Source/Programs/UnrealBuildAccelerator/Common/Private/UbaScheduler.cpp
+++ b/Engine/Source/Programs/UnrealBuildAccelerator/Common/Private/UbaScheduler.cpp
@@ -482,10 +482,12 @@ namespace uba
 		auto& si = process.GetStartInfo();
 		auto& ei = *(ExitProcessInfo*)si.userData;
 		u32 processIndex = ei.processIndex;
-		float processWeight = ei.startInfo->weight;
+		float processWeight = 0;
 
 		if (isLocal)
 		{
+			processWeight = ei.startInfo->weight;
+
 			SCOPED_CRITICAL_SECTION(m_wasReturnedLock, lock);
 			if (ei.hasCheckedWasReturned)
 				return false;
@@ -499,7 +501,10 @@ namespace uba
 		}
 		
 		if (processIndex == ~0u)
-			return m_session.GetLogger().Error(TC("Process has no process index. This should not happen"));
+		{
+			UBA_ASSERT(!isLocal);
+			return false;
+		}
 
 		if (isLocal)
 			if (!process.WaitForExit(10*1000))
