[UBA]
* Fixed bug in scheduler for processes returned from remotes not being properly rescheduled

diff --git a/Engine/Source/Programs/UnrealBuildAccelerator/Common/Private/UbaScheduler.cpp b/Engine/Source/Programs/UnrealBuildAccelerator/Common/Private/UbaScheduler.cpp
--- a/Engine/Source/Programs/UnrealBuildAccelerator/Common/Private/UbaScheduler.cpp
+++ b/Engine/Source/Programs/UnrealBuildAccelerator/Common/Private/UbaScheduler.cpp
@@ -497,6 +497,7 @@ namespace uba
 		}
 		else
 		{
+			ei.wasReturned = true;
 			process.Cancel(); // Cancel will call ProcessExited
 		}
 		
