[UBA]
* Fixed crash that could happen when pressing ctrl-c where message could access connection that was deleted

diff --git a/Engine/Source/Programs/UnrealBuildAccelerator/Common/Private/UbaNetworkClient.cpp b/Engine/Source/Programs/UnrealBuildAccelerator/Common/Private/UbaNetworkClient.cpp
--- a/Engine/Source/Programs/UnrealBuildAccelerator/Common/Private/UbaNetworkClient.cpp
+++ b/Engine/Source/Programs/UnrealBuildAccelerator/Common/Private/UbaNetworkClient.cpp
@@ -715,6 +715,7 @@ namespace uba
 		message.m_response = response;
 		message.m_responseCapacity = responseCapacity;
 		message.m_connection = &connection;
+		message.m_cryptoIV = connection.cryptoIV;
 
 		BinaryWriter& writer = *message.m_sendWriter;
 
@@ -923,10 +924,8 @@ namespace uba
 		if (m_client->m_cryptoKey)
 		{
 			UBA_ASSERT(!response.GetPosition());
-			UBA_ASSERT(m_connection);
 			TimerScope ts(m_client->m_decryptTimer);
-			Guid iv = m_connection->cryptoIV;
-			if (!Crypto::Decrypt(m_client->m_logger, m_client->m_cryptoKey, (u8*)m_response, m_responseSize, iv))
+			if (!Crypto::Decrypt(m_client->m_logger, m_client->m_cryptoKey, (u8*)m_response, m_responseSize, m_cryptoIV))
 			{
 				m_error = 10;
 				return false;
diff --git a/Engine/Source/Programs/UnrealBuildAccelerator/Common/Public/UbaNetworkMessage.h b/Engine/Source/Programs/UnrealBuildAccelerator/Common/Public/UbaNetworkMessage.h
--- a/Engine/Source/Programs/UnrealBuildAccelerator/Common/Public/UbaNetworkMessage.h
+++ b/Engine/Source/Programs/UnrealBuildAccelerator/Common/Public/UbaNetworkMessage.h
@@ -44,6 +44,7 @@ namespace uba
 		MessagePriority m_priority;
 		Atomic<u32> m_error;
 		NetworkClient::Connection* m_connection = nullptr;
+		Guid m_cryptoIV;
 		friend NetworkClient;
 		NetworkMessage(const NetworkMessage&) = delete;
 		NetworkMessage(NetworkMessage&&) = delete;
