[UBA]
* Added UbaAgent -mincon=<x> to be able to force more than 1 tcp connection even though ping is very low

diff --git a/Engine/Source/Programs/UnrealBuildAccelerator/Agent/Private/UbaAgent.cpp b/Engine/Source/Programs/UnrealBuildAccelerator/Agent/Private/UbaAgent.cpp
--- a/Engine/Source/Programs/UnrealBuildAccelerator/Agent/Private/UbaAgent.cpp
+++ b/Engine/Source/Programs/UnrealBuildAccelerator/Agent/Private/UbaAgent.cpp
@@ -80,6 +80,7 @@ namespace uba
 		logger.Info(TC("  -proxyaddr=<addr>       Which address that agent will use if being assigned to be proxy for other agents. If not set it will automatically fetch"));
 		logger.Info(TC("  -maxcpu=<number>        Max number of processes that can be started. Defaults to \"%u\" on this machine"), DefaultProcessorCount);
 		logger.Info(TC("  -mulcpu=<number>        This value multiplies with number of cpu to figure out max cpu. Defaults to 1.0"));
+		logger.Info(TC("  -mincon=<number>        Min number of connections for agent. Defaults to 1 (amount up to max will depend on ping)"));
 		logger.Info(TC("  -maxcon=<number>        Max number of connections that can be started by agent. Defaults to \"%u\" (amount up to max will depend on ping)"), DefaultMaxConnectionCount);
 		logger.Info(TC("  -maxworkers=<number>    Max number of workers is started by agent. Defaults to \"%u\""), DefaultProcessorCount);
 		logger.Info(TC("  -capacity=<gigaby>      Capacity of local store. Defaults to %u gigabytes"), DefaultCapacityGb);
@@ -345,6 +346,7 @@ namespace uba
 		u32 maxProcessCount = DefaultProcessorCount;
 		u32 maxWorkerCount = DefaultProcessorCount;
 		float mulProcessValue = 1.0f;
+		u32 minConnectionCount = 1;
 		u32 maxConnectionCount = DefaultMaxConnectionCount;
 		u32 storageCapacityGb = DefaultCapacityGb;
 		StringBuffer<256> host;
@@ -451,6 +453,13 @@ namespace uba
 				if (!parseValue(mulProcessValue))
 					return false;
 			}
+			else if (name.Equals(TCV("-mincon")))
+			{
+				if (!parseValue(minConnectionCount))
+					return false;
+				if (minConnectionCount == 0)
+					return PrintHelp(TC("Invalid value for -mincon"));
+			}
 			else if (name.Equals(TCV("-maxcon")) || name.Equals(TCV("-maxtcp")))
 			{
 				if (!parseValue(maxConnectionCount))
@@ -750,6 +759,8 @@ namespace uba
 		if (memKillLoadPercent < memWaitLoadPercent)
 			memKillLoadPercent = memWaitLoadPercent;
 
+		if (maxConnectionCount < minConnectionCount)
+			maxConnectionCount = minConnectionCount;
 
 		FilteredLogWriter logWriter(g_consoleLogWriter, verbose ? LogEntryType_Debug : LogEntryType_Detail);
 		LoggerWithWriter logger(logWriter, TC(""));
@@ -1191,7 +1202,7 @@ namespace uba
 				continue;
 
 			Event wakeupSessionWait(false);
-			Atomic<u32> targetConnectionCount = 1;
+			Atomic<u32> targetConnectionCount = minConnectionCount;
 			StorageClient* storageClient = nullptr;
 
 
@@ -1476,7 +1487,7 @@ namespace uba
 				// For every 15ms latency on "best ping") we increase targetConnectionCount up to maxConnectionCount
 				if (!storageClient->IsUsingProxy())
 					if (u64 bestPing = sessionClient->GetBestPing())
-						targetConnectionCount = Min(u32(TimeToMs(bestPing) / 15), maxConnectionCount);
+						targetConnectionCount = Max(minConnectionCount, Min(u32(TimeToMs(bestPing) / 15), maxConnectionCount));
 
 				if (!isTerminating)
 				{
