[UBA]
* Improved mem tracking logic in scheduler to also take page file write activity into account. If excessive writing happens, processes gets killed and startwait threshold lowered... threshold will then increase back to normal over time.

This makes compiling with clang 19+ much better without running into crazy page file swapping or even oom/bluescreen

diff --git a/Engine/Source/Programs/UnrealBuildAccelerator/Common/Private/UbaScheduler.cpp b/Engine/Source/Programs/UnrealBuildAccelerator/Common/Private/UbaScheduler.cpp
--- a/Engine/Source/Programs/UnrealBuildAccelerator/Common/Private/UbaScheduler.cpp
+++ b/Engine/Source/Programs/UnrealBuildAccelerator/Common/Private/UbaScheduler.cpp
@@ -189,15 +189,10 @@ namespace uba
 
 		if (info.trackMemory && IsWindows && !IsRunningWine()) // Only implemented for normal windows for now
 		{
-			u64 memAvail;
-			u64 memTotal;
-			if (GetMemoryInfo(m_session.GetLogger(), memAvail, memTotal))
-			{
-				m_memTrack = true;
-				m_memTracingLevel = info.memTracingLevel;
-				m_memWaitThreshold = u64(double(memTotal) * double(info.memStartWaitPercent) / 100.0);
-				m_memKillThreshold = u64(double(memTotal) * double(info.memStartKillPercent) / 100.0);
-			}
+			m_memTrack = true;
+			m_memTracingLevel = info.memTracingLevel;
+			m_memStartWaitPercent = info.memStartWaitPercent;
+			m_memStartKillPercent = info.memStartKillPercent;
 		}
 
 		m_session.SetOuterScheduler(this);
@@ -221,9 +216,13 @@ namespace uba
 		if (m_memTrack)
 		{
 			LoadMemTrackTable();
-			m_session.SetNativeProcessCreatedFunc([this](Process& process) { NativeProcessCreated(process); });
+			m_session.SetNativeProcessCreatedFunc([this](ProcessImpl& process) { NativeProcessCreated(process); });
 			m_memThreadEvent.Create(true);
 			m_memThread.Start([this]() { ThreadMemoryCheckLoop();  return 0; }, TC("UbaMemTrackLoop"));
+
+			// Make sure we have an estimate before starting to spawn processes
+			while (!m_memThread.Wait(1) && !m_memEstimatedUsage)
+				;
 		}
 
 		m_loop = true;
@@ -1145,7 +1144,7 @@ namespace uba
 			m_session.GetLogger().Info(TC("Killed process %s (%s) due to memory pressure"), ph.GetStartInfo().GetDescription(), TimeToText(GetTime() - startTime).str);
 	}
 
-	void Scheduler::NativeProcessCreated(Process& process)
+	void Scheduler::NativeProcessCreated(ProcessImpl& process)
 	{
 		if (!m_memTrack)
 			return;
@@ -1154,7 +1153,7 @@ namespace uba
 		auto ei = (ExitProcessInfo*)impl.m_startInfo.userData;
 		StringKey key(impl.m_nativeProcessId, impl.m_nativeProcessCreationTime);
 		SCOPED_FUTEX(m_processesLock, l);
-		m_processes[key] = ei;
+		m_processes.try_emplace(key, ProcessRecord{ process, ei });
 	}
 
 	void Scheduler::LocalProcessStart(ProcessStartInfo2& info)
@@ -1337,9 +1336,38 @@ namespace uba
 
 		//u64 randomKillTime = GetTime() + MsToTime(5000);
 
-		u32 timeoutMs = 0;
+		Vector<u8> queryBuffer(1024 * 1024);
 
-		Vector<u8> processInfoBuffer(1024 * 1024);
+		#if PLATFORM_WINDOWS
+		struct UBA_SYSTEM_PERFORMANCE_INFORMATION
+		{
+			LARGE_INTEGER IdleTime;
+			LARGE_INTEGER ReadTransferCount;
+			LARGE_INTEGER WriteTransferCount;
+			LARGE_INTEGER OtherTransferCount;
+			ULONG ReadOperationCount;
+			ULONG WriteOperationCount;
+			ULONG OtherOperationCount;
+			ULONG AvailablePages;
+			ULONG TotalCommittedPages;
+			ULONG TotalCommitLimit;
+			ULONG PeakCommitment;
+			ULONG PageFaults;
+			ULONG WriteCopyFaults;
+			ULONG TransitionFaults;
+			ULONG CacheTransitionFaults;
+			ULONG DemandZeroFaults;
+			ULONG PagesRead;
+			ULONG PageReadIos;
+			ULONG CacheReads;
+			ULONG CacheIos;
+			ULONG PagefilePagesWritten;
+			// ...
+		};
+		static_assert(sizeof(UBA_SYSTEM_PERFORMANCE_INFORMATION) <= sizeof(SYSTEM_PERFORMANCE_INFORMATION));
+		UBA_SYSTEM_PERFORMANCE_INFORMATION prevPerfInfo = {};
+		bool firstPerfInfo = true;
+		#endif
 
 		constexpr u32 StartRow = 20;
 
@@ -1354,27 +1382,42 @@ namespace uba
 			trace.StatusUpdate(StartRow+2, 1, TC("MemBaseline"), LogEntryType_Info, nullptr);
 			trace.StatusUpdate(StartRow+3, 1, TC("MemUntrack"), LogEntryType_Info, nullptr);
 			trace.StatusUpdate(StartRow+4, 1, TC("MemStartWait"), LogEntryType_Info, nullptr);
-			trace.StatusUpdate(StartRow+4, 6, BytesToText(m_memWaitThreshold).str, LogEntryType_Info, nullptr);
 			trace.StatusUpdate(StartRow+5, 1, TC("MemStartKill"), LogEntryType_Info, nullptr);
-			trace.StatusUpdate(StartRow+5, 6, BytesToText(m_memKillThreshold).str, LogEntryType_Info, nullptr);
+			trace.StatusUpdate(StartRow+6, 1, TC("PagefileRead"), LogEntryType_Info, nullptr);
+			trace.StatusUpdate(StartRow+7, 1, TC("PagefileWrite"), LogEntryType_Info, nullptr);
 		}
 
+		u64 memAvail;
+		u64 memTotal;
+		u64 maxPageFile;
+		if (!GetMemoryInfo(logger, memAvail, memTotal, &maxPageFile))
+		{
+			m_memTrack = false;
+			logger.Warning(TC("GetMemoryInfo failed. Memory tracking disabled."));
+			m_updateThreadWakeup.Set();
+			return;
+		}
+		u64 physMem = memTotal - maxPageFile;
+
 		UnorderedMap<u32, u32> memoryGroupToIndex;
 
+		u64 memKillThreshold = 0;
+		
+
+		u8 memStartWaitPercent = m_memStartWaitPercent;
+		u8 memStartWaitPercentMin = u8(double(physMem)/double(memTotal) * 100.0f);
+		u64 lastPageFileWriteTimeMs = 0;
+
+		u32 timeoutMs = 0;
+
 		while (!m_memThreadEvent.IsSet(timeoutMs))
 		{
 			timeoutMs = 1000;
 
-			u64 memAvail;
-			u64 memTotal;
-			u64 maxPageFile;
-			if (!GetMemoryInfo(logger, memAvail, memTotal, &maxPageFile))
-			{
-				m_memTrack = false;
-				logger.Warning(TC("GetMemoryInfo failed. Memory tracking disabled."));
-				m_updateThreadWakeup.Set();
-				break;
-			}
+			GetMemoryInfo(logger, memAvail, memTotal, &maxPageFile);
+
+			m_memWaitThreshold = u64(double(memTotal) * double(memStartWaitPercent) / 100.0);
+			memKillThreshold = u64(double(memTotal) * double(m_memStartKillPercent) / 100.0);
 
 			u64 untrackedMemory = 0;
 			u64 processMemTotalPageCount = 0;
@@ -1383,12 +1426,12 @@ namespace uba
 			ULONG returnedSize;
 			while (true)
 			{
-				NTSTATUS res = NtQuerySystemInformation(SystemProcessInformation, processInfoBuffer.data(), (u32)processInfoBuffer.size(), &returnedSize);
+				NTSTATUS res = NtQuerySystemInformation(SystemProcessInformation, queryBuffer.data(), (u32)queryBuffer.size(), &returnedSize);
 				if (res == STATUS_SUCCESS)
 					break;
 				if (res == STATUS_INFO_LENGTH_MISMATCH)
 				{
-					processInfoBuffer.resize(returnedSize + (10*PagefileUsageOffset));
+					queryBuffer.resize(returnedSize + (10*PagefileUsageOffset));
 					continue;
 				}
 
@@ -1398,7 +1441,7 @@ namespace uba
 				break;
 			}
 
-			u8* it = processInfoBuffer.data();
+			u8* it = queryBuffer.data();
 			ULONG nextEntryOffset = *(ULONG*)it;
 			if (nextEntryOffset < PagefileUsageOffset + 8)
 			{
@@ -1417,12 +1460,13 @@ namespace uba
 					auto procIt = m_processes.find(StringKey(pid, createTime));
 					if (procIt != m_processes.end())
 					{
+						ProcessRecord rec = procIt->second;
 						u64 procMem = *(u64*)(it + PagefileUsageOffset);
-						if (auto ei = procIt->second)
-							if (ei->startInfo->memoryGroupId == 0)
+						if (rec.ep)
+							if (rec.ep->startInfo->memoryGroupId == 0)
 								untrackedMemory += procMem;
+						AtomicMax(rec.process.m_processStats.peakMemory, procMem);
 						processMemTotalPageCount += procMem;
-						//AtomicMax(((ProcessImpl*)procIt->second.process)->m_processStats.peakMemory, spi2->PagefileUsage);
 					}
 
 					nextEntryOffset = *(ULONG*)it;
@@ -1431,8 +1475,8 @@ namespace uba
 			}
 			#endif
 
-			u64 memoryUsage = memTotal - memAvail;
-			u64 baselineMemory = memoryUsage - processMemTotalPageCount;
+			u64 memUsed = memTotal - memAvail;
+			u64 memBaseLine = memUsed - processMemTotalPageCount;
 			u64 estimatedProcessMemoryUsage = 0;
 			u64 estimatedMemoryUsage = 0;
 
@@ -1462,31 +1506,68 @@ namespace uba
 						trace.StatusUpdate(index, colIndex + 5, b.data, LogEntryType_Info, nullptr);
 					}
 				}
-				estimatedMemoryUsage = baselineMemory + untrackedMemory + estimatedProcessMemoryUsage;
-				m_memEstimatedUsage = estimatedMemoryUsage < memoryUsage ? memoryUsage : estimatedMemoryUsage;
+				estimatedMemoryUsage = memBaseLine + untrackedMemory + estimatedProcessMemoryUsage;
+				m_memEstimatedUsage = estimatedMemoryUsage < memUsed ? memUsed : estimatedMemoryUsage;
 			}
 
 			if (m_memTracingLevel >= 1)
 			{
-				trace.StatusUpdate(StartRow, 6, BytesToText(memoryUsage).str, LogEntryType_Info, nullptr);
+				trace.StatusUpdate(StartRow, 6, BytesToText(memUsed).str, LogEntryType_Info, nullptr);
 				trace.StatusUpdate(StartRow+1, 6, BytesToText(estimatedMemoryUsage).str, LogEntryType_Info, nullptr);
 				trace.StatusUpdate(StartRow, 15, BytesToText(processMemTotalPageCount).str, LogEntryType_Info, nullptr);
 				trace.StatusUpdate(StartRow+1, 15, BytesToText(estimatedProcessMemoryUsage).str, LogEntryType_Info, nullptr);
-				trace.StatusUpdate(StartRow+2, 6, BytesToText(baselineMemory).str, LogEntryType_Info, nullptr);
+				trace.StatusUpdate(StartRow+2, 6, BytesToText(memBaseLine).str, LogEntryType_Info, nullptr);
 				trace.StatusUpdate(StartRow+3, 6, BytesToText(untrackedMemory).str, LogEntryType_Info, nullptr);
+				trace.StatusUpdate(StartRow+4, 6, BytesToText(m_memWaitThreshold).str, LogEntryType_Info, nullptr);
+				trace.StatusUpdate(StartRow+5, 6, BytesToText(memKillThreshold).str, LogEntryType_Info, nullptr);
 			}
 
-			if (m_memKillThreshold)
+			if (memKillThreshold)
 			{
-				while (memoryUsage > m_memKillThreshold)
+				while (memUsed > memKillThreshold)
 				{
 					KillNewestLocalProcess();
 					if (!GetMemoryInfo(logger, memAvail, memTotal, &maxPageFile))
 						break;
-					memoryUsage = memTotal - memAvail;
+					memUsed = memTotal - memAvail;
 				}
 			}
 
+			#if PLATFORM_WINDOWS
+			auto& currPerfInfo = *(UBA_SYSTEM_PERFORMANCE_INFORMATION*)queryBuffer.data();
+			if (NtQuerySystemInformation(SystemPerformanceInformation, &currPerfInfo, sizeof(SYSTEM_PERFORMANCE_INFORMATION), NULL) == STATUS_SUCCESS && !firstPerfInfo)
+			{
+				ULONG reads = currPerfInfo.PagesRead - prevPerfInfo.PagesRead;
+				ULONG writes = currPerfInfo.PagefilePagesWritten - prevPerfInfo.PagefilePagesWritten;
+				trace.StatusUpdate(StartRow+6, 6, CountToText(reads).str, LogEntryType_Info, nullptr);
+				trace.StatusUpdate(StartRow+7, 6, CountToText(writes).str, LogEntryType_Info, nullptr);
+
+				if (memUsed > physMem && writes > 1000)
+				{
+					memStartWaitPercent = u8(double(memUsed)/double(memTotal) * 100.0f);
+					KillNewestLocalProcess();
+					lastPageFileWriteTimeMs = TimeToMs(GetTime());
+				}
+				else if (writes > 0)
+				{
+					if (memStartWaitPercent > memStartWaitPercentMin)
+						--memStartWaitPercent;
+					lastPageFileWriteTimeMs = TimeToMs(GetTime());
+				}
+				else if (TimeToMs(GetTime()) - lastPageFileWriteTimeMs > 5*1000)
+				{
+					if (memStartWaitPercent < m_memStartWaitPercent)
+					{
+						++memStartWaitPercent;
+						if (memStartWaitPercent == m_memStartWaitPercent)
+							m_updateThreadWakeup.Set(); // Just in case.. this should never happen in practice
+					}
+				}
+			}
+			firstPerfInfo = false;
+			prevPerfInfo = currPerfInfo;
+			#endif
+
 			//if (randomKillTime < GetTime())
 			//{
 			//	randomKillTime = GetTime() + MsToTime(500);
diff --git a/Engine/Source/Programs/UnrealBuildAccelerator/Common/Private/UbaSessionServer.cpp b/Engine/Source/Programs/UnrealBuildAccelerator/Common/Private/UbaSessionServer.cpp
--- a/Engine/Source/Programs/UnrealBuildAccelerator/Common/Private/UbaSessionServer.cpp
+++ b/Engine/Source/Programs/UnrealBuildAccelerator/Common/Private/UbaSessionServer.cpp
@@ -525,7 +525,7 @@ namespace uba
 		m_remoteProcessReturnedEvent = remoteProcessReturnedEvent;
 	}
 
-	void SessionServer::SetNativeProcessCreatedFunc(const Function<void(Process&)>& func)
+	void SessionServer::SetNativeProcessCreatedFunc(const Function<void(ProcessImpl&)>& func)
 	{
 		m_nativeProcessCreatedFunc = func;
 	}
diff --git a/Engine/Source/Programs/UnrealBuildAccelerator/Common/Public/UbaScheduler.h b/Engine/Source/Programs/UnrealBuildAccelerator/Common/Public/UbaScheduler.h
--- a/Engine/Source/Programs/UnrealBuildAccelerator/Common/Public/UbaScheduler.h
+++ b/Engine/Source/Programs/UnrealBuildAccelerator/Common/Public/UbaScheduler.h
@@ -12,6 +12,7 @@ namespace uba
 	class Config;
 	class ConfigTable;
 	class Process;
+	class ProcessImpl;
 	class RootPaths;
 	class SessionServer;
 	class WorkManagerImpl;
@@ -144,7 +145,7 @@ namespace uba
 		void FinishProcess(const ProcessHandle& handle);
 
 		void KillNewestLocalProcess();
-		void NativeProcessCreated(Process& process);
+		void NativeProcessCreated(ProcessImpl& process);
 		void LocalProcessStart(ProcessStartInfo2& info);
 		bool LocalProcessOkToSpawn();
 		void LocalProcessExit(ProcessStartInfo2& info, const ProcessHandle& handle);
@@ -215,8 +216,9 @@ namespace uba
 		Thread m_memThread;
 		bool m_memDelaySpawning = false;
 		Atomic<u64> m_memEstimatedUsage;
-		u64 m_memWaitThreshold = 0;
-		u64 m_memKillThreshold = 0;
+		Atomic<u64> m_memWaitThreshold;
+		u8 m_memStartWaitPercent = 0;
+		u8 m_memStartKillPercent = 0;
 		u8 m_memTracingLevel = 0;
 
 		struct MemGroupStats;
@@ -226,7 +228,8 @@ namespace uba
 		UnorderedMap<u32, MemGroupEntry> m_memGroupDb;
 
 		Futex m_processesLock;
-		UnorderedMap<StringKey, ExitProcessInfo*> m_processes;
+		struct ProcessRecord { ProcessImpl& process; ExitProcessInfo* ep; };
+		UnorderedMap<StringKey, ProcessRecord> m_processes;
 
 		Scheduler(const Scheduler&) = delete;
 		void operator=(const Scheduler&) = delete;
diff --git a/Engine/Source/Programs/UnrealBuildAccelerator/Common/Public/UbaSessionServer.h b/Engine/Source/Programs/UnrealBuildAccelerator/Common/Public/UbaSessionServer.h
--- a/Engine/Source/Programs/UnrealBuildAccelerator/Common/Public/UbaSessionServer.h
+++ b/Engine/Source/Programs/UnrealBuildAccelerator/Common/Public/UbaSessionServer.h
@@ -49,7 +49,7 @@ namespace uba
 		void SetRemoteProcessReturnedEvent(const Function<void(Process&)>& remoteProcessReturnedEvent);
 
 		// Callback that is called when process is returned. This could be because a client unexpectedly disconnected or is running out of memory
-		void SetNativeProcessCreatedFunc(const Function<void(Process&)>& func);
+		void SetNativeProcessCreatedFunc(const Function<void(ProcessImpl&)>& func);
 
 		// Wait for all queued up/active tasks to finish
 		void WaitOnAllTasks();
@@ -140,7 +140,7 @@ namespace uba
 		ReaderWriterLock m_remoteProcessReturnedEventLock;
 		Function<void(Process&)> m_remoteProcessReturnedEvent;
 
-		Function<void(Process&)> m_nativeProcessCreatedFunc;
+		Function<void(ProcessImpl&)> m_nativeProcessCreatedFunc;
 
 		CriticalSection m_remoteProcessAndSessionLock; // Can be re-entrant.
 		List<ProcessHandle> m_queuedRemoteProcesses;
