[UBA]
* Changed some trace api to use StringView instead of const tchar*
* Changed some trace api to use StringView instead of const StringView&

diff --git a/Engine/Source/Programs/UnrealBuildAccelerator/Common/Private/UbaLogger.cpp b/Engine/Source/Programs/UnrealBuildAccelerator/Common/Private/UbaLogger.cpp
--- a/Engine/Source/Programs/UnrealBuildAccelerator/Common/Private/UbaLogger.cpp
+++ b/Engine/Source/Programs/UnrealBuildAccelerator/Common/Private/UbaLogger.cpp
@@ -415,6 +415,8 @@ namespace uba
 			Format1dp(str, double(bytes)/double(v * v * v * v), TC("tb"));
 	}
 
+	BytesToText::operator StringView () const { return ToView(str); }
+
 	CountToText::CountToText(u64 count)
 	{
 		if (count < 1000)
@@ -429,6 +431,8 @@ namespace uba
 			Format1dp(str, double(count) / (1000ull * 1000 * 1000 * 1000), TC("t"));
 	}
 
+	CountToText::operator StringView () const { return ToView(str); }
+
 	TimeToText::TimeToText(u64 time, bool allowMinutes) : TimeToText(time, allowMinutes, GetFrequency()) {}
 	TimeToText::TimeToText(u64 time, bool allowMinutes, u64 frequency)
 	{
@@ -457,6 +461,8 @@ namespace uba
 		}
 	}
 
+	TimeToText::operator StringView () const { return ToView(str); }
+
 #if UBA_DEBUG_LOGGER
 	thread_local u32 t_debugLogScopeCount = 0;
 
diff --git a/Engine/Source/Programs/UnrealBuildAccelerator/Common/Private/UbaScheduler.cpp b/Engine/Source/Programs/UnrealBuildAccelerator/Common/Private/UbaScheduler.cpp
--- a/Engine/Source/Programs/UnrealBuildAccelerator/Common/Private/UbaScheduler.cpp
+++ b/Engine/Source/Programs/UnrealBuildAccelerator/Common/Private/UbaScheduler.cpp
@@ -1105,7 +1105,7 @@ namespace uba
 	{
 		StringBuffer<> str;
 		str.Appendf(TC("Hits %u Misses %u"), m_cacheHitCount.load(), m_cacheMissCount.load());
-		m_session.GetTrace().StatusUpdate(1, 6, str.data, LogEntryType_Info, nullptr);
+		m_session.GetTrace().StatusUpdate(1, 6, str, LogEntryType_Info);
 	}
 
 	void Scheduler::UpdateQueueCounter(int offset)
@@ -1375,16 +1375,16 @@ namespace uba
 
 		if (m_memTracingLevel >= 1)
 		{
-			trace.StatusUpdate(StartRow, 1, TC("MemTot"), LogEntryType_Info, nullptr);
-			trace.StatusUpdate(StartRow+1, 1, TC("MemTot Est"), LogEntryType_Info, nullptr);
-			trace.StatusUpdate(StartRow, 10, TC("MemProc"), LogEntryType_Info, nullptr);
-			trace.StatusUpdate(StartRow+1, 10, TC("MemProc Est"), LogEntryType_Info, nullptr);
-			trace.StatusUpdate(StartRow+2, 1, TC("MemBaseline"), LogEntryType_Info, nullptr);
-			trace.StatusUpdate(StartRow+3, 1, TC("MemUntrack"), LogEntryType_Info, nullptr);
-			trace.StatusUpdate(StartRow+4, 1, TC("MemStartWait"), LogEntryType_Info, nullptr);
-			trace.StatusUpdate(StartRow+5, 1, TC("MemStartKill"), LogEntryType_Info, nullptr);
-			trace.StatusUpdate(StartRow+6, 1, TC("PagefileRead"), LogEntryType_Info, nullptr);
-			trace.StatusUpdate(StartRow+7, 1, TC("PagefileWrite"), LogEntryType_Info, nullptr);
+			trace.StatusUpdate(StartRow, 1, TCV("MemTot"), LogEntryType_Info);
+			trace.StatusUpdate(StartRow+1, 1, TCV("MemTot Est"), LogEntryType_Info);
+			trace.StatusUpdate(StartRow, 10, TCV("MemProc"), LogEntryType_Info);
+			trace.StatusUpdate(StartRow+1, 10, TCV("MemProc Est"), LogEntryType_Info);
+			trace.StatusUpdate(StartRow+2, 1, TCV("MemBaseline"), LogEntryType_Info);
+			trace.StatusUpdate(StartRow+3, 1, TCV("MemUntrack"), LogEntryType_Info);
+			trace.StatusUpdate(StartRow+4, 1, TCV("MemStartWait"), LogEntryType_Info);
+			trace.StatusUpdate(StartRow+5, 1, TCV("MemStartKill"), LogEntryType_Info);
+			trace.StatusUpdate(StartRow+6, 1, TCV("PagefileRead"), LogEntryType_Info);
+			trace.StatusUpdate(StartRow+7, 1, TCV("PagefileWrite"), LogEntryType_Info);
 		}
 
 		u64 memAvail;
@@ -1498,12 +1498,12 @@ namespace uba
 							index = StartRow + u32(memoryGroupToIndex.size() - 1);
 							StringBuffer<> b;
 							b.AppendValue(kv.first);
-							trace.StatusUpdate(index, colIndex, b.data, LogEntryType_Info, nullptr);
+							trace.StatusUpdate(index, colIndex, b, LogEntryType_Info);
 						}
 
 						StringBuffer<> b;
 						b.Append(BytesToText(stats.average).str);
-						trace.StatusUpdate(index, colIndex + 5, b.data, LogEntryType_Info, nullptr);
+						trace.StatusUpdate(index, colIndex + 5, b, LogEntryType_Info);
 					}
 				}
 				estimatedMemoryUsage = memBaseLine + untrackedMemory + estimatedProcessMemoryUsage;
@@ -1512,14 +1512,14 @@ namespace uba
 
 			if (m_memTracingLevel >= 1)
 			{
-				trace.StatusUpdate(StartRow, 6, BytesToText(memUsed).str, LogEntryType_Info, nullptr);
-				trace.StatusUpdate(StartRow+1, 6, BytesToText(estimatedMemoryUsage).str, LogEntryType_Info, nullptr);
-				trace.StatusUpdate(StartRow, 15, BytesToText(processMemTotalPageCount).str, LogEntryType_Info, nullptr);
-				trace.StatusUpdate(StartRow+1, 15, BytesToText(estimatedProcessMemoryUsage).str, LogEntryType_Info, nullptr);
-				trace.StatusUpdate(StartRow+2, 6, BytesToText(memBaseLine).str, LogEntryType_Info, nullptr);
-				trace.StatusUpdate(StartRow+3, 6, BytesToText(untrackedMemory).str, LogEntryType_Info, nullptr);
-				trace.StatusUpdate(StartRow+4, 6, BytesToText(m_memWaitThreshold).str, LogEntryType_Info, nullptr);
-				trace.StatusUpdate(StartRow+5, 6, BytesToText(memKillThreshold).str, LogEntryType_Info, nullptr);
+				trace.StatusUpdate(StartRow, 6, BytesToText(memUsed), LogEntryType_Info);
+				trace.StatusUpdate(StartRow+1, 6, BytesToText(estimatedMemoryUsage), LogEntryType_Info);
+				trace.StatusUpdate(StartRow, 15, BytesToText(processMemTotalPageCount), LogEntryType_Info);
+				trace.StatusUpdate(StartRow+1, 15, BytesToText(estimatedProcessMemoryUsage), LogEntryType_Info);
+				trace.StatusUpdate(StartRow+2, 6, BytesToText(memBaseLine), LogEntryType_Info);
+				trace.StatusUpdate(StartRow+3, 6, BytesToText(untrackedMemory), LogEntryType_Info);
+				trace.StatusUpdate(StartRow+4, 6, BytesToText(m_memWaitThreshold), LogEntryType_Info);
+				trace.StatusUpdate(StartRow+5, 6, BytesToText(memKillThreshold), LogEntryType_Info);
 			}
 
 			if (memKillThreshold)
@@ -1539,8 +1539,12 @@ namespace uba
 			{
 				ULONG reads = currPerfInfo.PagesRead - prevPerfInfo.PagesRead;
 				ULONG writes = currPerfInfo.PagefilePagesWritten - prevPerfInfo.PagefilePagesWritten;
-				trace.StatusUpdate(StartRow+6, 6, CountToText(reads).str, LogEntryType_Info, nullptr);
-				trace.StatusUpdate(StartRow+7, 6, CountToText(writes).str, LogEntryType_Info, nullptr);
+
+				if (m_memTracingLevel >= 1)
+				{
+					trace.StatusUpdate(StartRow+6, 6, CountToText(reads), LogEntryType_Info);
+					trace.StatusUpdate(StartRow+7, 6, CountToText(writes), LogEntryType_Info);
+				}
 
 				if (memUsed > physMem && writes > 1000)
 				{
diff --git a/Engine/Source/Programs/UnrealBuildAccelerator/Common/Private/UbaSessionServer.cpp b/Engine/Source/Programs/UnrealBuildAccelerator/Common/Private/UbaSessionServer.cpp
--- a/Engine/Source/Programs/UnrealBuildAccelerator/Common/Private/UbaSessionServer.cpp
+++ b/Engine/Source/Programs/UnrealBuildAccelerator/Common/Private/UbaSessionServer.cpp
@@ -587,7 +587,7 @@ namespace uba
 
 	void SessionServer::UpdateStatus(u32 statusRow, u32 statusColumn, const tchar* statusText, LogEntryType statusType, const tchar* statusLink)
 	{
-		m_trace.StatusUpdate(statusRow, statusColumn, statusText, statusType, statusLink);
+		m_trace.StatusUpdate(statusRow, statusColumn, ToView(statusText), statusType, statusLink ? ToView(statusLink) : StringView());
 	}
 
 	void SessionServer::AddProcessBreadcrumbs(u32 processId, const tchar* breadcrumbs, bool deleteOld)
diff --git a/Engine/Source/Programs/UnrealBuildAccelerator/Common/Private/UbaTrace.cpp b/Engine/Source/Programs/UnrealBuildAccelerator/Common/Private/UbaTrace.cpp
--- a/Engine/Source/Programs/UnrealBuildAccelerator/Common/Private/UbaTrace.cpp
+++ b/Engine/Source/Programs/UnrealBuildAccelerator/Common/Private/UbaTrace.cpp
@@ -238,7 +238,7 @@ namespace uba
 	}
 
 
-	u32 Trace::AddString(const StringView& string)
+	u32 Trace::AddString(StringView string)
 	{
 		if (!m_memoryBegin)
 			return 0;
@@ -274,7 +274,7 @@ namespace uba
 		if (writeTime) \
 			writer.Write7BitEncoded(time);
 
-	void Trace::SessionAdded(u32 sessionId, u32 clientId, const StringView& name, const StringView& info)
+	void Trace::SessionAdded(u32 sessionId, u32 clientId, StringView name, StringView info)
 	{
 		BEGIN_TRACE_ENTRY(SessionAdded);
 		writer.WriteString(name);
@@ -316,7 +316,7 @@ namespace uba
 		writer.WriteU32(sessionId);
 	}
 
-	void Trace::ProcessAdded(u32 sessionId, u32 processId, const StringView& description, const StringView& breadcrumbs)
+	void Trace::ProcessAdded(u32 sessionId, u32 processId, StringView description, StringView breadcrumbs)
 	{
 		BEGIN_TRACE_ENTRY(ProcessAdded);
 		writer.WriteU32(sessionId);
@@ -325,7 +325,7 @@ namespace uba
 		writer.WriteLongString(breadcrumbs, 20);
 	}
 
-	void Trace::ProcessEnvironmentUpdated(u32 processId, const StringView& reason, const u8* data, u64 dataSize, const StringView& breadcrumbs)
+	void Trace::ProcessEnvironmentUpdated(u32 processId, StringView reason, const u8* data, u64 dataSize, StringView breadcrumbs)
 	{
 		BEGIN_TRACE_ENTRY(ProcessEnvironmentUpdated);
 		writer.WriteU32(processId);
@@ -364,14 +364,14 @@ namespace uba
 		ProcessExited(processId, exitCode, writer.GetData(), writer.GetPosition(), {});
 	}
 
-	void Trace::ProcessReturned(u32 processId, const StringView& reason)
+	void Trace::ProcessReturned(u32 processId, StringView reason)
 	{
 		BEGIN_TRACE_ENTRY(ProcessReturned);
 		writer.WriteU32(processId);
 		writer.WriteString(reason);
 	}
 
-	void Trace::ProcessAddBreadcrumbs(u32 processId, const StringView& breadcrumbs, bool deleteOld)
+	void Trace::ProcessAddBreadcrumbs(u32 processId, StringView breadcrumbs, bool deleteOld)
 	{
 		BEGIN_TRACE_ENTRY(ProcessBreadcrumbs);
 		writer.WriteU32(processId);
@@ -400,7 +400,7 @@ namespace uba
 			writer.Write7BitEncoded(fileSize);
 	}
 
-	void Trace::FileFetchBegin(u32 clientId, const CasKey& key, const StringView& hint)
+	void Trace::FileFetchBegin(u32 clientId, const CasKey& key, StringView hint)
 	{
 		u32 stringIndex = AddString(hint);
 		BEGIN_TRACE_ENTRY(FileFetchBegin);
@@ -424,7 +424,7 @@ namespace uba
 		writer.WriteCasKey(key);
 	}
 
-	void Trace::FileStoreBegin(u32 clientId, const CasKey& key, u64 size, const StringView& hint, bool detailed)
+	void Trace::FileStoreBegin(u32 clientId, const CasKey& key, u64 size, StringView hint, bool detailed)
 	{
 		if (detailed)
 		{
@@ -450,7 +450,7 @@ namespace uba
 		writer.WriteCasKey(key);
 	}
 
-	void Trace::WorkBegin(u32 workIndex, const StringView& desc, const Color& color)
+	void Trace::WorkBegin(u32 workIndex, StringView desc, const Color& color)
 	{
 		u32 stringIndex = AddString(desc);
 		BEGIN_TRACE_ENTRY(WorkBegin);
@@ -459,7 +459,7 @@ namespace uba
 		writer.WriteU32(color);
 	}
 
-	void Trace::WorkHint(u32 workIndex, const StringView& hint, u64 startTime)
+	void Trace::WorkHint(u32 workIndex, StringView hint, u64 startTime)
 	{
 		u32 stringIndex = AddString(hint);
 		BEGIN_TRACE_ENTRY(WorkHint);
@@ -482,14 +482,14 @@ namespace uba
 		writer.Write7BitEncoded(errorCount);
 	}
 
-	void Trace::StatusUpdate(u32 statusRow, u32 statusColumn, const tchar* statusText, LogEntryType statusType, const tchar* statusLink)
+	void Trace::StatusUpdate(u32 statusRow, u32 statusColumn, StringView statusText, LogEntryType statusType, StringView statusLink)
 	{
 		BEGIN_TRACE_ENTRY(StatusUpdate);
 		writer.Write7BitEncoded(statusRow);
 		writer.Write7BitEncoded(statusColumn);
 		writer.WriteString(statusText);
 		writer.WriteByte(statusType);
-		writer.WriteString(statusLink ? statusLink : TC(""));
+		writer.WriteString(statusLink);
 	}
 
 	void Trace::DriveUpdate(const tchar drive, u8 busyPercent, u32 readCount, u64 readBytes, u32 writeCount, u64 writeBytes)
@@ -555,7 +555,7 @@ namespace uba
 		writer.WriteBytes(data, dataSize);
 	}
 
-	void Trace::TaskBegin(u32 processId, const StringView& description, const StringView& details, const Color& color)
+	void Trace::TaskBegin(u32 processId, StringView description, StringView details, const Color& color)
 	{
 		BEGIN_TRACE_ENTRY(TaskBegin);
 		writer.Write7BitEncoded(processId);
@@ -564,7 +564,7 @@ namespace uba
 		writer.WriteU32(color);
 	}
 
-	void Trace::TaskHint(u32 processId, const StringView& hint)
+	void Trace::TaskHint(u32 processId, StringView hint)
 	{
 		BEGIN_TRACE_ENTRY(TaskHint);
 		writer.Write7BitEncoded(processId);
diff --git a/Engine/Source/Programs/UnrealBuildAccelerator/Common/Public/UbaLogger.h b/Engine/Source/Programs/UnrealBuildAccelerator/Common/Public/UbaLogger.h
--- a/Engine/Source/Programs/UnrealBuildAccelerator/Common/Public/UbaLogger.h
+++ b/Engine/Source/Programs/UnrealBuildAccelerator/Common/Public/UbaLogger.h
@@ -95,6 +95,7 @@ namespace uba
 	{
 		BytesToText(u64 bytes, bool base2 = false);
 		operator const tchar* () const { return str; };
+		operator StringView () const;
 		tchar str[32];
 	};
 
@@ -104,6 +105,7 @@ namespace uba
 	{
 		CountToText(u64 count);
 		operator const tchar* () const { return str; };
+		operator StringView () const;
 		tchar str[32];
 	};
 
@@ -115,6 +117,7 @@ namespace uba
 		TimeToText(u64 time, bool allowMinutes = false);
 		TimeToText(u64 time, bool allowMinutes, u64 frequency);
 		operator const tchar*() const { return str; };
+		operator StringView () const;
 		tchar str[32];
 	};
 
diff --git a/Engine/Source/Programs/UnrealBuildAccelerator/Common/Public/UbaTrace.h b/Engine/Source/Programs/UnrealBuildAccelerator/Common/Public/UbaTrace.h
--- a/Engine/Source/Programs/UnrealBuildAccelerator/Common/Public/UbaTrace.h
+++ b/Engine/Source/Programs/UnrealBuildAccelerator/Common/Public/UbaTrace.h
@@ -84,33 +84,35 @@ namespace uba
 		Trace(LogWriter& logWriter);
 		~Trace();
 
+		const TString& GetNamedTrace() { return m_namedTrace; }
+
 		bool IsWriting() const { return m_memoryBegin != nullptr; }
 
 		bool StartWrite(const tchar* namedTrace, u64 traceMemCapacity = 64*1024*1024);
-		void SessionAdded(u32 sessionId, u32 clientId, const StringView& name, const StringView& info);
+		void SessionAdded(u32 sessionId, u32 clientId, StringView name, StringView info);
 		void SessionUpdate(u32 sessionId, u32 connectionCount, u64 send, u64 recv, u64 lastPing, u64 memAvail, u64 memTotal, float cpuLoad);
 		void SessionNotification(u32 sessionId, const tchar* text);
 		void SessionSummary(u32 sessionId, const u8* data, u64 dataSize);
 		void SessionDisconnect(u32 sessionId);
-		void ProcessAdded(u32 sessionId, u32 processId, const StringView& description, const StringView& breadcrumbs);
-		void ProcessEnvironmentUpdated(u32 processId, const StringView& reason, const u8* data, u64 dataSize, const StringView& breadcrumbs);
+		void ProcessAdded(u32 sessionId, u32 processId, StringView description, StringView breadcrumbs);
+		void ProcessEnvironmentUpdated(u32 processId, StringView reason, const u8* data, u64 dataSize, StringView breadcrumbs);
 		void ProcessExited(u32 processId, u32 exitCode, const u8* data, u64 dataSize, const Vector<ProcessLogLine>& logLines);
 		void ProcessExited(u32 processId, u32 exitCode);
-		void ProcessReturned(u32 processId, const StringView& reason);
-		void ProcessAddBreadcrumbs(u32 processId, const StringView& breadcrumbs, bool deleteOld);
+		void ProcessReturned(u32 processId, StringView reason);
+		void ProcessAddBreadcrumbs(u32 processId, StringView breadcrumbs, bool deleteOld);
 		void ProxyCreated(u32 clientId, const tchar* proxyName);
 		void ProxyUsed(u32 clientId, const tchar* proxyName);
 		void FileFetchLight(u32 clientId, const CasKey& key, u64 fileSize);
-		void FileFetchBegin(u32 clientId, const CasKey& key, const StringView& hint);
+		void FileFetchBegin(u32 clientId, const CasKey& key, StringView hint);
 		void FileFetchSize(u32 clientId, const CasKey& key, u64 size);
 		void FileFetchEnd(u32 clientId, const CasKey& key);
-		void FileStoreBegin(u32 clientId, const CasKey& key, u64 size, const StringView& hint, bool detailed);
+		void FileStoreBegin(u32 clientId, const CasKey& key, u64 size, StringView hint, bool detailed);
 		void FileStoreEnd(u32 clientId, const CasKey& key);
-		void WorkBegin(u32 workIndex, const StringView& desc, const Color& color);
-		void WorkHint(u32 workIndex, const StringView& hint, u64 startTime);
+		void WorkBegin(u32 workIndex, StringView desc, const Color& color);
+		void WorkHint(u32 workIndex, StringView hint, u64 startTime);
 		void WorkEnd(u32 workIndex);
 		void ProgressUpdate(u32 processesTotal, u32 processesDone, u32 errorCount);
-		void StatusUpdate(u32 statusRow, u32 statusColumn, const tchar* statusText, LogEntryType statusType, const tchar* statusLink);
+		void StatusUpdate(u32 statusRow, u32 statusColumn, StringView statusText, LogEntryType statusType, StringView statusLink = {});
 		void DriveUpdate(const tchar drive, u8 busyPercent, u32 readCount, u64 readBytes, u32 writeCount, u64 writeBytes);
 		void RemoteExecutionDisabled();
 		void SchedulerUpdate(bool spawnPaused);
@@ -122,8 +124,8 @@ namespace uba
 		void CacheEndWrite(u32 processId, bool success, const u8* data, u64 dataSize);
 		void CacheSummary(const u8* data, u64 dataSize);
 
-		void TaskBegin(u32 processId, const StringView& description, const StringView& details, const Color& color);
-		void TaskHint(u32 processId, const StringView& hint);
+		void TaskBegin(u32 processId, StringView description, StringView details, const Color& color);
+		void TaskHint(u32 processId, StringView hint);
 		void TaskEnd(u32 processId, bool success);
 
 		// Writes the current state of the trace to the specified output file without stopping anything
@@ -143,7 +145,7 @@ namespace uba
 		struct WriterScope;
 		void FreeMemory();
 		bool EnsureMemory(u64 size);
-		u32 AddString(const StringView& string);
+		u32 AddString(StringView string);
 
 		LoggerWithWriter m_logger;
 		TString m_namedTrace;
