[UBA]
* Refactored trace system to be more standalone
* Added support for providing trace from outside of session creation
* Added machine id as a thing to trace stream

diff --git a/Engine/Source/Programs/UnrealBuildAccelerator/Common/Private/UbaEnvironment.cpp b/Engine/Source/Programs/UnrealBuildAccelerator/Common/Private/UbaEnvironment.cpp
--- a/Engine/Source/Programs/UnrealBuildAccelerator/Common/Private/UbaEnvironment.cpp
+++ b/Engine/Source/Programs/UnrealBuildAccelerator/Common/Private/UbaEnvironment.cpp
@@ -3,6 +3,29 @@
 #include "UbaEnvironment.h"
 #include "UbaStringBuffer.h"
 
+#if PLATFORM_WINDOWS
+#include <powerbase.h>
+#pragma comment(lib, "Powrprof.lib")
+#elif PLATFORM_MAC
+#include <mach/vm_statistics.h>
+#include <mach/mach_types.h>
+#include <mach/mach_init.h>
+#include <mach/mach_host.h>
+#include <mach/mach.h>
+#include <sys/sysctl.h>
+#endif
+
+#if PLATFORM_WINDOWS
+typedef struct _PROCESSOR_POWER_INFORMATION {
+    ULONG  Number;
+    ULONG  MaxMhz;
+    ULONG  CurrentMhz;
+    ULONG  MhzLimit;
+    ULONG  MaxIdleState;
+    ULONG  CurrentIdleState;
+} PROCESSOR_POWER_INFORMATION, *PPROCESSOR_POWER_INFORMATION;
+#endif
+
 namespace uba
 {
 	MutexHandle CreateMutexW(bool bInitialOwner, const tchar* lpName)
@@ -200,6 +223,37 @@ namespace uba
 	}
 #endif
 
+	void GetMachineId(StringBufferBase& out)
+	{
+		bool success = false;
+		tchar str[256];
+
+#if PLATFORM_WINDOWS
+		DWORD strBytes = sizeof(str);
+		success = RegGetValueW(HKEY_LOCAL_MACHINE, TC("SOFTWARE\\Microsoft\\Cryptography"), TC("MachineGuid"), RRF_RT_REG_SZ, NULL, str, &strBytes) == ERROR_SUCCESS;
+		if (success)
+			out.Append(str);
+#elif PLATFORM_LINUX
+		if (FILE* id = fopen("/etc/machine-id", "r"))
+		{
+			success = fgets(str, sizeof_array(str), id) != nullptr;
+			fclose(id);
+			if (success)
+				out.Append(str);
+		}
+#else
+		if (FILE* id = popen("ioreg -rd1 -c IOPlatformExpertDevice | awk -F\\\" '/IOPlatformUUID/ { print $4 }'", "r"))
+		{
+			success = fgets(str, sizeof(str), id) != nullptr;
+			fclose(id);
+			if (success)
+				out.Append(str, strcspn(str, "\n"));
+		}
+#endif
+		if (!success)
+			out.Append(TCV("NoMachineId"));
+	}
+
 	bool GetMemoryInfo(Logger& logger, u64& outAvailable, u64& outTotal, u64* outMaxPageFileSize)
 	{
 		if (outMaxPageFileSize)
@@ -236,6 +290,213 @@ namespace uba
 		GetPhysicallyInstalledSystemMemory(memKb);
 		outTotal = memKb*1024*1024;
 		outAvailable = outTotal;
+#endif
+		return true;
+	}
+
+	void GetSystemInfo(Logger& logger, StringBufferBase& out)
+	{
+		u32 cpuCount = GetLogicalProcessorCount();
+		u32 cpuGroupCount = GetProcessorGroupCount();
+
+		StringBuffer<128> cpuStr(TC("CPU"));
+		if (IsRunningArm())
+			cpuStr.Append(TCV("[Arm]"));
+		cpuStr.Append(':');
+		if (cpuGroupCount != 1)
+			cpuStr.AppendValue(cpuGroupCount).Append('x');
+		cpuStr.AppendValue(cpuCount/cpuGroupCount);
+
+		u64 totalMemoryInKilobytes = 0;
+
+#if PLATFORM_WINDOWS
+		GetPhysicallyInstalledSystemMemory(&totalMemoryInKilobytes);
+
+		{
+			u32 maxMHz = 0;
+			DWORD valueSize = 4;
+			const tchar* key = TC("HARDWARE\\DESCRIPTION\\System\\CentralProcessor\\0");
+			LSTATUS res = RegGetValueW(HKEY_LOCAL_MACHINE, key, TC("~MHz"), RRF_RT_REG_DWORD, NULL, &maxMHz, &valueSize);
+			if (res != ERROR_SUCCESS)
+			{
+				// This will not always be the same and since we use the system info as part of key for client uniqueness it is annoying to get multiple sessions for same instance
+				Vector<PROCESSOR_POWER_INFORMATION> procInfos;
+				procInfos.resize(cpuCount);
+				if (CallNtPowerInformation(ProcessorInformation, NULL, 0, procInfos.data(), cpuCount*sizeof(PROCESSOR_POWER_INFORMATION)) == STATUS_SUCCESS)
+					maxMHz = procInfos[0].MaxMhz;
+			}
+			cpuStr.Appendf(TC(" @ %.1fGHz"), float(maxMHz) / 1000.0f);
+		}
+
+#else
+		u64 throwAway;
+		GetMemoryInfo(logger, throwAway, totalMemoryInKilobytes);
+		totalMemoryInKilobytes /= 1024;
+		double processorMhz = 0.0;
+	#if PLATFORM_LINUX
+		if (FILE *fp = fopen("/proc/cpuinfo", "r"))
+		{
+			char line[256];
+			while (processorMhz == 0.0 && fgets(line, sizeof(line), fp))
+				if (strncmp(line, "cpu MHz", 7) == 0)
+					sscanf(line, "cpu MHz : %lf", &processorMhz);
+			fclose(fp);
+			cpuStr.Appendf(" @ %.1fGHz", processorMhz/1000.0);
+		}
+	#else
+		char brand[128];
+		size_t size = sizeof(brand);
+		if (sysctlbyname("machdep.cpu.brand_string", brand, &size, NULL, 0) == 0)
+			cpuStr.Clear().Append(brand).Appendf(" CPU:%u", cpuCount);
+	#endif
+#endif
+
+		out.Appendf(TC("%s Mem:%ugb"), cpuStr.data, u32(totalMemoryInKilobytes/(1024*1024)));
+
+		#if PLATFORM_WINDOWS
+		if (!IsRunningWine())
+		{
+			DWORD value = 0;
+			DWORD valueSize = 4;
+			const tchar* fsKey = TC("SYSTEM\\CurrentControlSet\\Control\\FileSystem");
+			LSTATUS res = RegGetValueW(HKEY_LOCAL_MACHINE, fsKey, TC("NtfsDisableLastAccessUpdate"), RRF_RT_REG_DWORD, NULL, &value, &valueSize);
+			if (res != ERROR_SUCCESS)
+			{
+				logger.Detail(TC("Failed to retreive ntfs registry key (%i)"), res);
+			}
+			else
+			{
+				u32 lastAccessSettingsValue = value & 0xf;
+				if (lastAccessSettingsValue == 0 || lastAccessSettingsValue == 2)
+					out.Append(TCV(" NtfsLastAccessEnabled"));
+			}
+			value = 0;
+			res = RegGetValueW(HKEY_LOCAL_MACHINE, fsKey, TC("NtfsDisable8dot3NameCreation"), RRF_RT_REG_DWORD, NULL, &value, &valueSize);
+			if (res == ERROR_SUCCESS)
+				if (value == 0)
+					out.Append(TCV(" NtfsShortNamesEnabled"));
+		}
+		else
+		{
+			//StringBuffer<> testDir;
+			//testDir.Append(m_rootDir).Append(TCV("UbaTestShortNames"));
+			//::RemoveDirectory(testDir.data);
+			//wchar_t shortName[1024];
+			//if (::CreateDirectoryW(testDir.data, NULL))
+			//	if (GetShortPathName(testDir.data, shortName, 1024) != 0 && !Contains(shortName, TC("UbaTestShortNames")))
+			//		out.Append(TCV(" NtfsShortNamesEnabled"));
+		}
+		#endif
+	}
+
+	bool GetCpuTime(u64& outTotalTime, u64& outIdleTime)
+	{
+		outTotalTime = 0;
+		outIdleTime = 0;
+#if PLATFORM_WINDOWS
+
+		GROUP_AFFINITY originalAffinity {};
+		GROUP_AFFINITY newAffinity{};
+		static u16 groupCount = GetActiveProcessorGroupCount();
+		if (groupCount <= 1)
+		{
+			u64 idleTime, kernelTime, userTime;
+			if (!GetSystemTimes((FILETIME*)&idleTime, (FILETIME*)&kernelTime, (FILETIME*)&userTime))
+				return false;
+			outIdleTime += idleTime;
+			outTotalTime += kernelTime + userTime;
+		}
+		else
+		{
+			for (u16 group=0; group!=groupCount; ++group)
+			{
+				newAffinity.Mask = ~0ull;
+				newAffinity.Group = group;
+				if (!SetThreadGroupAffinity(GetCurrentThread(), &newAffinity, group == 0 ? &originalAffinity : NULL))
+					return false;
+				u64 idleTime, kernelTime, userTime;
+				if (!GetSystemTimes((FILETIME*)&idleTime, (FILETIME*)&kernelTime, (FILETIME*)&userTime))
+					return false;
+				outIdleTime += idleTime;
+				outTotalTime += kernelTime + userTime;
+			}
+			if (!SetThreadGroupAffinity(GetCurrentThread(), &originalAffinity, nullptr))
+				return false;
+		}
+
+#elif PLATFORM_LINUX
+		int fd = open("/proc/stat", O_RDONLY | O_CLOEXEC);
+		if (fd == -1)
+			return false;
+		char buffer[512];
+		int size = read(fd, buffer, sizeof_array(buffer) - 1);
+		close(fd);
+		if (size == -1)
+			return false;
+		buffer[size] = 0;
+		char* endl = strchr(buffer, '\n');
+		if (!endl)
+			return false;
+		u64 values[16] = { 0 };
+		u32 valueCount = 0;
+		*endl = 0;
+		char* parsePos = buffer;
+
+		// cpu
+		parsePos = strchr(parsePos, ' ');
+		if (!parsePos)
+			return false;
+
+		while (true)
+		{
+			// remove spaces
+			while (*parsePos && (*parsePos < '0' || *parsePos > '9'))
+				++parsePos;
+			if (!*parsePos)
+				break;
+			char* numberStart = parsePos;
+			while (*parsePos && *parsePos >= '0' && *parsePos <= '9')
+				++parsePos;
+			bool isLast = *parsePos == 0;
+			*parsePos = 0;
+			++parsePos;
+			values[valueCount++] = strtoull(numberStart, nullptr, 10);
+			if (isLast)
+				break;
+		}
+
+		// user: normal processes executing in user mode
+		// nice: niced processes executing in user mode
+		// system: processes executing in kernel mode
+		// idle: twiddling thumbs
+		// iowait: waiting for I/O to complete
+		// irq: servicing interrupts
+		// softirq: servicing softirqs
+		// steal
+		if (valueCount <= 6)
+			return false;
+		u64 work = values[0] + values[1] + values[2];
+		outIdleTime = values[3] + values[4] + values[5] + values[6] + values[7];
+		outTotalTime = work + outIdleTime;
+
+#else // PLATFORM_MAC
+		mach_msg_type_number_t  CpuMsgCount = 0;
+		processor_flavor_t CpuInfoType = PROCESSOR_CPU_LOAD_INFO;;
+		natural_t CpuCount = 0;
+		processor_cpu_load_info_t CpuData;
+		host_t host = mach_host_self();
+		u64 work = 0;
+		int res = host_processor_info(host, CpuInfoType, &CpuCount, (processor_info_array_t *)&CpuData, &CpuMsgCount);
+		if(res != KERN_SUCCESS)
+			return false;//m_logger.Error(TC("Kernel error: %s"), mach_error_string(res));
+		for(int i = 0; i < (int)CpuCount; i++)
+		{
+			work += CpuData[i].cpu_ticks[CPU_STATE_SYSTEM];
+			work += CpuData[i].cpu_ticks[CPU_STATE_USER];
+			work += CpuData[i].cpu_ticks[CPU_STATE_NICE];
+			outIdleTime += CpuData[i].cpu_ticks[CPU_STATE_IDLE];
+		}
+		outTotalTime = work + outIdleTime;
 #endif
 		return true;
 	}
diff --git a/Engine/Source/Programs/UnrealBuildAccelerator/Common/Private/UbaSession.cpp b/Engine/Source/Programs/UnrealBuildAccelerator/Common/Private/UbaSession.cpp
--- a/Engine/Source/Programs/UnrealBuildAccelerator/Common/Private/UbaSession.cpp
+++ b/Engine/Source/Programs/UnrealBuildAccelerator/Common/Private/UbaSession.cpp
@@ -17,16 +17,7 @@
 #include "UbaWorkManager.h"
 #include <span>
 
-#if PLATFORM_WINDOWS
-#include <powerbase.h>
-#pragma comment(lib, "Powrprof.lib")
-#elif PLATFORM_MAC
-#include <mach/vm_statistics.h>
-#include <mach/mach_types.h>
-#include <mach/mach_init.h>
-#include <mach/mach_host.h>
-#include <mach/mach.h>
-#include <sys/sysctl.h>
+#if !PLATFORM_WINDOWS
 extern char **environ;
 #endif
 
@@ -34,17 +25,6 @@ extern char **environ;
 
 //////////////////////////////////////////////////////////////////////////////
 
-#if PLATFORM_WINDOWS
-typedef struct _PROCESSOR_POWER_INFORMATION {
-    ULONG  Number;
-    ULONG  MaxMhz;
-    ULONG  CurrentMhz;
-    ULONG  MhzLimit;
-    ULONG  MaxIdleState;
-    ULONG  CurrentIdleState;
-} PROCESSOR_POWER_INFORMATION, *PPROCESSOR_POWER_INFORMATION;
-#endif
-
 namespace uba
 {
 	bool g_dummy;
@@ -1518,10 +1498,11 @@ namespace uba
 	:	m_storage(info.storage)
 	,	m_logger(info.logWriter, logPrefix)
 	,	m_workManager(workManager)
+	,	m_ownsTrace(info.trace == nullptr)
 	,	m_directoryTable(m_directoryTableMemory)
 	,	m_fileMappingBuffer(m_logger, &workManager)
 	,	m_processCommunicationAllocator(m_logger, TC("CommunicationAllocator"))
-	,	m_trace(info.logWriter)
+	,	m_trace(info.trace ? *info.trace : *new Trace(info.logWriter))
 	,	m_writeFilesBottleneck(info.writeFilesBottleneck)
 	,	m_writeFilesFileMapMax(u64(info.writeFilesFileMapMaxMb)*1024*1024)
 	,	m_writeFilesNoBufferingMin(u64(info.writeFilesNoBufferingMinMb)*1024*1024)
@@ -1581,7 +1562,6 @@ namespace uba
 		m_allowLinkDependencyCrawler = info.allowLinkDependencyCrawler;
 
 		m_detailedTrace = info.detailedTrace;
-		m_traceIntervalMs = info.traceIntervalMs;
 		m_traceChildProcesses = info.traceChildProcesses;
 		m_traceWrittenFiles = info.traceWrittenFiles;
 		m_logToFile = info.logToFile;
@@ -1645,7 +1625,7 @@ namespace uba
 		#endif
 
 		if (info.traceOutputFile)
-			m_traceOutputFile.Append(info.traceOutputFile);
+			m_traceOutputFile = info.traceOutputFile;
 
 		if (m_readIntermediateFilesCompressed && !m_runningRemote)
 			m_dependencyCrawler.Init(
@@ -1706,27 +1686,42 @@ namespace uba
 			if (!m_volumeCache.Init(m_logger))
 				return false;
 
-		StringBuffer<> traceName;
-		if (info.traceName && *info.traceName)
-			traceName.Append(info.traceName);
-		else if (info.launchVisualizer || !m_traceOutputFile.IsEmpty() || info.traceEnabled)
+		if (m_ownsTrace)
 		{
-			traceName.Append(m_id);
+			StringBuffer<> traceName;
+			if (info.traceName && *info.traceName)
+				traceName.Append(info.traceName);
+			else if (info.launchVisualizer || !m_traceOutputFile.empty() || info.traceEnabled)
+			{
+				traceName.Append(m_id);
+
+				OwnerInfo ownerInfo = GetOwnerInfo();
+				if (ownerInfo.pid)
+					traceName.Appendf(TC("_%s%u"), ownerInfo.id, ownerInfo.pid);
 
-			OwnerInfo ownerInfo = GetOwnerInfo();
-			if (ownerInfo.pid)
-				traceName.Appendf(TC("_%s%u"), ownerInfo.id, ownerInfo.pid);
+				if (!info.useUniqueId)
+				{
+					Guid guid;
+					CreateGuid(guid);
+					traceName.Append(GuidToString(guid).str);
+				}
+			}
 
-			if (!info.useUniqueId)
+			if (!traceName.IsEmpty())
 			{
-				Guid guid;
-				CreateGuid(guid);
-				traceName.Append(GuidToString(guid).str);
+				u64 traceReserveSize = info.traceReserveSizeMb * 1024 * 1024;
+				if (m_detailedTrace)
+					traceReserveSize *= 2;
+				m_trace.StartWriteAndThread(IsWindows ? traceName.data : nullptr, traceReserveSize, true); // non-windows named shared memory not implemented (only needed for UbaVisualizer which you can't run on linux either way)
 			}
 		}
 
-		if (!traceName.IsEmpty())
-			StartTrace(IsWindows ? traceName.data : nullptr, info.traceReserveSizeMb); // non-windows named shared memory not implemented (only needed for UbaVisualizer which you can't run on linux either way)
+		if (m_trace.IsWriting())
+		{
+			StringBuffer<512> sessionInfo;
+			GetSessionInfo(sessionInfo);
+			m_trace.SessionInfo(0, sessionInfo);
+		}
 
 		#if PLATFORM_WINDOWS
 		if (info.launchVisualizer)
@@ -1739,7 +1734,7 @@ namespace uba
 			launcherCmd.Append(TCV("\""));
 			launcherCmd.AppendDir(fileName);
 			launcherCmd.Append(TCV("\\UbaVisualizer.exe\""));
-			launcherCmd.Append(TCV(" -named=")).Append(traceName);
+			launcherCmd.Append(TCV(" -named=")).Append(m_trace.GetNamedTrace());
 			STARTUPINFOW si;
 			memset(&si, 0, sizeof(si));
 			PROCESS_INFORMATION pi;
@@ -1776,7 +1771,11 @@ namespace uba
 
 	Session::~Session()
 	{
-		StopTrace(m_traceOutputFile.data);
+		if (m_ownsTrace)
+		{
+			m_trace.StopThread();
+			m_trace.StopWrite(m_traceOutputFile.c_str());
+		}
 
 		CancelAllProcessesAndWait();
 		FlushDeadProcesses();
@@ -1802,6 +1801,9 @@ namespace uba
 		#if UBA_DEBUG_LOGGER
 		m_debugLogger = StopDebugLogger(m_debugLogger);
 		#endif
+
+		if (m_ownsTrace)
+			delete &m_trace;
 	}
 
 	void Session::CancelAllProcessesAndWait(bool terminate)
@@ -2547,50 +2549,9 @@ namespace uba
 		stats.Print(m_logger);
 	}
 
-	void Session::StartTrace(const tchar* traceName, u32 traceReserveSizeMb)
-	{
-		if (traceName)
-			LoggerWithWriter(m_logger.m_writer).Info(TC("---- Starting trace: %s ----"), traceName);
-		else
-			LoggerWithWriter(m_logger.m_writer).Info(TC("---- Starting trace ----"));
-
-		u64 traceReserveSize = u64(traceReserveSizeMb)*1024*1024;
-		if (m_detailedTrace)
-			traceReserveSize *= 2;
-
-		if (!m_trace.StartWrite(traceName, traceReserveSize))
-			return;
-
-		StringBuffer<256> buf;
-		if (!GetComputerNameW(buf))
-			buf.Append(TCV("LOCAL"));
-		StringBuffer<> systemInfo;
-		GetSystemInfo(systemInfo);
-		m_trace.SessionAdded(0, {}, buf, systemInfo);
-		StartTraceThread();
-	}
-
-	bool Session::StopTrace(const tchar* writeFile)
-	{
-		StopTraceThread();
-		return m_trace.StopWrite(writeFile);
-	}
-
 	bool Session::SaveSnapshotOfTrace()
 	{
-		return m_trace.Write(m_traceOutputFile.data, true);
-	}
-
-	void Session::StartTraceThread()
-	{
-		m_traceThreadEvent.Create(true);
-		m_traceThread.Start([this]() { ThreadTraceLoop(); return 0; }, TC("UbaTraceLoop"));
-	}
-
-	void Session::StopTraceThread()
-	{
-		m_traceThreadEvent.Set();
-		m_traceThread.Wait();
+		return m_trace.Write(m_traceOutputFile.c_str(), true);
 	}
 
 	void Session::PrintSummary(Logger& logger)
@@ -3683,109 +3644,19 @@ namespace uba
 		m_directoryTableMemCommitted += toCommit;
 	}
 
-	void Session::GetSystemInfo(StringBufferBase& out)
+	void Session::GetSessionInfo(StringBufferBase& out)
 	{
-		u32 cpuCount = GetLogicalProcessorCount();
-		u32 cpuGroupCount = GetProcessorGroupCount();
-
-		StringBuffer<128> cpuStr(TC("CPU"));
-		if (IsRunningArm())
-			cpuStr.Append(TCV("[Arm]"));
-		cpuStr.Append(':');
-		if (cpuGroupCount != 1)
-			cpuStr.AppendValue(cpuGroupCount).Append('x');
-		cpuStr.AppendValue(cpuCount/cpuGroupCount);
-
-		u64 totalMemoryInKilobytes = 0;
-
-#if PLATFORM_WINDOWS
-		GetPhysicallyInstalledSystemMemory(&totalMemoryInKilobytes);
-
-		{
-			u32 maxMHz = 0;
-			DWORD valueSize = 4;
-			const tchar* key = TC("HARDWARE\\DESCRIPTION\\System\\CentralProcessor\\0");
-			LSTATUS res = RegGetValueW(HKEY_LOCAL_MACHINE, key, TC("~MHz"), RRF_RT_REG_DWORD, NULL, &maxMHz, &valueSize);
-			if (res != ERROR_SUCCESS)
-			{
-				// This will not always be the same and since we use the system info as part of key for client uniqueness it is annoying to get multiple sessions for same instance
-				Vector<PROCESSOR_POWER_INFORMATION> procInfos;
-				procInfos.resize(cpuCount);
-				if (CallNtPowerInformation(ProcessorInformation, NULL, 0, procInfos.data(), cpuCount*sizeof(PROCESSOR_POWER_INFORMATION)) == STATUS_SUCCESS)
-					maxMHz = procInfos[0].MaxMhz;
-			}
-			cpuStr.Appendf(TC(" @ %.1fGHz"), float(maxMHz) / 1000.0f);
-		}
-
-#else
-		u64 throwAway;
-		GetMemoryInfo(m_logger, throwAway, totalMemoryInKilobytes);
-		totalMemoryInKilobytes /= 1024;
-		double processorMhz = 0.0;
-	#if PLATFORM_LINUX
-		if (FILE *fp = fopen("/proc/cpuinfo", "r"))
-		{
-			char line[256];
-			while (processorMhz == 0.0 && fgets(line, sizeof(line), fp))
-				if (strncmp(line, "cpu MHz", 7) == 0)
-					sscanf(line, "cpu MHz : %lf", &processorMhz);
-			fclose(fp);
-			cpuStr.Appendf(" @ %.1fGHz", processorMhz/1000.0);
-		}
-	#else
-		char brand[128];
-		size_t size = sizeof(brand);
-		if (sysctlbyname("machdep.cpu.brand_string", brand, &size, NULL, 0) == 0)
-			cpuStr.Clear().Append(brand).Appendf(" CPU:%u", cpuCount);
-	#endif
-#endif
-
-		const tchar* capacityStr = TC("NoLimit");
-		u64 capacity = m_storage.GetStorageCapacity();
-		BytesToText temp(capacity);
-		if (capacity)
-			capacityStr = temp.str;
-		out.Appendf(TC("%s Mem:%ugb Cas:%s/%s"), cpuStr.data, u32(totalMemoryInKilobytes/(1024*1024)), BytesToText(m_storage.GetStorageUsed()).str, capacityStr);
+		out.Append(TCV("Cas:"));
+		out.Append(BytesToText(m_storage.GetStorageUsed()).str).Append('/');
+		if (u64 capacity = m_storage.GetStorageCapacity())
+			out.Append(BytesToText(capacity).str);
+		else
+			out.Append(TCV("NoLimit"));
 
 		StringBuffer<128> zone;
 		if (m_storage.GetZone(zone))
 			out.Append(TCV(" Zone:")).Append(zone);
 
-		#if PLATFORM_WINDOWS
-		if (!IsRunningWine())
-		{
-			DWORD value = 0;
-			DWORD valueSize = 4;
-			const tchar* fsKey = TC("SYSTEM\\CurrentControlSet\\Control\\FileSystem");
-			LSTATUS res = RegGetValueW(HKEY_LOCAL_MACHINE, fsKey, TC("NtfsDisableLastAccessUpdate"), RRF_RT_REG_DWORD, NULL, &value, &valueSize);
-			if (res != ERROR_SUCCESS)
-			{
-				m_logger.Detail(TC("Failed to retreive ntfs registry key (%i)"), res);
-			}
-			else
-			{
-				u32 lastAccessSettingsValue = value & 0xf;
-				if (lastAccessSettingsValue == 0 || lastAccessSettingsValue == 2)
-					out.Append(TCV(" NtfsLastAccessEnabled"));
-			}
-			value = 0;
-			res = RegGetValueW(HKEY_LOCAL_MACHINE, fsKey, TC("NtfsDisable8dot3NameCreation"), RRF_RT_REG_DWORD, NULL, &value, &valueSize);
-			if (res == ERROR_SUCCESS)
-				if (value == 0)
-					out.Append(TCV(" NtfsShortNamesEnabled"));
-		}
-		else
-		{
-			StringBuffer<> testDir;
-			testDir.Append(m_rootDir).Append(TCV("UbaTestShortNames"));
-			::RemoveDirectory(testDir.data);
-			wchar_t shortName[1024];
-			if (::CreateDirectoryW(testDir.data, NULL))
-				if (GetShortPathName(testDir.data, shortName, 1024) != 0 && !Contains(shortName, TC("UbaTestShortNames")))
-					out.Append(TCV(" NtfsShortNamesEnabled"));
-		}
-		#endif
-
 		if (!m_extraInfo.empty())
 			out.Append(m_extraInfo);
 
@@ -3794,140 +3665,6 @@ namespace uba
 		#endif
 	}
 
-	bool GetCpuTime(u64& outTotalTime, u64& outIdleTime)
-	{
-		outTotalTime = 0;
-		outIdleTime = 0;
-#if PLATFORM_WINDOWS
-
-		GROUP_AFFINITY originalAffinity {};
-		GROUP_AFFINITY newAffinity{};
-		static u16 groupCount = GetActiveProcessorGroupCount();
-		if (groupCount <= 1)
-		{
-			u64 idleTime, kernelTime, userTime;
-			if (!GetSystemTimes((FILETIME*)&idleTime, (FILETIME*)&kernelTime, (FILETIME*)&userTime))
-				return false;
-			outIdleTime += idleTime;
-			outTotalTime += kernelTime + userTime;
-		}
-		else
-		{
-			for (u16 group=0; group!=groupCount; ++group)
-			{
-				newAffinity.Mask = ~0ull;
-				newAffinity.Group = group;
-				if (!SetThreadGroupAffinity(GetCurrentThread(), &newAffinity, group == 0 ? &originalAffinity : NULL))
-					return false;
-				u64 idleTime, kernelTime, userTime;
-				if (!GetSystemTimes((FILETIME*)&idleTime, (FILETIME*)&kernelTime, (FILETIME*)&userTime))
-					return false;
-				outIdleTime += idleTime;
-				outTotalTime += kernelTime + userTime;
-			}
-			if (!SetThreadGroupAffinity(GetCurrentThread(), &originalAffinity, nullptr))
-				return false;
-		}
-
-#elif PLATFORM_LINUX
-		int fd = open("/proc/stat", O_RDONLY | O_CLOEXEC);
-		if (fd == -1)
-			return false;
-		char buffer[512];
-		int size = read(fd, buffer, sizeof_array(buffer) - 1);
-		close(fd);
-		if (size == -1)
-			return false;
-		buffer[size] = 0;
-		char* endl = strchr(buffer, '\n');
-		if (!endl)
-			return false;
-		u64 values[16] = { 0 };
-		u32 valueCount = 0;
-		*endl = 0;
-		char* parsePos = buffer;
-
-		// cpu
-		parsePos = strchr(parsePos, ' ');
-		if (!parsePos)
-			return false;
-
-		while (true)
-		{
-			// remove spaces
-			while (*parsePos && (*parsePos < '0' || *parsePos > '9'))
-				++parsePos;
-			if (!*parsePos)
-				break;
-			char* numberStart = parsePos;
-			while (*parsePos && *parsePos >= '0' && *parsePos <= '9')
-				++parsePos;
-			bool isLast = *parsePos == 0;
-			*parsePos = 0;
-			++parsePos;
-			values[valueCount++] = strtoull(numberStart, nullptr, 10);
-			if (isLast)
-				break;
-		}
-
-		// user: normal processes executing in user mode
-		// nice: niced processes executing in user mode
-		// system: processes executing in kernel mode
-		// idle: twiddling thumbs
-		// iowait: waiting for I/O to complete
-		// irq: servicing interrupts
-		// softirq: servicing softirqs
-		// steal
-		if (valueCount <= 6)
-			return false;
-		u64 work = values[0] + values[1] + values[2];
-		outIdleTime = values[3] + values[4] + values[5] + values[6] + values[7];
-		outTotalTime = work + outIdleTime;
-
-#else // PLATFORM_MAC
-		mach_msg_type_number_t  CpuMsgCount = 0;
-		processor_flavor_t CpuInfoType = PROCESSOR_CPU_LOAD_INFO;;
-		natural_t CpuCount = 0;
-		processor_cpu_load_info_t CpuData;
-		host_t host = mach_host_self();
-		u64 work = 0;
-		int res = host_processor_info(host, CpuInfoType, &CpuCount, (processor_info_array_t *)&CpuData, &CpuMsgCount);
-		if(res != KERN_SUCCESS)
-			return false;//m_logger.Error(TC("Kernel error: %s"), mach_error_string(res));
-		for(int i = 0; i < (int)CpuCount; i++)
-		{
-			work += CpuData[i].cpu_ticks[CPU_STATE_SYSTEM];
-			work += CpuData[i].cpu_ticks[CPU_STATE_USER];
-			work += CpuData[i].cpu_ticks[CPU_STATE_NICE];
-			outIdleTime += CpuData[i].cpu_ticks[CPU_STATE_IDLE];
-		}
-		outTotalTime = work + outIdleTime;
-#endif
-		return true;
-	}
-
-	float Session::UpdateCpuLoad()
-	{
-		u64 totalTime = 0;
-		u64 idleTime = 0;
-		if (!GetCpuTime(totalTime, idleTime))
-			return m_cpuLoad;
-
-		u64 totalTimeSinceLastTime = totalTime - m_previousTotalCpuTime;
-		u64 idleTimeSinceLastTime = idleTime - m_previousIdleCpuTime;
-
-		float cpuLoad = 1.0f - ((totalTimeSinceLastTime > 0) ? (float(idleTimeSinceLastTime) / float(totalTimeSinceLastTime)) : 0);
-
-		m_previousTotalCpuTime = totalTime;
-		m_previousIdleCpuTime = idleTime;
-
-		// TODO: This is the wrong solution.. but can't repro the bad values some people get
-		if (cpuLoad >= 0 && cpuLoad <= 1.0f)
-			m_cpuLoad = cpuLoad;
-
-		return m_cpuLoad;
-	}
-
 	bool Session::HasVfs(RootsHandle handle) const
 	{
 		return (handle & 1ull) == 1ull;
@@ -4092,16 +3829,6 @@ namespace uba
 		return rootsEntry.roots.NormalizeString<char>(m_logger, (const char*)sourceData, sourceSize, handleString, true, hint);
 	}
 
-	void Session::ThreadTraceLoop()
-	{
-		while (true)
-		{
-			TraceSessionUpdate();
-			if (m_traceThreadEvent.IsSet(m_traceIntervalMs))
-				break;
-		}
-	}
-
 	void Session::TraceWrittenFile(u32 processId, const StringView& file, u64 size)
 	{
 		if (!m_traceWrittenFiles)
@@ -4122,10 +3849,6 @@ namespace uba
 		m_trace.ProcessAddBreadcrumbs(processId, str, false);
 	}
 
-	void Session::TraceSessionUpdate()
-	{
-	}
-
 	void Session::RunDependencyCrawler(ProcessImpl& process)
 	{
 		auto& startInfo = process.GetStartInfo();
diff --git a/Engine/Source/Programs/UnrealBuildAccelerator/Common/Private/UbaSessionClient.cpp b/Engine/Source/Programs/UnrealBuildAccelerator/Common/Private/UbaSessionClient.cpp
--- a/Engine/Source/Programs/UnrealBuildAccelerator/Common/Private/UbaSessionClient.cpp
+++ b/Engine/Source/Programs/UnrealBuildAccelerator/Common/Private/UbaSessionClient.cpp
@@ -79,7 +79,7 @@ namespace uba
 
 		m_useDependencyCrawler = info.useDependencyCrawler;
 
-		m_processIdCounter = ~0u / 2; // We set this value to a very high value.. because it will be used by child processes and we don't want id from server and child process id to collide
+		m_processIdCounter = 1'000'000; // We set this value to a very high value.. because it will be used by child processes and we don't want id from server and child process id to collide
 
 		if (m_name.IsEmpty())
 			GetComputerNameW(m_name);
@@ -97,6 +97,8 @@ namespace uba
 
 		m_nameToHashTableMem.Init(NameToHashMemSize);
 
+		m_trace.SetThreadUpdateCallback([this]() { TraceSessionUpdate(); }, info.traceIntervalMs);
+
 		Create(info);
 
 
@@ -1306,17 +1308,23 @@ namespace uba
 		writer.WriteCasKey(keys[0]);
 		writer.WriteCasKey(keys[1]);
 
+		StringBuffer<256> machineId;
+		GetMachineId(machineId);
+		writer.WriteString(machineId);
+
 		writer.WriteU32(m_maxProcessCount);
 		writer.WriteBool(m_dedicated);
 
 		StringBuffer<> info;
-		GetSystemInfo(info);
+		GetSystemInfo(m_logger, info);
+		info.Append(' ');
+		GetSessionInfo(info);
 		writer.WriteString(info);
 
 		u64 memAvail;
 		u64 memTotal;
 		GetMemoryInfo(m_logger, memAvail, memTotal);
-		float cpuLoad = UpdateCpuLoad();
+		float cpuLoad = m_trace.UpdateCpuLoad();
 
 		m_cpuUsage = cpuLoad;
 		m_memAvail = memAvail;
@@ -1407,14 +1415,13 @@ namespace uba
 		if (m_shouldSendTraceToServer)
 		{
 			//if (m_detailedTrace)
-			{
 				m_client.SetWorkTracker(&m_trace);
-			}
-			StartTrace(nullptr, 256);
+
+			m_trace.StartWriteAndThread(nullptr, 256*1024*1024, true);
 		}
 		else
 		{
-			StartTraceThread(); // We still use trace thread for ping messages
+			m_trace.StartThread(); // We still use trace thread for ping messages
 		}
 
 		BuildEnvironmentVariables(reader);
@@ -1638,7 +1645,7 @@ namespace uba
 		if (TimeToMs(time - m_lastPingSendTime) < 2000) // Ping every ~2 seconds... this is so server can disconnect a client quickly if no ping is coming
 			return;
 
-		float cpuLoad = UpdateCpuLoad();
+		float cpuLoad = m_trace.UpdateCpuLoad();
 		m_cpuUsage = cpuLoad;
 
 		StackBinaryWriter<128> writer;
@@ -2195,7 +2202,8 @@ namespace uba
 		m_client.FlushWork();
 
 
-		StopTraceThread();
+		m_trace.StopThread();
+		m_trace.SetThreadUpdateCallback({});
 
 		if (m_shouldSendTraceToServer)
 		{
@@ -2213,7 +2221,7 @@ namespace uba
 
 			StringBuffer<> ubaFile(m_sessionLogDir);
 			ubaFile.Append(TCV("Trace.uba"));
-			if (StopTrace(ubaFile.data))
+			if (m_trace.StopWrite(ubaFile.data))
 			{
 				WrittenFile f;
 				f.backedName = ubaFile.data;
diff --git a/Engine/Source/Programs/UnrealBuildAccelerator/Common/Private/UbaSessionServer.cpp b/Engine/Source/Programs/UnrealBuildAccelerator/Common/Private/UbaSessionServer.cpp
--- a/Engine/Source/Programs/UnrealBuildAccelerator/Common/Private/UbaSessionServer.cpp
+++ b/Engine/Source/Programs/UnrealBuildAccelerator/Common/Private/UbaSessionServer.cpp
@@ -223,6 +223,8 @@ namespace uba
 		m_detoursLibrary[IsArmBinary] = detoursFile.data;
 		#endif
 
+		m_trace.SetThreadUpdateCallback([this]() { TraceSessionUpdate(); }, info.traceIntervalMs);
+
 		if (!Create(info))
 		{
 			UBA_ASSERT(false);
@@ -232,7 +234,10 @@ namespace uba
 
 	SessionServer::~SessionServer()
 	{
-		StopTraceThread();
+		if (m_ownsTrace)
+			m_trace.StopThread();
+
+		m_trace.SetThreadUpdateCallback({});
 
 		m_server.SetWorkTracker(nullptr);
 		m_server.UnregisterOnClientDisconnected(ServiceId);
@@ -560,6 +565,21 @@ namespace uba
 		FlushDeadProcesses();
 	}
 
+	void SessionServer::WaitOnAllClients()
+	{
+		while (true)
+		{
+			SCOPED_CRITICAL_SECTION(m_remoteProcessAndSessionLock, queueLock);
+			bool connected = false;
+			for (auto s : m_clientSessions)
+				connected |= s->connected;
+			queueLock.Leave();
+			if (!connected)
+				break;
+			Sleep(10);
+		}
+	}
+
 	void SessionServer::SetMaxRemoteProcessCount(u32 count)
 	{
 		m_maxRemoteProcessCount.exchange(count);
@@ -634,6 +654,11 @@ namespace uba
 		return m_outerScheduler;
 	}
 
+	bool SessionServer::GetRemoteTraceEnabled()
+	{
+		return m_remoteTraceEnabled;
+	}
+
 	ProcessHandle SessionServer::GetNewestLocalProcess()
 	{
 		u64 newestTime = 0;
@@ -840,6 +865,9 @@ namespace uba
 			return true;
 		}
 
+		StringBuffer<256> machineId;
+		reader.ReadString(machineId);
+
 		u32 processSlotCount = reader.ReadU32();
 		bool dedicated = reader.ReadBool();
 
@@ -890,7 +918,8 @@ namespace uba
 
 		WriteRemoteEnvironmentVariables(writer);
 
-		m_trace.SessionAdded(sessionId, connectionInfo.GetId(), name, info); // Must be inside lock for TraceSessionUpdate() to not include
+		m_trace.SessionAdded(sessionId, connectionInfo.GetId(), name, machineId); // Must be inside lock for TraceSessionUpdate() to not include
+		m_trace.SessionInfo(sessionId, info);
 		m_trace.SessionUpdate(sessionId, 1, 0, 0, 0, memAvail, memTotal, cpuLoad);
 
 		lock.Leave();
@@ -2241,7 +2270,7 @@ namespace uba
 		}
 		lock.Leave();
 
-		float cpuLoad = UpdateCpuLoad();
+		float cpuLoad = m_trace.UpdateCpuLoad();
 		u64 memAvail;
 		u64 memTotal;
 		GetMemoryInfo(m_logger, memAvail, memTotal);
diff --git a/Engine/Source/Programs/UnrealBuildAccelerator/Common/Private/UbaTrace.cpp b/Engine/Source/Programs/UnrealBuildAccelerator/Common/Private/UbaTrace.cpp
--- a/Engine/Source/Programs/UnrealBuildAccelerator/Common/Private/UbaTrace.cpp
+++ b/Engine/Source/Programs/UnrealBuildAccelerator/Common/Private/UbaTrace.cpp
@@ -58,6 +58,51 @@ namespace uba
 		bool m_isValid = true;
 	};
 
+	bool Trace::StartWriteAndThread(const tchar* namedTrace, u64 traceMemCapacity, bool logStart)
+	{
+		if (logStart)
+		{
+			if (namedTrace)
+				m_logger.Info(TC("---- Starting trace: %s ----"), namedTrace);
+			else
+				m_logger.Info(TC("---- Starting trace ----"));
+		}
+
+		if (!StartWrite(namedTrace, traceMemCapacity))
+			return false;
+
+		StringBuffer<256> buf;
+		if (!GetComputerNameW(buf))
+			buf.Append(TCV("LOCAL"));
+		SessionAdded(0, {}, buf, {});
+
+		StringBuffer<> systemInfo;
+		GetSystemInfo(m_logger, systemInfo);
+		SessionInfo(0, systemInfo);
+
+		StartThread();
+		return true;
+	}
+
+	void Trace::StartThread()
+	{
+		m_threadEvent.Create(true);
+		m_thread.Start([this]() { ThreadLoop(); return 0; }, TC("UbaTraceLoop"));
+	}
+
+	void Trace::StopThread()
+	{
+		m_threadEvent.Set();
+		m_thread.Wait();
+	}
+
+	void Trace::SetThreadUpdateCallback(Function<void()>&& callback, u32 traceIntervalMs)
+	{
+		SCOPED_FUTEX(m_threadCallbackLock, lock);
+		m_threadCallback = callback;
+		m_threadIntervalMs = traceIntervalMs;
+	}
+
 	bool Trace::StartWrite(const tchar* namedTrace, u64 traceMemCapacity)
 	{
 		m_memoryCapacity = traceMemCapacity;
@@ -204,6 +249,51 @@ namespace uba
 		*lineCount = logWriter.count;
 	}
 
+	float Trace::UpdateCpuLoad()
+	{
+		u64 totalTime = 0;
+		u64 idleTime = 0;
+		if (!GetCpuTime(totalTime, idleTime))
+			return m_cpuLoad;
+
+		u64 totalTimeSinceLastTime = totalTime - m_previousTotalCpuTime;
+		u64 idleTimeSinceLastTime = idleTime - m_previousIdleCpuTime;
+
+		float cpuLoad = 1.0f - ((totalTimeSinceLastTime > 0) ? (float(idleTimeSinceLastTime) / float(totalTimeSinceLastTime)) : 0);
+
+		m_previousTotalCpuTime = totalTime;
+		m_previousIdleCpuTime = idleTime;
+
+		// TODO: This is the wrong solution.. but can't repro the bad values some people get
+		if (cpuLoad >= 0 && cpuLoad <= 1.0f)
+			m_cpuLoad = cpuLoad;
+
+		return m_cpuLoad;
+	}
+
+	void Trace::ThreadLoop()
+	{
+		while (true)
+		{
+			SCOPED_FUTEX_READ(m_threadCallbackLock, lock);
+			if (m_threadCallback)
+				m_threadCallback();
+			else
+			{
+				float cpuLoad = UpdateCpuLoad();
+				u64 memAvail;
+				u64 memTotal;
+				GetMemoryInfo(m_logger, memAvail, memTotal);
+				SessionUpdate(0, 0, 0, 0, 0, memAvail, memTotal, cpuLoad);
+			}
+			lock.Leave();
+
+
+			if (m_threadEvent.IsSet(m_threadIntervalMs))
+				break;
+		}
+	}
+
 	void Trace::FreeMemory()
 	{
 		if (m_memoryBegin)
@@ -274,15 +364,22 @@ namespace uba
 		if (writeTime) \
 			writer.Write7BitEncoded(time);
 
-	void Trace::SessionAdded(u32 sessionId, u32 clientId, StringView name, StringView info)
+	void Trace::SessionAdded(u32 sessionId, u32 clientId, StringView name, StringView machineId)
 	{
 		BEGIN_TRACE_ENTRY(SessionAdded);
 		writer.WriteString(name);
-		writer.WriteString(info);
+		writer.WriteString(machineId);
 		writer.Write7BitEncoded(clientId);
 		writer.WriteU32(sessionId);
 	}
 
+	void Trace::SessionInfo(u32 sessionId, StringView info)
+	{
+		BEGIN_TRACE_ENTRY(SessionInfo);
+		writer.WriteU32(sessionId);
+		writer.WriteString(info);
+	}
+
 	void Trace::SessionUpdate(u32 sessionId, u32 connectionCount, u64 send, u64 recv, u64 lastPing, u64 memAvail, u64 memTotal, float cpuLoad)
 	{
 		BEGIN_TRACE_ENTRY(SessionUpdate);
diff --git a/Engine/Source/Programs/UnrealBuildAccelerator/Common/Public/UbaEnvironment.h b/Engine/Source/Programs/UnrealBuildAccelerator/Common/Public/UbaEnvironment.h
--- a/Engine/Source/Programs/UnrealBuildAccelerator/Common/Public/UbaEnvironment.h
+++ b/Engine/Source/Programs/UnrealBuildAccelerator/Common/Public/UbaEnvironment.h
@@ -25,5 +25,8 @@ namespace uba
 	void ElevateCurrentThreadPriority();
 	void PrefetchVirtualMemory(const void* mem, u64 size);
 
+	void GetMachineId(StringBufferBase& out);
 	bool GetMemoryInfo(Logger& logger, u64& outAvailable, u64& outTotal, u64* outMaxPageFileSize = nullptr);
+	void GetSystemInfo(Logger& logger, StringBufferBase& out);
+	bool GetCpuTime(u64& outTotalTime, u64& outIdleTime);
 }
\ No newline at end of file
diff --git a/Engine/Source/Programs/UnrealBuildAccelerator/Common/Public/UbaNetwork.h b/Engine/Source/Programs/UnrealBuildAccelerator/Common/Public/UbaNetwork.h
--- a/Engine/Source/Programs/UnrealBuildAccelerator/Common/Public/UbaNetwork.h
+++ b/Engine/Source/Programs/UnrealBuildAccelerator/Common/Public/UbaNetwork.h
@@ -16,7 +16,7 @@ namespace uba
 	
 	static constexpr u32 SystemNetworkVersion = 1339;
 	static constexpr u32 StorageNetworkVersion = 4;
-	static constexpr u32 SessionNetworkVersion = 46;
+	static constexpr u32 SessionNetworkVersion = 47;
 	static constexpr u32 CacheNetworkVersion = 5;
 	static constexpr u32 CacheBucketVersion = 3;
 
diff --git a/Engine/Source/Programs/UnrealBuildAccelerator/Common/Public/UbaSession.h b/Engine/Source/Programs/UnrealBuildAccelerator/Common/Public/UbaSession.h
--- a/Engine/Source/Programs/UnrealBuildAccelerator/Common/Public/UbaSession.h
+++ b/Engine/Source/Programs/UnrealBuildAccelerator/Common/Public/UbaSession.h
@@ -122,10 +122,6 @@ namespace uba
 		void ProcessExited(ProcessImpl& process, u64 executionTime);
 		void FlushDeadProcesses();
 		void PrintProcessStats(ProcessStats& stats, const tchar* logName);
-		void StartTrace(const tchar* traceName, u32 traceReserveSizeMb);
-		bool StopTrace(const tchar* writeFile);
-		void StartTraceThread();
-		void StopTraceThread();
 		u32 GetFileMappingSize();
 		u32 GetMemoryMapAlignment(const StringView& fileName) const;
 		u32 GetMemoryMapAlignment(const StringView& fileName, bool runningRemote) const;
@@ -203,8 +199,7 @@ namespace uba
 		bool CopyImports(Vector<BinaryModule>& out, const tchar* library, tchar* applicationDir, tchar* applicationDirEnd, UnorderedSet<TString>& handledImports, const char* const* loaderPaths);
 		bool CreateProcessJobObject();
 		void EnsureDirectoryTableMemory(u64 neededSize);
-		void GetSystemInfo(StringBufferBase& out);
-		float UpdateCpuLoad();
+		void GetSessionInfo(StringBufferBase& out);
 
 		bool HasVfs(RootsHandle handle) const;
 		RootsHandle WithVfs(RootsHandle key, bool vfs) const;
@@ -215,9 +210,7 @@ namespace uba
 		bool ExtractSymbolsFromObjectFile(const CloseFileMessage& msg, const tchar* fileName, u64 fileSize);
 		bool DevirtualizeDepsFile(RootsHandle rootsHandle, MemoryBlock& destData, const void* sourceData, u64 sourceSize, bool escapeSpaces, const tchar* hint);
 		
-		void ThreadTraceLoop();
 		void TraceWrittenFile(u32 processId, const StringView& file, u64 size = 0);
-		virtual void TraceSessionUpdate();
 
 		void RunDependencyCrawler(ProcessImpl& process);
 
@@ -252,6 +245,7 @@ namespace uba
 		bool m_allowLocalDetour;
 		bool m_extractObjFilesSymbols;
 		bool m_allowLinkDependencyCrawler;
+		bool m_ownsTrace;
 
 		u64 m_keepOutputFileMemoryMapsThreshold;
 
@@ -330,15 +324,9 @@ namespace uba
 		u32 m_wantsOnCloseIdCounter = 1;
 
 		SessionStats m_stats;
-		Trace m_trace;
-		Event m_traceThreadEvent;
-		Thread m_traceThread;
-		StringBuffer<256> m_traceOutputFile;
-		u32 m_traceIntervalMs;
+		Trace& m_trace;
+		TString m_traceOutputFile;
 		TString m_extraInfo;
-		u64 m_previousTotalCpuTime = 0;
-		u64 m_previousIdleCpuTime = 0;
-		float m_cpuLoad = 0;
 
 		#if PLATFORM_WINDOWS
 		Futex m_processJobObjectLock;
diff --git a/Engine/Source/Programs/UnrealBuildAccelerator/Common/Public/UbaSessionClient.h b/Engine/Source/Programs/UnrealBuildAccelerator/Common/Public/UbaSessionClient.h
--- a/Engine/Source/Programs/UnrealBuildAccelerator/Common/Public/UbaSessionClient.h
+++ b/Engine/Source/Programs/UnrealBuildAccelerator/Common/Public/UbaSessionClient.h
@@ -80,7 +80,8 @@ namespace uba
 		virtual bool FlushWrittenFiles(ProcessImpl& process) override;
 		virtual bool UpdateEnvironment(ProcessImpl& process, const StringView& reason, bool resetStats) override;
 		virtual bool LogLine(ProcessImpl& process, const tchar* line, LogEntryType logType) override;
-		virtual void TraceSessionUpdate() override;
+		
+		void TraceSessionUpdate();
 
 		struct InternalProcessStartInfo;
 		struct ModuleInfo;
diff --git a/Engine/Source/Programs/UnrealBuildAccelerator/Common/Public/UbaSessionCreateInfo.h b/Engine/Source/Programs/UnrealBuildAccelerator/Common/Public/UbaSessionCreateInfo.h
--- a/Engine/Source/Programs/UnrealBuildAccelerator/Common/Public/UbaSessionCreateInfo.h
+++ b/Engine/Source/Programs/UnrealBuildAccelerator/Common/Public/UbaSessionCreateInfo.h
@@ -8,6 +8,7 @@ namespace uba
 {
 	class Config;
 	class Storage;
+	class Trace;
 
 	struct SessionCreateInfo
 	{
@@ -50,5 +51,7 @@ namespace uba
 		u32 writeFilesFileMapMaxMb = 100000;	// If file size is smaller than this number, then use memory maps for writing to disk
 		u32 writeFilesNoBufferingMinMb = 16;	// If file size is at least this number, then overlapped io with no buffering will be used
 		u32 traceIntervalMs = 500;				// milliseconds between each trace entry
+
+		Trace* trace = nullptr;					// If no trace is provided, session will create a new one
 	};
 }
diff --git a/Engine/Source/Programs/UnrealBuildAccelerator/Common/Public/UbaSessionServer.h b/Engine/Source/Programs/UnrealBuildAccelerator/Common/Public/UbaSessionServer.h
--- a/Engine/Source/Programs/UnrealBuildAccelerator/Common/Public/UbaSessionServer.h
+++ b/Engine/Source/Programs/UnrealBuildAccelerator/Common/Public/UbaSessionServer.h
@@ -54,6 +54,9 @@ namespace uba
 		// Wait for all queued up/active tasks to finish
 		void WaitOnAllTasks();
 
+		// Wait for all clients to disconnect.
+		void WaitOnAllClients();
+
 		// Can be used to hint the session how many processes (that can run remotely that are left.. session can then start disabling remote execution on clients not needed anymore
 		void SetMaxRemoteProcessCount(u32 count);
 
@@ -86,6 +89,8 @@ namespace uba
 		void SetOuterScheduler(Scheduler* scheduler);
 		Scheduler* GetOuterScheduler();
 
+		bool GetRemoteTraceEnabled();
+
 		ProcessHandle GetNewestLocalProcess();
 
 	protected:
@@ -117,7 +122,8 @@ namespace uba
 		virtual void FileEntryAdded(StringKey fileNameKey, u64 lastWritten, u64 size) override final;
 		virtual bool RunSpecialProgram(ProcessImpl& process, BinaryReader& reader, BinaryWriter& writer) override final;
 		virtual void PrintSessionStats(Logger& logger) override final;
-		virtual void TraceSessionUpdate() override final;
+		
+		void TraceSessionUpdate();
 
 		void WriteRemoteEnvironmentVariables(BinaryWriter& writer);
 		bool InitializeNameToHashTable();
diff --git a/Engine/Source/Programs/UnrealBuildAccelerator/Common/Public/UbaTrace.h b/Engine/Source/Programs/UnrealBuildAccelerator/Common/Public/UbaTrace.h
--- a/Engine/Source/Programs/UnrealBuildAccelerator/Common/Public/UbaTrace.h
+++ b/Engine/Source/Programs/UnrealBuildAccelerator/Common/Public/UbaTrace.h
@@ -5,10 +5,13 @@
 #include "UbaFileMapping.h"
 #include "UbaLogger.h"
 #include "UbaProcessStats.h"
+#include "UbaThread.h"
 #include "UbaWorkManager.h"
 
 namespace uba
 {
+	////////////////////////////////////////////////////////////////////////////////////////////////////
+
 	struct ProcessLogLine;
 	struct ProcessStats;
 
@@ -67,6 +70,7 @@ namespace uba
 		UBA_TRACE_TYPE(TaskEnd) \
 		UBA_TRACE_TYPE(SchedulerUpdate) \
 		UBA_TRACE_TYPE(SchedulerKillProcess) \
+		UBA_TRACE_TYPE(SessionInfo) \
 
 	enum TraceType : u8
 	{
@@ -75,21 +79,39 @@ namespace uba
 		#undef UBA_TRACE_TYPE
 	};
 
-	static constexpr u32 TraceVersion = 47;
+	static constexpr u32 TraceVersion = 48;
 	static constexpr u32 TraceReadCompatibilityVersion = 6;
 
-	class Trace : public WorkTracker
+	////////////////////////////////////////////////////////////////////////////////////////////////////
+
+	class Trace final : public WorkTracker
 	{
 	public:
 		Trace(LogWriter& logWriter);
 		~Trace();
 
-		const TString& GetNamedTrace() { return m_namedTrace; }
+		bool StartWriteAndThread(const tchar* namedTrace, u64 traceMemCapacity, bool logStart);
+
+		void StartThread();
+		void StopThread();
+
+		void SetThreadUpdateCallback(Function<void()>&& callback, u32 traceIntervalMs = 500);
+
+
+		bool StartWrite(const tchar* namedTrace, u64 traceMemCapacity);
+
+		// Writes the current state of the trace to the specified output file without stopping anything
+		bool Write(const tchar* writeFileName, bool writeSummary);
+
+		// Stops the trace and writes it to the specified output file. This releases the internal memory buffer.
+		bool StopWrite(const tchar* writeFileName);
 
+
+		const TString& GetNamedTrace() { return m_namedTrace; }
 		bool IsWriting() const { return m_memoryBegin != nullptr; }
 
-		bool StartWrite(const tchar* namedTrace, u64 traceMemCapacity = 64*1024*1024);
-		void SessionAdded(u32 sessionId, u32 clientId, StringView name, StringView info);
+		void SessionAdded(u32 sessionId, u32 clientId, StringView name, StringView machineId);
+		void SessionInfo(u32 sessionId, StringView info);
 		void SessionUpdate(u32 sessionId, u32 connectionCount, u64 send, u64 recv, u64 lastPing, u64 memAvail, u64 memTotal, float cpuLoad);
 		void SessionNotification(u32 sessionId, const tchar* text);
 		void SessionSummary(u32 sessionId, const u8* data, u64 dataSize);
@@ -128,12 +150,6 @@ namespace uba
 		void TaskHint(u32 processId, StringView hint);
 		void TaskEnd(u32 processId, bool success);
 
-		// Writes the current state of the trace to the specified output file without stopping anything
-		bool Write(const tchar* writeFileName, bool writeSummary);
-
-		// Stops the trace and writes it to the specified output file. This releases the internal memory buffer.
-		bool StopWrite(const tchar* writeFileName);
-
 		virtual u32 TrackWorkStart(const StringView& desc, const Color& color) final override;
 		virtual void TrackWorkHint(u32 id, const StringView& hint, u64 startTime = 0) final override;
 		virtual void TrackWorkEnd(u32 id) final override;
@@ -141,8 +157,11 @@ namespace uba
 		void WriteSessionSummary(const Function<void(Logger& logger)>& summaryFunc);
 		static void WriteSummaryText(BinaryWriter& writer, const Function<void(Logger& logger)>& summaryFunc);
 
+		float UpdateCpuLoad();
+
 	private:
 		struct WriterScope;
+		void ThreadLoop();
 		void FreeMemory();
 		bool EnsureMemory(u64 size);
 		u32 AddString(StringView string);
@@ -163,9 +182,21 @@ namespace uba
 
 		Atomic<u32> m_workCounter;
 
+		Futex m_threadCallbackLock;
+		Function<void()> m_threadCallback;
+		u32 m_threadIntervalMs = 500;
+		Event m_threadEvent;
+		Thread m_thread;
+
+		u64 m_previousTotalCpuTime = 0;
+		u64 m_previousIdleCpuTime = 0;
+		float m_cpuLoad = 0;
+
 		friend class SessionServer;
 	};
 
+	////////////////////////////////////////////////////////////////////////////////////////////////////
+
 	struct OwnerInfo
 	{
 		const tchar* id;
@@ -173,4 +204,6 @@ namespace uba
 	};
 
 	const OwnerInfo& GetOwnerInfo();
+
+	////////////////////////////////////////////////////////////////////////////////////////////////////
 }
\ No newline at end of file
diff --git a/Engine/Source/Programs/UnrealBuildAccelerator/Host/Private/UbaExports.cpp b/Engine/Source/Programs/UnrealBuildAccelerator/Host/Private/UbaExports.cpp
--- a/Engine/Source/Programs/UnrealBuildAccelerator/Host/Private/UbaExports.cpp
+++ b/Engine/Source/Programs/UnrealBuildAccelerator/Host/Private/UbaExports.cpp
@@ -572,7 +572,8 @@ extern "C"
 
 	void SessionServer_UpdateStatus(uba::SessionServer* server, uba::u32 statusRow, uba::u32 statusColumn, const uba::tchar* statusText, uba::LogEntryType statusType, const uba::tchar* statusLink)
 	{
-		server->UpdateStatus(statusRow, statusColumn, statusText, statusType, statusLink);
+		using namespace uba;
+		server->GetTrace().StatusUpdate(statusRow, statusColumn, ToView(statusText), statusType, statusLink ? ToView(statusLink) : StringView());
 	}
 
 	void SessionServer_AddProcessBreadcrumbs(uba::SessionServer* server, uba::u32 processId, const uba::tchar* breadcrumbs, bool deleteOld)
@@ -600,6 +601,16 @@ extern "C"
 			server->SaveSnapshotOfTrace();
 	}
 
+	void SessionServer_AddInfo(uba::SessionServer* server, const uba::tchar* info)
+	{
+		server->GetTrace().SessionInfo(0, uba::ToView(info));
+	}
+
+	uba::Trace* SessionServer_GetTrace(uba::SessionServer* server)
+	{
+		return &server->GetTrace();
+	}
+
 	uba::u32 SessionServer_TaskBegin(uba::SessionServer* server, const uba::tchar* description, const uba::tchar* details, uba::u32 color)
 	{
 		uba::u32 taskId = server->CreateProcessId();
@@ -999,6 +1010,58 @@ extern "C"
 		delete result;
 	}
 
+	uba::Trace* Trace_Create(const uba::tchar* name)
+	{
+		using namespace uba;
+		auto trace = new Trace(g_consoleLogWriter);
+
+		StringBuffer<256> fullName;
+
+		if (name && *name)
+			fullName.Append(name).Append('_');
+
+		time_t rawtime;
+		time(&rawtime);
+		tm ti;
+		localtime_s(&ti, &rawtime);
+		fullName.Appendf(TC("%02i%02i%02i"), ti.tm_hour, ti.tm_min, ti.tm_sec);
+
+		trace->StartWriteAndThread(fullName.data, 64*1024*1024, false);
+		return trace;
+	}
+
+	uba::u32 Trace_TaskBegin(uba::Trace* trace, const uba::tchar* description, const uba::tchar* details, uba::u32 color)
+	{
+		using namespace uba;
+		static Atomic<u32> taskIdCounter = 2'000'000; // Not to collide with process ids
+		u32 taskId = taskIdCounter++;
+		trace->TaskBegin(taskId, ToView(description), ToView(details), color);
+		return taskId;
+	}
+
+	void Trace_TaskHint(uba::Trace* trace, uba::u32 taskId, const uba::tchar* hint)
+	{
+		trace->TaskHint(taskId, uba::ToView(hint));
+	}
+
+	void Trace_TaskEnd(uba::Trace* trace, uba::u32 taskId, bool success)
+	{
+		trace->TaskEnd(taskId, success);
+	}
+
+	void Trace_UpdateStatus(uba::Trace* trace, uba::u32 statusRow, uba::u32 statusColumn, const uba::tchar* statusText, uba::LogEntryType statusType, const uba::tchar* statusLink)
+	{
+		using namespace uba;
+		trace->StatusUpdate(statusRow, statusColumn, ToView(statusText), statusType, statusLink ? ToView(statusLink) : StringView());
+	}
+
+	void Trace_Destroy(uba::Trace* trace, const uba::tchar* writeFile)
+	{
+		trace->StopThread();
+		trace->StopWrite(writeFile);
+		delete trace;
+	}
+
 	void Uba_SetCustomAssertHandler(Uba_CustomAssertHandler* handler)
 	{
 		uba::SetCustomAssertHandler(handler);
diff --git a/Engine/Source/Programs/UnrealBuildAccelerator/Host/Public/UbaExports.h b/Engine/Source/Programs/UnrealBuildAccelerator/Host/Public/UbaExports.h
--- a/Engine/Source/Programs/UnrealBuildAccelerator/Host/Public/UbaExports.h
+++ b/Engine/Source/Programs/UnrealBuildAccelerator/Host/Public/UbaExports.h
@@ -17,6 +17,7 @@ namespace uba
 	class SessionServer;
 	class Scheduler;
 	class StorageServer;
+	class Trace;
 	enum ProcessExitedResponse : uba::u8;
 	struct CacheResult;
 	struct ProcessStartInfo;
@@ -144,6 +145,8 @@ extern "C"
 	UBA_API void SessionServer_RegisterCustomService(uba::SessionServer* server, SessionServer_CustomServiceFunction* function, void* userData = nullptr);
 	UBA_API void SessionServer_RegisterCrossArchitectureMapping(uba::SessionServer* server, const uba::tchar* from, const uba::tchar* to);
 	UBA_API void SessionServer_SaveSnapshotOfTrace(uba::SessionServer* server);
+	UBA_API void SessionServer_AddInfo(uba::SessionServer* server, const uba::tchar* info);
+	UBA_API uba::Trace* SessionServer_GetTrace(uba::SessionServer* server);
 	UBA_API uba::u32 SessionServer_TaskBegin(uba::SessionServer* server, const uba::tchar* description, const uba::tchar* details, uba::u32 color);
 	UBA_API void SessionServer_TaskHint(uba::SessionServer* server, uba::u32 taskId, const uba::tchar* hint);
 	UBA_API void SessionServer_TaskEnd(uba::SessionServer* server, uba::u32 taskId, bool success);
@@ -180,6 +183,14 @@ extern "C"
 	UBA_API uba::u32 CacheResult_GetLogLineType(uba::CacheResult* result, uba::u32 index);
 	UBA_API void CacheResult_Delete(uba::CacheResult* result);
 
+	// Trace (automatically created by session if not provided)
+	UBA_API uba::Trace* Trace_Create(const uba::tchar* name);
+	UBA_API uba::u32 Trace_TaskBegin(uba::Trace* trace, const uba::tchar* description, const uba::tchar* details, uba::u32 color);
+	UBA_API void Trace_TaskHint(uba::Trace* trace, uba::u32 taskId, const uba::tchar* hint);
+	UBA_API void Trace_TaskEnd(uba::Trace* trace, uba::u32 taskId, bool success);
+	UBA_API void Trace_UpdateStatus(uba::Trace* trace, uba::u32 statusRow, uba::u32 statusColumn, const uba::tchar* statusText, uba::LogEntryType statusType, const uba::tchar* statusLink);
+	UBA_API void Trace_Destroy(uba::Trace* trace, const uba::tchar* writeFile);
+
 	// Misc
 	using Uba_CustomAssertHandler = void(const uba::tchar* text);
 	UBA_API void Uba_SetCustomAssertHandler(Uba_CustomAssertHandler* handler);
diff --git a/Engine/Source/Programs/UnrealBuildAccelerator/Visualizer/Private/UbaTraceReader.cpp b/Engine/Source/Programs/UnrealBuildAccelerator/Visualizer/Private/UbaTraceReader.cpp
--- a/Engine/Source/Programs/UnrealBuildAccelerator/Visualizer/Private/UbaTraceReader.cpp
+++ b/Engine/Source/Programs/UnrealBuildAccelerator/Visualizer/Private/UbaTraceReader.cpp
@@ -30,6 +30,30 @@ namespace uba
 		return emptySession;
 	}
 
+	void TraceView::Session::GetFullName(StringBufferBase& out)
+	{
+		out.Append(name);
+		out.Append(TCV(" ("));
+		u32 infoIndex = 0;
+		for (auto& info : infos)
+		{
+			if (info.text.empty())
+				continue;
+			if (!infoIndex)
+				out.Append(TCV(" ("));
+			else
+				out.Append(' ');
+			out.Append(info.text);
+			++infoIndex;
+		}
+		#if UBA_DEBUG
+		out.Append(' ').Append(machineId);
+		#endif
+
+		if (infoIndex)
+			out.Append(TCV(")"));
+	}
+
 	void TraceView::Clear()
 	{
 		version = 0;
@@ -277,11 +301,14 @@ namespace uba
 			return true;
 		m_memoryEnd = m_memoryBegin + *(u32*)m_memoryBegin;
 		outChanged = !m_activeProcesses.empty() || m_memoryPos != m_memoryEnd;
-		bool res = ReadMemory(out, true, maxTime);
+		bool keepRunning = ReadMemory(out, true, maxTime);
 		if (m_hostProcess && WaitForSingleObject(m_hostProcess, 0) != WAIT_TIMEOUT)
+		{
 			StopAllActive(out, GetTime() - out.realStartTime);
+			keepRunning = false;
+		}
 
-		if (res || m_namedTrace.empty())
+		if (keepRunning || m_namedTrace.empty())
 			return true;
 
 		// Move memory to local mapping in case we want to replay
@@ -301,7 +328,7 @@ namespace uba
 		m_memoryBegin = memoryBegin;
 		out.finished = true;
 
-		return res;
+		return keepRunning;
 	}
 
 	bool TraceReader::ReadMemory(TraceView& out, bool trackHost, u64 maxTime)
@@ -402,42 +429,23 @@ namespace uba
 			StringBuffer<128> sessionName;
 			if (!reader.TryReadString(sessionName))
 				return false;
-			StringBuffer<> sessionInfo;
-			if (!reader.TryReadString(sessionInfo))
+			StringBuffer<> machineId;
+			if (!reader.TryReadString(machineId))
 				return false;
 			Guid clientUid = ReadClientId(out, reader);
 			u32 sessionIndex = reader.ReadU32();
-			TString hyperlink;
 
-			const tchar* hyperlinkPos = nullptr;
-			if (!sessionInfo.Contains(TC("http://"), true, &hyperlinkPos))
-				sessionInfo.Contains(TC("https://"), true, &hyperlinkPos);
-			if (hyperlinkPos)
-			{
-				const tchar* hyperlinkEnd = sessionInfo.data + sessionInfo.count;
-				if (const tchar* space = TStrchr(hyperlinkPos, ' '))
-					hyperlinkEnd = space;
-
-				hyperlink.assign(hyperlinkPos, hyperlinkEnd);
-				if (*hyperlinkEnd == ' ')
-					++hyperlinkEnd;
-				while (hyperlinkPos[-1] == ' ')
-					--hyperlinkPos;
-				if (hyperlinkPos[-1] == ',')
-					--hyperlinkPos;
-				memmove(const_cast<tchar*>(hyperlinkPos), hyperlinkEnd, (sessionInfo.data + sessionInfo.count - hyperlinkEnd)*sizeof(tchar));
-				sessionInfo.Resize(sessionInfo.count - (hyperlinkEnd - hyperlinkPos));
-			}
+			TraceView::SessionInfo info;
 
-			StringBuffer<> fullName;
-			fullName.Append(sessionName).Append(TCV(" (")).Append(sessionInfo).Append(TCV(")"));
+			if (out.version < 48)
+				info = ReadSessionInfo(machineId);
 
 			// Check if we can re-use existing session (same machine was disconnected and then reconnected)
 			u32 virtualSessionIndex = sessionIndex;
 			for (u32 i = 0; i != out.sessions.size(); ++i)
 			{
 				auto& oldSession = out.sessions[i];
-				if (oldSession.fullName != fullName.data)
+				if (oldSession.machineId != machineId.data)
 					continue;
 				if (oldSession.disconnectTime == ~u64(0))
 					break;
@@ -455,6 +463,7 @@ namespace uba
 				oldSession.proxyName.clear();
 				oldSession.proxyCreated = false;
 				oldSession.notification.clear();
+				oldSession.infos.clear();
 				//oldSession.fetchedFiles.clear();
 				//oldSession.storedFiles.clear();
 				virtualSessionIndex = i;
@@ -472,13 +481,24 @@ namespace uba
 
 			TraceView::Session& session = GetSession(out, sessionIndex);
 			session.name = sessionName.data;
-			session.fullName = fullName.data;
-			session.hyperlink = hyperlink;
+			session.machineId = machineId.data;
+			if (!info.text.empty())
+				session.infos.push_back(std::move(info));
 			session.clientUid = clientUid;
 
 			++out.activeSessionCount;
 			break;
 		}
+		case TraceType_SessionInfo:
+		{
+			u32 sessionIndex = reader.ReadU32();
+			StringBuffer<> info;
+			reader.ReadString(info);
+			TraceView::Session& session = GetSession(out, sessionIndex);
+			session.infos.push_back(ReadSessionInfo(info));
+			break;
+		}
+
 		case TraceType_SessionUpdate:
 		{
 			if (out.version < 15)
@@ -666,7 +686,7 @@ namespace uba
 				}
 			}
 
-			ProcessBegin(out, sessionIndex, id, time, std::move(desc), std::move(breadcrumbs));
+			ProcessBegin(out, sessionIndex, id, time, std::move(desc), std::move(breadcrumbs), TraceView::ProcessType_Normal);
 			break;
 		}
 		case TraceType_ProcessExited:
@@ -1155,7 +1175,7 @@ namespace uba
 			TString desc = reader.ReadString();
 			TString details = reader.ReadString();
 			u32 color =  reader.ReadU32();(void)color;
-			ProcessBegin(out, 0, id, time, std::move(desc), std::move(details))->type = TraceView::ProcessType_Task;
+			ProcessBegin(out, 0, id, time, std::move(desc), std::move(details), TraceView::ProcessType_Task);
 			break;
 		}
 		case TraceType_TaskHint:
@@ -1296,7 +1316,7 @@ namespace uba
 		{
 			u32 id = u32(reader.Read7BitEncoded());
 			TString desc = reader.ReadString();
-			ProcessBegin(out, 0, id, time, std::move(desc), {})->type = TraceView::ProcessType_CacheFetchTest;
+			ProcessBegin(out, 0, id, time, std::move(desc), {}, TraceView::ProcessType_CacheFetchTest);
 			break;
 		}
 		case TraceType_CacheEndFetch:
@@ -1511,7 +1531,7 @@ namespace uba
 		return nullptr;
 	}
 
-	TraceView::Process* TraceReader::ProcessBegin(TraceView& out, u32 sessionIndex, u32 id, u64 time, TString&& description, TString&& breadcrumbs)
+	TraceView::Process* TraceReader::ProcessBegin(TraceView& out, u32 sessionIndex, u32 id, u64 time, TString&& description, TString&& breadcrumbs, TraceView::ProcessType type)
 	{
 		TraceView::Processor* processor = nullptr;
 		TraceView::Session& session = GetSession(out, sessionIndex);
@@ -1542,10 +1562,14 @@ namespace uba
 		out.activeProcessCounts.emplace_back(time, activeCount);
 		out.maxActiveProcessCount = Max(out.maxActiveProcessCount, activeCount);
 
-		++session.processActiveCount;
-		++out.totalProcessActiveCount;
+		if (type == TraceView::ProcessType_Normal)
+		{
+			++session.processActiveCount;
+			++out.totalProcessActiveCount;
+		}
 
 		process.id = id;
+		process.type = type;
 		process.description = std::move(description);
 		process.breadcrumbs = std::move(breadcrumbs);
 		process.start = time;
@@ -1564,14 +1588,17 @@ namespace uba
 		m_activeProcesses.erase(findIt);
 		outSessionIndex = active.sessionIndex;
 		auto& session = GetSession(out, active.sessionIndex);
-		++session.processExitedCount;
-		--session.processActiveCount;
-
-		++out.totalProcessExitedCount;
-		--out.totalProcessActiveCount;
 
 		TraceView::Process& process = session.processors[active.processorIndex].processes[active.processIndex];
 
+		if (process.type == TraceView::ProcessType_Normal)
+		{
+			++session.processExitedCount;
+			--session.processActiveCount;
+			++out.totalProcessExitedCount;
+			--out.totalProcessActiveCount;
+		}
+
 		out.activeProcessCounts.emplace_back(time, u16(m_activeProcesses.size()));
 
 		process.stop = time;
@@ -1579,5 +1606,30 @@ namespace uba
 		return &process;
 	}
 
+	TraceView::SessionInfo TraceReader::ReadSessionInfo(StringBufferBase& info)
+	{
+		TString hyperlink;
+		const tchar* hyperlinkPos = nullptr;
+		if (!info.Contains(TC("http://"), true, &hyperlinkPos))
+			info.Contains(TC("https://"), true, &hyperlinkPos);
+		if (hyperlinkPos)
+		{
+			const tchar* hyperlinkEnd = info.data + info.count;
+			if (const tchar* space = TStrchr(hyperlinkPos, ' '))
+				hyperlinkEnd = space;
+
+			hyperlink.assign(hyperlinkPos, hyperlinkEnd);
+			if (*hyperlinkEnd == ' ')
+				++hyperlinkEnd;
+			while (hyperlinkPos[-1] == ' ')
+				--hyperlinkPos;
+			if (hyperlinkPos[-1] == ',')
+				--hyperlinkPos;
+			memmove(const_cast<tchar*>(hyperlinkPos), hyperlinkEnd, (info.data + info.count - hyperlinkEnd)*sizeof(tchar));
+			info.Resize(info.count - (hyperlinkEnd - hyperlinkPos));
+		}
+		return { info.data, hyperlink };
+	}
+
 #endif // PLATFORM_WINDOWS
-}
\ No newline at end of file
+}
diff --git a/Engine/Source/Programs/UnrealBuildAccelerator/Visualizer/Private/UbaTraceReader.h b/Engine/Source/Programs/UnrealBuildAccelerator/Visualizer/Private/UbaTraceReader.h
--- a/Engine/Source/Programs/UnrealBuildAccelerator/Visualizer/Private/UbaTraceReader.h
+++ b/Engine/Source/Programs/UnrealBuildAccelerator/Visualizer/Private/UbaTraceReader.h
@@ -97,11 +97,17 @@ namespace uba
 			Vector<u64> writeBytes;
 		};
 
+		struct SessionInfo
+		{
+			TString text;
+			TString hyperlink;
+		};
+
 		struct Session
 		{
 			TString name;
-			TString fullName;
-			TString hyperlink;
+			TString machineId;
+			List<SessionInfo> infos;
 			Guid clientUid;
 			Vector<Processor> processors;
 			Vector<u64> updates;
@@ -141,6 +147,8 @@ namespace uba
 
 			TString proxyName;
 			bool proxyCreated = false;
+
+			void GetFullName(StringBufferBase& out);
 		};
 
 		struct ProcessLocation
@@ -234,9 +242,11 @@ namespace uba
 		TraceView::Session& GetSession(TraceView& out, u32 sessionIndex);
 		TraceView::Session* GetSession(TraceView& out, const Guid& clientUid);
 
-		TraceView::Process* ProcessBegin(TraceView& out, u32 sessionIndex, u32 id, u64 time, TString&& description, TString&& breadcrumbs);
+		TraceView::Process* ProcessBegin(TraceView& out, u32 sessionIndex, u32 id, u64 time, TString&& description, TString&& breadcrumbs, TraceView::ProcessType type);
 		TraceView::Process* ProcessEnd(TraceView& out, u32& outSessionIndex, u32 id, u64 time);
 
+		TraceView::SessionInfo ReadSessionInfo(StringBufferBase& info);
+
 		UnorderedMap<u32, TraceView::ProcessLocation> m_activeProcesses;
 
 		struct WorkRecordLocation { u32 track; u32 index; };
diff --git a/Engine/Source/Programs/UnrealBuildAccelerator/Visualizer/Private/UbaVisualizer.cpp b/Engine/Source/Programs/UnrealBuildAccelerator/Visualizer/Private/UbaVisualizer.cpp
--- a/Engine/Source/Programs/UnrealBuildAccelerator/Visualizer/Private/UbaVisualizer.cpp
+++ b/Engine/Source/Programs/UnrealBuildAccelerator/Visualizer/Private/UbaVisualizer.cpp
@@ -1372,7 +1372,7 @@ namespace uba
 					LineTo(hdc, clientRect.right, posY);
 
 					StringBuffer<> text;
-					text.Append(session.fullName);
+					session.GetFullName(text);
 
 					if (hasUpdates && session.disconnectTime == ~u64(0))
 					{
@@ -2863,9 +2863,12 @@ namespace uba
 		}
 		else
 		{
-			drawText(L"Finished Processes: %u (local: %u)", m_traceView.totalProcessExitedCount, session.processExitedCount);
-			drawText(L"Active Processes: %u (local: %u)", m_traceView.totalProcessActiveCount, session.processActiveCount);
-			drawText(L"Active Helpers: %u", Max(1u, m_traceView.activeSessionCount) - 1);
+			if (m_traceView.totalProcessActiveCount || m_traceView.totalProcessExitedCount)
+			{
+				drawText(L"Finished Processes: %u (local: %u)", m_traceView.totalProcessExitedCount, session.processExitedCount);
+				drawText(L"Active Processes: %u (local: %u)", m_traceView.totalProcessActiveCount, session.processActiveCount);
+				drawText(L"Active Helpers: %u", Max(1u, m_traceView.activeSessionCount) - 1);
+			}
 
 			if (session.highestSendPerS || session.highestRecvPerS)
 			{
@@ -4077,9 +4080,16 @@ namespace uba
 			{
 				ShellExecuteW(NULL, L"open", m_hyperLinkSelected.c_str(), NULL, NULL, SW_SHOW);
 			}
-			else if (m_sessionSelectedIndex != ~0u && !m_traceView.sessions[m_sessionSelectedIndex].hyperlink.empty())
+			else if (m_sessionSelectedIndex != ~0u)
 			{
-				ShellExecuteW(NULL, L"open", m_traceView.sessions[m_sessionSelectedIndex].hyperlink.c_str(), NULL, NULL, SW_SHOW);
+				auto& session = m_traceView.sessions[m_sessionSelectedIndex];
+				for (auto& info : session.infos)
+				{
+					if (info.hyperlink.empty())
+						continue;
+					ShellExecuteW(NULL, L"open", info.hyperlink.c_str(), NULL, NULL, SW_SHOW);
+					break;
+				}
 			}
 			else
 			{
@@ -4107,8 +4117,17 @@ namespace uba
 				ScreenToClient(hWnd, &pt);
 				useHand = PtInRect(&r, { pt.x, pt.y });
 			}
-			else
-				useHand = !m_hyperLinkSelected.empty() || (m_sessionSelectedIndex != ~0u && !m_traceView.sessions[m_sessionSelectedIndex].hyperlink.empty());
+			else if (!m_hyperLinkSelected.empty())
+			{
+				useHand = true;
+			}
+			else if (m_sessionSelectedIndex != ~0u)
+			{
+				auto& session = m_traceView.sessions[m_sessionSelectedIndex];
+				for (auto& info : session.infos)
+					if (!info.hyperlink.empty())
+						useHand = true;
+			}
 			SetCursor(useHand ? hand : arrow);
 			break;
 		}
@@ -4292,7 +4311,9 @@ namespace uba
 			{
 				TString str;
 				auto& session = m_traceView.sessions[m_sessionSelectedIndex];
-				str.append(session.fullName).append(TC("\n"));
+				StringBuffer<> fullName;
+				session.GetFullName(fullName);
+				str.append(fullName.data).append(TC("\n"));
 				for (auto& line : session.summary)
 					str.append(line).append(TC("\n"));
 				CopyTextToClipboard(str);
