[UBA]
* Fixed so when Session.RemoteTraceEnabled is set, scheduler wait for clients to disconnect

diff --git a/Engine/Source/Programs/UnrealBuildAccelerator/Common/Private/UbaScheduler.cpp b/Engine/Source/Programs/UnrealBuildAccelerator/Common/Private/UbaScheduler.cpp
--- a/Engine/Source/Programs/UnrealBuildAccelerator/Common/Private/UbaScheduler.cpp
+++ b/Engine/Source/Programs/UnrealBuildAccelerator/Common/Private/UbaScheduler.cpp
@@ -245,6 +245,12 @@ namespace uba
 		m_session.WaitOnAllTasks();
 		SkipAllQueued();
 
+		if (m_session.GetRemoteTraceEnabled())
+		{
+			m_session.DisableRemoteExecution();
+			m_session.WaitOnAllClients();
+		}
+
 		if (m_memTrack)
 		{
 			m_memThreadEvent.Set();
