[UBA]
* Changed so GetMachineId is using GetComputerName for macos for now

diff --git a/Engine/Source/Programs/UnrealBuildAccelerator/Common/Private/UbaEnvironment.cpp b/Engine/Source/Programs/UnrealBuildAccelerator/Common/Private/UbaEnvironment.cpp
--- a/Engine/Source/Programs/UnrealBuildAccelerator/Common/Private/UbaEnvironment.cpp
+++ b/Engine/Source/Programs/UnrealBuildAccelerator/Common/Private/UbaEnvironment.cpp
@@ -242,13 +242,15 @@ namespace uba
 				out.Append(str);
 		}
 #else
-		if (FILE* id = popen("ioreg -rd1 -c IOPlatformExpertDevice | awk -F\\\" '/IOPlatformUUID/ { print $4 }'", "r"))
-		{
-			success = fgets(str, sizeof(str), id) != nullptr;
-			fclose(id);
-			if (success)
-				out.Append(str, strcspn(str, "\n"));
-		}
+		(void)str;
+		success = GetComputerNameW(out); // TODO: Revisit below
+		//if (FILE* id = popen("ioreg -rd1 -c IOPlatformExpertDevice | awk -F\\\" '/IOPlatformUUID/ { print $4 }'", "r"))
+		//{
+		//	success = fgets(str, sizeof(str), id) != nullptr;
+		//	fclose(id);
+		//	if (success)
+		//		out.Append(str, strcspn(str, "\n"));
+		//}
 #endif
 		if (!success)
 			out.Append(TCV("NoMachineId"));
