[UBA]
* Improved logging around killing processes

diff --git a/Engine/Source/Programs/UnrealBuildAccelerator/Common/Private/UbaScheduler.cpp b/Engine/Source/Programs/UnrealBuildAccelerator/Common/Private/UbaScheduler.cpp
--- a/Engine/Source/Programs/UnrealBuildAccelerator/Common/Private/UbaScheduler.cpp
+++ b/Engine/Source/Programs/UnrealBuildAccelerator/Common/Private/UbaScheduler.cpp
@@ -513,8 +513,13 @@ namespace uba
 		}
 
 		if (isLocal)
-			if (!process.WaitForExit(10*1000))
-				return m_session.GetLogger().Error(TC("Took more than 10 seconds to wait for process to exit (%s)"), si.GetDescription());
+		{
+			if (!process.WaitForExit(8*60*60*1000)) // 8 hours, hopefully someone will come and get me so I can debug :)
+			{
+				m_updateThreadWakeup.Set();
+				return m_session.GetLogger().Error(TC("Took more than 8 hours to wait for process to exit (%s)"), si.GetDescription());
+			}
+		}
 
 		SCOPED_FUTEX(m_processEntriesLock, lock);
 
@@ -1137,7 +1142,7 @@ namespace uba
 	}
 
 
-	void Scheduler::KillNewestLocalProcess()
+	void Scheduler::KillNewestLocalProcess(StringView reason)
 	{
 		ProcessHandle ph = m_session.GetNewestLocalProcess();
 		if (!ph.IsValid())
@@ -1147,7 +1152,7 @@ namespace uba
 
 		u64 startTime = GetTime();
 		if (ProcessReturned(*ph.m_process, true))
-			m_session.GetLogger().Info(TC("Killed process %s (%s) due to memory pressure"), ph.GetStartInfo().GetDescription(), TimeToText(GetTime() - startTime).str);
+			m_session.GetLogger().Info(TC("Killed process %s (%s) - %s"), ph.GetStartInfo().GetDescription(), TimeToText(GetTime() - startTime).str, reason.data);
 	}
 
 	void Scheduler::NativeProcessCreated(ProcessImpl& process)
@@ -1157,6 +1162,22 @@ namespace uba
 
 		auto& impl = (ProcessImpl&)process;
 		auto ei = (ExitProcessInfo*)impl.m_startInfo.userData;
+		
+		// Child process, we still need to track
+		if (!ei)
+		{
+			impl.m_startInfo.userData = this;
+			impl.m_startInfo.exitedFunc = [](void* userData, const ProcessHandle& handle, ProcessExitedResponse& response)
+				{
+					auto& impl = *(ProcessImpl*)handle.m_process;
+					auto& scheduler = *(Scheduler*)userData;
+					StringKey key(impl.m_nativeProcessId, impl.m_nativeProcessCreationTime);
+					SCOPED_FUTEX(scheduler.m_processesLock, l);
+					bool res = scheduler.m_processes.erase(key) == 1;
+					UBA_ASSERT(res);(void)res;
+				};
+		}
+
 		StringKey key(impl.m_nativeProcessId, impl.m_nativeProcessCreationTime);
 		SCOPED_FUTEX(m_processesLock, l);
 		m_processes.try_emplace(key, ProcessRecord{ process, ei });
@@ -1403,7 +1424,7 @@ namespace uba
 			m_updateThreadWakeup.Set();
 			return;
 		}
-		u64 physMem = memTotal - maxPageFile;
+		u64 memPhys = memTotal - maxPageFile;
 
 		UnorderedMap<u32, u32> memoryGroupToIndex;
 
@@ -1411,11 +1432,13 @@ namespace uba
 		
 
 		u8 memStartWaitPercent = m_memStartWaitPercent;
-		u8 memStartWaitPercentMin = u8(double(physMem)/double(memTotal) * 100.0f);
+		u8 memStartWaitPercentMin = u8(double(memPhys)/double(memTotal) * 100.0f);
 		u64 lastPageFileWriteTimeMs = 0;
 
 		u32 timeoutMs = 0;
 
+		StringBuffer<> str;
+
 		while (!m_memThreadEvent.IsSet(timeoutMs))
 		{
 			timeoutMs = 1000;
@@ -1457,8 +1480,11 @@ namespace uba
 				break;
 			}
 
+			u32 activeProcessCount;
 			{
 				SCOPED_FUTEX(m_processesLock, l);
+				activeProcessCount = u32(m_processes.size());
+
 				while (nextEntryOffset)
 				{
 					u64 createTime = *(u64*)(it + CreateTimeOffset);
@@ -1502,14 +1528,12 @@ namespace uba
 						if (res.second)
 						{
 							index = StartRow + u32(memoryGroupToIndex.size() - 1);
-							StringBuffer<> b;
-							b.AppendValue(kv.first);
-							trace.StatusUpdate(index, colIndex, b, LogEntryType_Info);
+							str.Clear().AppendValue(kv.first);
+							trace.StatusUpdate(index, colIndex, str, LogEntryType_Info);
 						}
 
-						StringBuffer<> b;
-						b.Append(BytesToText(stats.average).str);
-						trace.StatusUpdate(index, colIndex + 5, b, LogEntryType_Info);
+						str.Clear().Append(BytesToText(stats.average).str);
+						trace.StatusUpdate(index, colIndex + 5, str, LogEntryType_Info);
 					}
 				}
 				estimatedMemoryUsage = memBaseLine + untrackedMemory + estimatedProcessMemoryUsage;
@@ -1532,7 +1556,8 @@ namespace uba
 			{
 				while (memUsed > memKillThreshold)
 				{
-					KillNewestLocalProcess();
+					str.Clear().Appendf(TC("Low on memory (%s/%s). Kill threshold is %s"), BytesToText(memUsed).str, BytesToText(memTotal).str, BytesToText(memKillThreshold).str);
+					KillNewestLocalProcess(str);
 					if (!GetMemoryInfo(logger, memAvail, memTotal, &maxPageFile))
 						break;
 					memUsed = memTotal - memAvail;
@@ -1552,10 +1577,11 @@ namespace uba
 					trace.StatusUpdate(StartRow+7, 6, CountToText(writes), LogEntryType_Info);
 				}
 
-				if (memUsed > physMem && writes > 1000)
+				if (memUsed > memPhys && activeProcessCount > 6 && writes > 1000)
 				{
 					memStartWaitPercent = u8(double(memUsed)/double(memTotal) * 100.0f);
-					KillNewestLocalProcess();
+					str.Clear().Appendf(TC("Pagefile activity high (%lu writes). MemUsed %s, MemPhys %s, MemTot %s"), writes, BytesToText(memUsed).str, BytesToText(memPhys).str, BytesToText(memTotal).str);
+					KillNewestLocalProcess(str);
 					lastPageFileWriteTimeMs = TimeToMs(GetTime());
 				}
 				else if (writes > 0)
diff --git a/Engine/Source/Programs/UnrealBuildAccelerator/Common/Public/UbaScheduler.h b/Engine/Source/Programs/UnrealBuildAccelerator/Common/Public/UbaScheduler.h
--- a/Engine/Source/Programs/UnrealBuildAccelerator/Common/Public/UbaScheduler.h
+++ b/Engine/Source/Programs/UnrealBuildAccelerator/Common/Public/UbaScheduler.h
@@ -144,7 +144,7 @@ namespace uba
 		void UpdateActiveProcessCounter(bool isLocal, int offset);
 		void FinishProcess(const ProcessHandle& handle);
 
-		void KillNewestLocalProcess();
+		void KillNewestLocalProcess(StringView reason);
 		void NativeProcessCreated(ProcessImpl& process);
 		void LocalProcessStart(ProcessStartInfo2& info);
 		bool LocalProcessOkToSpawn();
