UBA: pgocvt needs to use link rules

diff --git a/Engine/Source/Programs/UnrealBuildAccelerator/Core/Private/UbaApplicationRules.cpp b/Engine/Source/Programs/UnrealBuildAccelerator/Core/Private/UbaApplicationRules.cpp
--- a/Engine/Source/Programs/UnrealBuildAccelerator/Core/Private/UbaApplicationRules.cpp
+++ b/Engine/Source/Programs/UnrealBuildAccelerator/Core/Private/UbaApplicationRules.cpp
@@ -634,7 +634,6 @@ namespace uba
 		virtual bool IsOutputFile(StringView fileName, StringView systemTemp) const override
 		{
 			return fileName.EndsWith(TCV(".c.d"))
-				|| fileName.EndsWith(TCV(".h.d"))
 				|| fileName.EndsWith(TCV(".h.d"))
 				|| fileName.EndsWith(TCV(".cc.d"))
 				|| fileName.EndsWith(TCV(".cpp.d"))
@@ -754,6 +753,7 @@ namespace uba
 			{ ApplicationHash(TC("instr-clang-cl.exe")),		new ApplicationRulesClangClExe() },
 			{ ApplicationHash(TC("include-what-you-use.exe")),	new ApplicationRulesIWYU() },
 			{ ApplicationHash(TC("ubatestapp.exe")),			new ApplicationRulesUbaTestApp() },
+			{ ApplicationHash(TC("pgocvt.exe")),				new ApplicationRulesLinkExe() },
 			//{ L"MSBuild.dll"),				new ApplicationRules() },
 			//{ L"BreakpadSymbolEncoder.exe"),	new ApplicationRulesClang() },
 			//{ L"cmd.exe"),		new ApplicationRules() },
