[UBA] Try to open file in WriteCasFileNoCheck() multiple times before giving up in case there is a race condition.

#rnx

diff --git a/Engine/Plugins/UbaController/Source/Private/UbaJobProcessor.cpp b/Engine/Plugins/UbaController/Source/Private/UbaJobProcessor.cpp
--- a/Engine/Plugins/UbaController/Source/Private/UbaJobProcessor.cpp
+++ b/Engine/Plugins/UbaController/Source/Private/UbaJobProcessor.cpp
@@ -355,9 +355,11 @@ void FUbaJobProcessor::StartUba()
 	int32 FolderIndex = 0;
 	while (true)
 	{
-		RootDir = FString::Printf(TEXT("%s/%s/%u"), *FUbaControllerModule::GetTempDir(), TEXT("UbaControllerStorageDir"), FolderIndex);
+		RootDir = FPaths::Combine(*FUbaControllerModule::GetTempDir(), TEXT("UbaControllerStorageDir"), FString::FromInt(FolderIndex));
 		if (Uba_GetExclusiveAccess(TCHAR_TO_UBASTRING(*RootDir)))
+		{
 			break;
+		}
 		++FolderIndex;
 	}
 	IFileManager::Get().MakeDirectory(*RootDir, true);
diff --git a/Engine/Source/Programs/UnrealBuildAccelerator/Common/Private/UbaStorage.cpp b/Engine/Source/Programs/UnrealBuildAccelerator/Common/Private/UbaStorage.cpp
--- a/Engine/Source/Programs/UnrealBuildAccelerator/Common/Private/UbaStorage.cpp
+++ b/Engine/Source/Programs/UnrealBuildAccelerator/Common/Private/UbaStorage.cpp
@@ -362,9 +362,28 @@ namespace uba
 			}
 		}
 
+		// Returns the number of failed open file attempts. If this is equal to maxOpenFileAttempts no file has been opened successfully.
+		auto TryOpenFileSequentialRead = [this](const tchar* fileName, FileHandle& outHandle, int maxOpenFileAttempts) -> int
+			{
+				UBA_ASSERT(maxOpenFileAttempts > 0);
+				int numFailedOpenFileAttempts = 0;
+				while (!OpenFileSequentialRead(m_logger, fileName, outHandle) && ++numFailedOpenFileAttempts < maxOpenFileAttempts)
+				{
+					// Delay next attempt for 100 milliseconds to open the file in case of a race condition
+					Sleep(100);
+				}
+				return numFailedOpenFileAttempts;
+			};
+
 		FileHandle readHandle;
-		if (!OpenFileSequentialRead(m_logger, fileName, readHandle))
-			return m_logger.Error(TC("[WriteCasFileNoCheck] Failed to open file %s for read (%s)"), fileName, LastErrorToText().data);
+		constexpr int maxOpenFileAttempts = 5;
+		if (int numFailedOpenFileAttempts = TryOpenFileSequentialRead(fileName, readHandle, maxOpenFileAttempts))
+		{
+			if (numFailedOpenFileAttempts == maxOpenFileAttempts)
+				return m_logger.Error(TC("[WriteCasFileNoCheck] Failed to open file %s for read (%s)"), fileName, LastErrorToText().data);
+			else
+				m_logger.Warning(TC("[WriteCasFileNoCheck] First attempts (%d) to open file %s for read failed; This may indicate a race condition"), numFailedOpenFileAttempts, fileName);
+		}
 		auto fileGuard = MakeGuard([&](){ CloseFile(fileName, readHandle); });
 
 		u64 readFileSize = 0;
