[UBA]
* Added more information to assert

diff --git a/Engine/Source/Programs/UnrealBuildAccelerator/Detours/Private/UbaDetoursFileMappingTable.cpp b/Engine/Source/Programs/UnrealBuildAccelerator/Detours/Private/UbaDetoursFileMappingTable.cpp
--- a/Engine/Source/Programs/UnrealBuildAccelerator/Detours/Private/UbaDetoursFileMappingTable.cpp
+++ b/Engine/Source/Programs/UnrealBuildAccelerator/Detours/Private/UbaDetoursFileMappingTable.cpp
@@ -31,7 +31,7 @@ namespace uba
 		BinaryReader reader(m_mem, startPosition);
 		while (reader.GetPosition() != tableSize)
 		{
-			UBA_ASSERT(reader.GetPosition() < tableSize);
+			UBA_ASSERTF(reader.GetPosition() < tableSize, TC("Table mismatch. Size is %u and read position is %llu. Started at %u"), tableSize, reader.GetPosition(), startPosition);
 			StringKey g = reader.ReadStringKey();
 			StringBuffer<1024> mappedFileName;
 			reader.ReadString(mappedFileName);
