[UBA]
* Bumped reserve size of trace file to 256mb when using Trace_Create. That is the max size we've seen when building with detailed trace enabled

diff --git a/Engine/Source/Programs/UnrealBuildAccelerator/Host/Private/UbaExports.cpp b/Engine/Source/Programs/UnrealBuildAccelerator/Host/Private/UbaExports.cpp
--- a/Engine/Source/Programs/UnrealBuildAccelerator/Host/Private/UbaExports.cpp
+++ b/Engine/Source/Programs/UnrealBuildAccelerator/Host/Private/UbaExports.cpp
@@ -1026,7 +1026,7 @@ extern "C"
 		localtime_s(&ti, &rawtime);
 		fullName.Appendf(TC("%02i%02i%02i"), ti.tm_hour, ti.tm_min, ti.tm_sec);
 
-		trace->StartWriteAndThread(fullName.data, 64*1024*1024, false);
+		trace->StartWriteAndThread(fullName.data, 256*1024*1024, false);
 		return trace;
 	}
 
