[UBA]
* Fixed so non-detoured processes still get a log file containing the command line if ubalog enabled

diff --git a/Engine/Source/Programs/UnrealBuildAccelerator/Common/Private/UbaLogger.cpp b/Engine/Source/Programs/UnrealBuildAccelerator/Common/Private/UbaLogger.cpp
--- a/Engine/Source/Programs/UnrealBuildAccelerator/Common/Private/UbaLogger.cpp
+++ b/Engine/Source/Programs/UnrealBuildAccelerator/Common/Private/UbaLogger.cpp
@@ -463,7 +463,7 @@ namespace uba
 
 	TimeToText::operator StringView () const { return ToView(str); }
 
-#if UBA_DEBUG_LOGGER
+#if UBA_DEBUG_LOGGER || UBA_DEBUG
 	thread_local u32 t_debugLogScopeCount = 0;
 
 	class DebugLogWriter : public LogWriter
diff --git a/Engine/Source/Programs/UnrealBuildAccelerator/Common/Private/UbaProcess.cpp b/Engine/Source/Programs/UnrealBuildAccelerator/Common/Private/UbaProcess.cpp
--- a/Engine/Source/Programs/UnrealBuildAccelerator/Common/Private/UbaProcess.cpp
+++ b/Engine/Source/Programs/UnrealBuildAccelerator/Common/Private/UbaProcess.cpp
@@ -1721,9 +1721,18 @@ namespace uba
 					if (CreateProcessW(m_realApplication.c_str(), (tchar*)commandLine.c_str(), NULL, NULL, TRUE, creationFlags, environment, workingDir, &si, &processInfo))
 					{
 						DisableProcessPowerThrottling(processInfo.hProcess);
+
+						#if UBA_DEBUG && UBA_DEBUG_LOG_ENABLED
+						if (*m_startInfo.logFile)
+						{
+							Logger* processLogger = StartDebugLogger(logger, m_startInfo.logFile);
+							processLogger->Info(TC("NOT DETOURED"));
+							processLogger->Log(LogEntryType_Info, commandLine.c_str(), u32(commandLine.size()));
+							StopDebugLogger(processLogger);
+						}
+						#endif
 						break;
 					}
-						
 				}
 
 				DWORD error = GetLastError();
diff --git a/Engine/Source/Programs/UnrealBuildAccelerator/Common/Public/UbaLogger.h b/Engine/Source/Programs/UnrealBuildAccelerator/Common/Public/UbaLogger.h
--- a/Engine/Source/Programs/UnrealBuildAccelerator/Common/Public/UbaLogger.h
+++ b/Engine/Source/Programs/UnrealBuildAccelerator/Common/Public/UbaLogger.h
@@ -123,7 +123,7 @@ namespace uba
 
 	////////////////////////////////////////////////////////////////////////////////////////////////////
 
-	#if UBA_DEBUG_LOGGER
+	#if UBA_DEBUG_LOGGER || UBA_DEBUG
 	Logger* StartDebugLogger(Logger& outerLogger, const tchar* fileName);
 	Logger* StopDebugLogger(Logger* logger);
 	#endif
