[UBA]
* Added Trace_SetGlobal that can be used to override all trace creation in sessions and use this global one instead
* Fixed so providing empty trace channel will not use any trace channel at all.
* Added so trace channel can be set from export api to avoid having trace show up in default channel used by visualizer

diff --git a/Engine/Source/Programs/UnrealBuildAccelerator/Common/Private/UbaTrace.cpp b/Engine/Source/Programs/UnrealBuildAccelerator/Common/Private/UbaTrace.cpp
--- a/Engine/Source/Programs/UnrealBuildAccelerator/Common/Private/UbaTrace.cpp
+++ b/Engine/Source/Programs/UnrealBuildAccelerator/Common/Private/UbaTrace.cpp
@@ -58,7 +58,7 @@ namespace uba
 		bool m_isValid = true;
 	};
 
-	bool Trace::StartWriteAndThread(const tchar* namedTrace, u64 traceMemCapacity, bool logStart)
+	bool Trace::StartWriteAndThread(const tchar* namedTrace, u64 traceMemCapacity, bool logStart, const tchar* channelName)
 	{
 		if (logStart)
 		{
@@ -68,7 +68,7 @@ namespace uba
 				m_logger.Info(TC("---- Starting trace ----"));
 		}
 
-		if (!StartWrite(namedTrace, traceMemCapacity))
+		if (!StartWrite(namedTrace, traceMemCapacity, channelName))
 			return false;
 
 		StringBuffer<256> buf;
@@ -103,7 +103,7 @@ namespace uba
 		m_threadIntervalMs = traceIntervalMs;
 	}
 
-	bool Trace::StartWrite(const tchar* namedTrace, u64 traceMemCapacity)
+	bool Trace::StartWrite(const tchar* namedTrace, u64 traceMemCapacity, const tchar* channelName)
 	{
 		m_memoryCapacity = traceMemCapacity;
 		m_memoryHandle = uba::CreateMemoryMappingW(m_logger, PAGE_READWRITE|SEC_RESERVE, m_memoryCapacity, namedTrace, TC("Trace"));
@@ -137,10 +137,11 @@ namespace uba
 		}
 
 
-		if (namedTrace && m_channel.Init())
+		if (namedTrace)
 		{
 			m_namedTrace = namedTrace;
-			m_channel.Write(namedTrace);
+			if (channelName && *channelName && m_channel.Init(channelName))
+				m_channel.Write(namedTrace);
 		}
 		return true;
 	}
@@ -185,7 +186,7 @@ namespace uba
 			return true;
 		auto g = MakeGuard([this]() { FreeMemory(); });
 
-		if (!m_namedTrace.empty())
+		if (m_channel.IsInitialized())
 			m_channel.Write(TC(""), m_namedTrace.c_str());
 
 		{
@@ -708,7 +709,7 @@ namespace uba
 
 		channelMutex.Append(channelName).Append(TCV("Mutex"));
 		m_mutex = CreateMutexW(false, channelMutex.data);
-		if (!m_mutex)
+		if (m_mutex == InvalidMutexHandle)
 			return false;
 
 		mg.Cancel();
@@ -717,17 +718,26 @@ namespace uba
 		return true;
 	}
 
-	TraceChannel::~TraceChannel()
+	void TraceChannel::Deinit()
 	{
 		#if PLATFORM_WINDOWS
 		if (m_mem)
 			UnmapViewOfFile(m_logger, m_mem, 256, TC("TraceChannel"));
 		if (m_memHandle.IsValid())
 			CloseFileMapping(m_logger, m_memHandle, TC("TraceChannel"));
-		CloseMutex(m_mutex);
+		if (m_mutex != InvalidMutexHandle)
+			CloseMutex(m_mutex);
+		m_mem = nullptr;
+		m_memHandle = {};
+		m_mutex = InvalidMutexHandle;
 		#endif
 	}
 
+	TraceChannel::~TraceChannel()
+	{
+		Deinit();
+	}
+
 	bool TraceChannel::Write(const tchar* traceName, const tchar* ifMatching)
 	{
 		#if PLATFORM_WINDOWS
@@ -751,6 +761,11 @@ namespace uba
 		return true;
 	}
 
+	bool TraceChannel::IsInitialized()
+	{
+		return m_mutex != InvalidMutexHandle;
+	}
+
 	static OwnerInfo InternalGetOwnerInfo()
 	{
 		static tchar buffer[260];
diff --git a/Engine/Source/Programs/UnrealBuildAccelerator/Common/Public/UbaTrace.h b/Engine/Source/Programs/UnrealBuildAccelerator/Common/Public/UbaTrace.h
--- a/Engine/Source/Programs/UnrealBuildAccelerator/Common/Public/UbaTrace.h
+++ b/Engine/Source/Programs/UnrealBuildAccelerator/Common/Public/UbaTrace.h
@@ -22,8 +22,10 @@ namespace uba
 		~TraceChannel();
 
 		bool Init(const tchar* channelName = TC("Default"));
+		void Deinit();
 		bool Write(const tchar* traceName, const tchar* ifMatching = nullptr);
 		bool Read(StringBufferBase& outTraceName);
+		bool IsInitialized();
 
 		Logger& m_logger;
 		MutexHandle m_mutex = InvalidMutexHandle;
@@ -90,7 +92,7 @@ namespace uba
 		Trace(LogWriter& logWriter);
 		~Trace();
 
-		bool StartWriteAndThread(const tchar* namedTrace, u64 traceMemCapacity, bool logStart);
+		bool StartWriteAndThread(const tchar* namedTrace, u64 traceMemCapacity, bool logStart, const tchar* channelName = TC("Default"));
 
 		void StartThread();
 		void StopThread();
@@ -98,7 +100,7 @@ namespace uba
 		void SetThreadUpdateCallback(Function<void()>&& callback, u32 traceIntervalMs = 500);
 
 
-		bool StartWrite(const tchar* namedTrace, u64 traceMemCapacity);
+		bool StartWrite(const tchar* namedTrace, u64 traceMemCapacity, const tchar* channelName = TC("Default"));
 
 		// Writes the current state of the trace to the specified output file without stopping anything
 		bool Write(const tchar* writeFileName, bool writeSummary);
@@ -109,6 +111,7 @@ namespace uba
 
 		const TString& GetNamedTrace() { return m_namedTrace; }
 		bool IsWriting() const { return m_memoryBegin != nullptr; }
+		TraceChannel& GetChannel() { return m_channel; }
 
 		void SessionAdded(u32 sessionId, u32 clientId, StringView name, StringView machineId);
 		void SessionInfo(u32 sessionId, StringView info);
diff --git a/Engine/Source/Programs/UnrealBuildAccelerator/Host/Private/UbaExports.cpp b/Engine/Source/Programs/UnrealBuildAccelerator/Host/Private/UbaExports.cpp
--- a/Engine/Source/Programs/UnrealBuildAccelerator/Host/Private/UbaExports.cpp
+++ b/Engine/Source/Programs/UnrealBuildAccelerator/Host/Private/UbaExports.cpp
@@ -114,6 +114,7 @@ namespace uba
 	// Not great with global but 
 	MutexHandle g_exclusiveMutex = InvalidMutexHandle;
 
+	Trace* g_globalTrace;
 }
 
 extern "C"
@@ -443,6 +444,7 @@ extern "C"
 		info->storeIntermediateFilesCompressed = storeIntermediateFilesCompressed;
 		//info->remoteTraceEnabled = true;
 		info->remoteLogEnabled = true;
+		info->trace = uba::g_globalTrace;
 		return info;
 	}
 
@@ -1010,7 +1012,7 @@ extern "C"
 		delete result;
 	}
 
-	uba::Trace* Trace_Create(const uba::tchar* name)
+	uba::Trace* Trace_Create(const uba::tchar* name, const uba::tchar* channel)
 	{
 		using namespace uba;
 		auto trace = new Trace(g_consoleLogWriter);
@@ -1026,7 +1028,7 @@ extern "C"
 		localtime_s(&ti, &rawtime);
 		fullName.Appendf(TC("%02i%02i%02i"), ti.tm_hour, ti.tm_min, ti.tm_sec);
 
-		trace->StartWriteAndThread(fullName.data, 256*1024*1024, false);
+		trace->StartWriteAndThread(fullName.data, 256*1024*1024, false, channel);
 		return trace;
 	}
 
@@ -1057,11 +1059,19 @@ extern "C"
 
 	void Trace_Destroy(uba::Trace* trace, const uba::tchar* writeFile)
 	{
+		if (uba::g_globalTrace == trace)
+			uba::g_globalTrace = nullptr;
+
 		trace->StopThread();
 		trace->StopWrite(writeFile);
 		delete trace;
 	}
 
+	void Trace_SetGlobal(uba::Trace* trace)
+	{
+		uba::g_globalTrace = trace;
+	}
+
 	void Uba_SetCustomAssertHandler(Uba_CustomAssertHandler* handler)
 	{
 		uba::SetCustomAssertHandler(handler);
diff --git a/Engine/Source/Programs/UnrealBuildAccelerator/Host/Public/UbaExports.h b/Engine/Source/Programs/UnrealBuildAccelerator/Host/Public/UbaExports.h
--- a/Engine/Source/Programs/UnrealBuildAccelerator/Host/Public/UbaExports.h
+++ b/Engine/Source/Programs/UnrealBuildAccelerator/Host/Public/UbaExports.h
@@ -184,12 +184,13 @@ extern "C"
 	UBA_API void CacheResult_Delete(uba::CacheResult* result);
 
 	// Trace (automatically created by session if not provided)
-	UBA_API uba::Trace* Trace_Create(const uba::tchar* name);
+	UBA_API uba::Trace* Trace_Create(const uba::tchar* name, const uba::tchar* channel = TC("Default"));
 	UBA_API uba::u32 Trace_TaskBegin(uba::Trace* trace, const uba::tchar* description, const uba::tchar* details, uba::u32 color);
 	UBA_API void Trace_TaskHint(uba::Trace* trace, uba::u32 taskId, const uba::tchar* hint);
 	UBA_API void Trace_TaskEnd(uba::Trace* trace, uba::u32 taskId, bool success);
 	UBA_API void Trace_UpdateStatus(uba::Trace* trace, uba::u32 statusRow, uba::u32 statusColumn, const uba::tchar* statusText, uba::LogEntryType statusType, const uba::tchar* statusLink);
 	UBA_API void Trace_Destroy(uba::Trace* trace, const uba::tchar* writeFile);
+	UBA_API void Trace_SetGlobal(uba::Trace* trace);
 
 	// Misc
 	using Uba_CustomAssertHandler = void(const uba::tchar* text);
