[UBA]
* Always assert if session create fails (shipping assert)
* Changed so scheduler always allow one local process even though memory/pagefile activity is high

diff --git a/Engine/Source/Programs/UnrealBuildAccelerator/Common/Private/UbaScheduler.cpp b/Engine/Source/Programs/UnrealBuildAccelerator/Common/Private/UbaScheduler.cpp
--- a/Engine/Source/Programs/UnrealBuildAccelerator/Common/Private/UbaScheduler.cpp
+++ b/Engine/Source/Programs/UnrealBuildAccelerator/Common/Private/UbaScheduler.cpp
@@ -1221,6 +1221,9 @@ namespace uba
 			return true;
 		}
 
+		if (!m_activeLocalProcesses) // Always let one run so it can finish eventually
+			return true;
+
 		if (m_memDelaySpawning)
 			return false;
 
diff --git a/Engine/Source/Programs/UnrealBuildAccelerator/Common/Private/UbaSessionServer.cpp b/Engine/Source/Programs/UnrealBuildAccelerator/Common/Private/UbaSessionServer.cpp
--- a/Engine/Source/Programs/UnrealBuildAccelerator/Common/Private/UbaSessionServer.cpp
+++ b/Engine/Source/Programs/UnrealBuildAccelerator/Common/Private/UbaSessionServer.cpp
@@ -227,7 +227,7 @@ namespace uba
 
 		if (!Create(info))
 		{
-			UBA_ASSERT(false);
+			UBA_SHIPPING_ASSERT(false);
 			return;
 		}
 	}
