[UBA]
* Fixed extremely rare deadlock that could happen when multiple storage clients uploads the same cas file to server. The one that got in first will start sending the file while the others wait for it to finish. The others have a 12 minutes timeout before they assume something is going wrong and take the activestores lock  to traverse and see. If the first client disconnects at the exact same time (within 100ms), then it will also take the activestores lock and inside the lock take the casentry lock... this would cause an entry inversion of locks and cause a deadlock

diff --git a/Engine/Source/Programs/UnrealBuildAccelerator/Common/Private/UbaStorageServer.cpp b/Engine/Source/Programs/UnrealBuildAccelerator/Common/Private/UbaStorageServer.cpp
--- a/Engine/Source/Programs/UnrealBuildAccelerator/Common/Private/UbaStorageServer.cpp
+++ b/Engine/Source/Programs/UnrealBuildAccelerator/Common/Private/UbaStorageServer.cpp
@@ -280,7 +280,10 @@ namespace uba
 			CasKey key = casEntry.key;
 			entryLock.Leave();
 			Sleep(100);
-			entryLock.Enter();
+			auto g = MakeGuard([&]() { entryLock.Enter(); });
+
+			if (connectionInfo.ShouldDisconnect())
+				return false;
 
 			if (++waitCount < 12*60*10)
 				continue;
