[UBA]
* Added logging if storage fails to update isrunning time stamp

diff --git a/Engine/Source/Programs/UnrealBuildAccelerator/Common/Private/UbaStorage.cpp b/Engine/Source/Programs/UnrealBuildAccelerator/Common/Private/UbaStorage.cpp
--- a/Engine/Source/Programs/UnrealBuildAccelerator/Common/Private/UbaStorage.cpp
+++ b/Engine/Source/Programs/UnrealBuildAccelerator/Common/Private/UbaStorage.cpp
@@ -861,10 +861,12 @@ namespace uba
 		isRunningName.Append(m_rootDir).Append(TCV(".isRunning"));
 		FileHandle fh = uba::CreateFileW(isRunningName.data, FILE_WRITE_ATTRIBUTES, FILE_SHARE_ALL, OPEN_EXISTING, FILE_FLAG_BACKUP_SEMANTICS);
 		if (fh == InvalidFileHandle)
-			return false;
-		SetFileLastWriteTime(fh, GetSystemTimeAsFileTime());
+			return m_logger.Warning(TC("Failed to open %s to update lastwritetime (%s)"), isRunningName.data, LastErrorToText().data);
+		bool success = SetFileLastWriteTime(fh, GetSystemTimeAsFileTime());
+		if (!success)
+			m_logger.Warning(TC("Failed to update lastwritetime on %s (%s)"), isRunningName.data, LastErrorToText().data);
 		uba::CloseFile(isRunningName.data, fh);
-		return true;
+		return success;
 	}
 
 	StorageImpl::StorageImpl(const StorageCreateInfo& info, const tchar* logPrefix)
