[UBA]
* Added so when storage is not gracefully shutdown it will output last save time of .isRunning file when fixing db and touch file

diff --git a/Engine/Source/Programs/UnrealBuildAccelerator/Common/Private/UbaLogger.cpp b/Engine/Source/Programs/UnrealBuildAccelerator/Common/Private/UbaLogger.cpp
--- a/Engine/Source/Programs/UnrealBuildAccelerator/Common/Private/UbaLogger.cpp
+++ b/Engine/Source/Programs/UnrealBuildAccelerator/Common/Private/UbaLogger.cpp
@@ -463,6 +463,40 @@ namespace uba
 
 	TimeToText::operator StringView () const { return ToView(str); }
 
+	#if PLATFORM_WINDOWS
+	#define TStrFtime ::wcsftime
+	#define TLocalTime(a, b) ::localtime_s(a, b) == 0
+	#else
+	#define TStrFtime ::strftime
+	#define TLocalTime(a, b) ::localtime_s(a, b) != 0
+	#endif
+
+	DateToText::DateToText(float seconds)
+	{
+		auto rawTime = (time_t)seconds;
+		struct tm timeinfo;
+		if (TLocalTime(&timeinfo, &rawTime))
+			TStrFtime(str, sizeof_array(str), TC("%Y-%m-%d %H:%M:%S"), &timeinfo);
+		else
+			str[0] = 0;
+	}
+
+	DateToText::DateToText(u64 fileTime)
+	{
+		#if PLATFORM_WINDOWS
+		u64 unixNs = (fileTime - 116444736'000'000'000ull) * 100ull;
+		#else
+		u64 unixNs = fileTime * 100;
+		#endif
+
+		time_t sec = (time_t)(unixNs / 1'000'000'000ull);
+		struct tm timeinfo;
+		if (TLocalTime(&timeinfo, &sec))
+			TStrFtime(str, sizeof_array(str), TC("%Y-%m-%d %H:%M:%S"), &timeinfo);
+		else
+			TStrcpy_s(str, sizeof_array(str), TC("BADTIME"));
+	}
+
 #if UBA_DEBUG_LOGGER || UBA_DEBUG
 	thread_local u32 t_debugLogScopeCount = 0;
 
diff --git a/Engine/Source/Programs/UnrealBuildAccelerator/Common/Private/UbaStorage.cpp b/Engine/Source/Programs/UnrealBuildAccelerator/Common/Private/UbaStorage.cpp
--- a/Engine/Source/Programs/UnrealBuildAccelerator/Common/Private/UbaStorage.cpp
+++ b/Engine/Source/Programs/UnrealBuildAccelerator/Common/Private/UbaStorage.cpp
@@ -992,7 +992,7 @@ namespace uba
 
 		SCOPED_FUTEX(m_casTableLoadSaveLock, loadSaveLock);
 
-		UBA_ASSERT(!m_casTableLoaded);
+		UBA_ASSERTF(!m_casTableLoaded, TC("CasTable already loaded..."));
 		m_casTableLoaded = true;
 		u64 startTime = GetTime();
 		tchar isRunningName[256];
@@ -1143,9 +1143,10 @@ namespace uba
 			FileHandle fh = uba::CreateFileW(isRunningName, 0, FILE_SHARE_ALL, OPEN_EXISTING, FILE_FLAG_BACKUP_SEMANTICS);
 			auto fhg = MakeGuard([&]() { uba::CloseFile(isRunningName, fh); });
 			if (fh != InvalidFileHandle)
-				GetFileLastWriteTime(fileTime, fh);
+				if (!GetFileLastWriteTime(fileTime, fh))
+					m_logger.Warning(TC("Failed to get lastwritetime of %s (%s)"), isRunningName, LastErrorToText().data);
 
-			m_logger.Info(TC("Previous run was not gracefully shutdown. Reparsing cas directory %s to check for added/missing files"), m_rootDir.data);
+			m_logger.Info(TC("Previous run was not gracefully shutdown. Reparsing cas directory %s to check for added/missing/corrupt files (Saved %s)"), m_rootDir.data, DateToText(fileTime).str);
 
 			if (!CheckAllCasFiles(fileTime))
 				return false;
diff --git a/Engine/Source/Programs/UnrealBuildAccelerator/Common/Public/UbaLogger.h b/Engine/Source/Programs/UnrealBuildAccelerator/Common/Public/UbaLogger.h
--- a/Engine/Source/Programs/UnrealBuildAccelerator/Common/Public/UbaLogger.h
+++ b/Engine/Source/Programs/UnrealBuildAccelerator/Common/Public/UbaLogger.h
@@ -121,6 +121,13 @@ namespace uba
 		tchar str[32];
 	};
 
+	struct DateToText
+	{
+		DateToText(float seconds);
+		DateToText(u64 fileTime);
+		tchar str[128];
+	};
+
 	////////////////////////////////////////////////////////////////////////////////////////////////////
 
 	#if UBA_DEBUG_LOGGER || UBA_DEBUG
diff --git a/Engine/Source/Programs/UnrealBuildAccelerator/Common/Public/UbaStorage.h b/Engine/Source/Programs/UnrealBuildAccelerator/Common/Public/UbaStorage.h
--- a/Engine/Source/Programs/UnrealBuildAccelerator/Common/Public/UbaStorage.h
+++ b/Engine/Source/Programs/UnrealBuildAccelerator/Common/Public/UbaStorage.h
@@ -149,6 +149,7 @@ namespace uba
 		virtual bool GetZone(StringBufferBase& out) override;
 		virtual bool Reset() override;
 		bool DeleteAllCas();
+		bool IsCasTableLoaded() const { return m_casTableLoaded; }
 
 		virtual bool StoreCompressed() const final { return m_storeCompressed; }
 		virtual void PrintSummary(Logger& logger) override;
diff --git a/Engine/Source/Programs/UnrealBuildAccelerator/Visualizer/Private/UbaVisualizer.cpp b/Engine/Source/Programs/UnrealBuildAccelerator/Visualizer/Private/UbaVisualizer.cpp
--- a/Engine/Source/Programs/UnrealBuildAccelerator/Visualizer/Private/UbaVisualizer.cpp
+++ b/Engine/Source/Programs/UnrealBuildAccelerator/Visualizer/Private/UbaVisualizer.cpp
@@ -3615,13 +3615,7 @@ namespace uba
 
 	StringBuffer<128> Visualizer::GetWorldTime(float seconds)
 	{
-		time_t rawTime = m_traceView.traceSystemStartTimeUs / 1'000'000;
-		rawTime += (time_t)seconds;
-		struct tm timeinfo;
-		localtime_s(&timeinfo, &rawTime);
-		StringBuffer<128> buffer;
-		buffer.count = u32(wcsftime(buffer.data, buffer.capacity, L"%Y-%m-%d %H:%M:%S", &timeinfo));
-		return buffer;
+		return StringBuffer<128>(DateToText(float(m_traceView.traceSystemStartTimeUs / 1'000'000) + seconds).str);
 	}
 
 	void Visualizer::PostNewTrace(u32 replay, bool paused)
