[UBA]
* Fixed extremely rare issue when an explicit disconnect call happens at the same time as a connection is being made and code is in a very small window of code. This would cause the shutdown to not call ondisconnected callbacks and also report an error

diff --git a/Engine/Source/Programs/UnrealBuildAccelerator/Common/Private/UbaNetworkClient.cpp b/Engine/Source/Programs/UnrealBuildAccelerator/Common/Private/UbaNetworkClient.cpp
--- a/Engine/Source/Programs/UnrealBuildAccelerator/Common/Private/UbaNetworkClient.cpp
+++ b/Engine/Source/Programs/UnrealBuildAccelerator/Common/Private/UbaNetworkClient.cpp
@@ -438,7 +438,16 @@ namespace uba
 			c.disconnectedEvent.IsSet(~0u);
 
 		if (m_connectionCount != 0)
-			m_logger.Error(TC("Should be zero connections here. %llu still active"), m_connectionCount.load());
+		{
+			// I think this could happen if m_isDisconnecting is set while a connection is in progress.
+			// it might early out in ConnectedCallback which will prevent connection->connected from being set to 1
+			// which means that NetworkClient::OnDisconnected will not enter the code that decrement m_connectionCount
+			m_logger.Info(TC("Should be zero connections here. %llu still active (was Disconnect called while a connect was being made?)"), m_connectionCount.load());
+
+			SCOPED_READ_LOCK(m_onDisconnectedFunctionsLock, lock);
+			for (auto& f : m_onDisconnectedFunctions)
+				f();
+		}
 
 		{
 			SCOPED_WRITE_LOCK(m_connectionsLock, lock2);
