[UBA]
* Fixed error happening when connecting to multiple cache servers

diff --git a/Engine/Source/Programs/UnrealBuildAccelerator/Host/Private/UbaExports.cpp b/Engine/Source/Programs/UnrealBuildAccelerator/Host/Private/UbaExports.cpp
--- a/Engine/Source/Programs/UnrealBuildAccelerator/Host/Private/UbaExports.cpp
+++ b/Engine/Source/Programs/UnrealBuildAccelerator/Host/Private/UbaExports.cpp
@@ -885,7 +885,9 @@ extern "C"
 				outConnectionCount = nc->GetConnectionCount();
 			});
 
-		cacheClient->GetStorage().LoadCasTable();
+		auto& storage = cacheClient->GetStorage();
+		if (!storage.IsCasTableLoaded())
+			storage.LoadCasTable();
 		return true;
 	}
 	
