[UBA]
* Added icx.exe application rules
* Changed so paths like [REDACTED_HANDLE]:\blah  is treated as absolute paths. Code previously checked for ':' at index 1.. and in this case it was an invalid path ([at] was in the beginning) but still treated as absolute..
* Modified some logging to make it easier to debug
* Changed so system temp folder itself is not a memory file (this logic is a bit shaky.. should probably be revisited at some point
* Added asserts for when TEMP env variable is not found

diff --git a/Engine/Source/Programs/UnrealBuildAccelerator/Core/Private/UbaApplicationRules.cpp b/Engine/Source/Programs/UnrealBuildAccelerator/Core/Private/UbaApplicationRules.cpp
--- a/Engine/Source/Programs/UnrealBuildAccelerator/Core/Private/UbaApplicationRules.cpp
+++ b/Engine/Source/Programs/UnrealBuildAccelerator/Core/Private/UbaApplicationRules.cpp
@@ -372,7 +372,7 @@ namespace uba
 
 		virtual bool KeepInMemory(StringView fileName, StringView systemTemp, bool isRunningRemote, bool isWrite) const override
 		{
-			return fileName.StartsWith(systemTemp);
+			return fileName.count > systemTemp.count && fileName.StartsWith(systemTemp);
 		}
 	};
 
@@ -697,6 +697,15 @@ namespace uba
 		}
 	};
 
+	class ApplicationRulesIcxExe : public ApplicationRulesClangClExe
+	{
+		using Super = ApplicationRulesClangClExe;
+		virtual bool IsOutputFile(StringView fileName, StringView systemTemp) const override
+		{
+			return Super::IsOutputFile(fileName, systemTemp) || fileName.count > systemTemp.count && fileName.StartsWith(systemTemp);
+		}
+	};
+
 	template<typename T, size_t size>
 	constexpr u32 ApplicationHash(T (&buffer)[size])
 	{
@@ -754,6 +763,7 @@ namespace uba
 			{ ApplicationHash(TC("include-what-you-use.exe")),	new ApplicationRulesIWYU() },
 			{ ApplicationHash(TC("ubatestapp.exe")),			new ApplicationRulesUbaTestApp() },
 			{ ApplicationHash(TC("pgocvt.exe")),				new ApplicationRulesLinkExe() },
+			{ ApplicationHash(TC("icx.exe")),					new ApplicationRulesIcxExe() },
 			//{ L"MSBuild.dll"),				new ApplicationRules() },
 			//{ L"BreakpadSymbolEncoder.exe"),	new ApplicationRulesClang() },
 			//{ L"cmd.exe"),		new ApplicationRules() },
diff --git a/Engine/Source/Programs/UnrealBuildAccelerator/Core/Public/UbaPathUtils.h b/Engine/Source/Programs/UnrealBuildAccelerator/Core/Public/UbaPathUtils.h
--- a/Engine/Source/Programs/UnrealBuildAccelerator/Core/Public/UbaPathUtils.h
+++ b/Engine/Source/Programs/UnrealBuildAccelerator/Core/Public/UbaPathUtils.h
@@ -99,6 +99,8 @@ namespace uba
 			}
 		}
 
+		bool isAbsolute = false;
+
 		while (tchar c = *read++)
 		{
 			if (c == '/')
@@ -110,6 +112,8 @@ namespace uba
 			}
 			else if (c == '.' && lastChar == '.')
 				hasDotDot = true;
+			else if (c == ':')
+				isAbsolute = true;
 			if (c == '\\')
 			{
 				if (lastChar == '.')
@@ -147,7 +151,7 @@ namespace uba
 		{
 			startsWithDoubleBackslash = true;
 		}
-		else if (buffer[1] != ':') // If not absolute, add current dir
+		else if (!isAbsolute) // If not absolute, add current dir
 		{
 			tchar* copyFrom = buffer;
 			if (copyFrom[0] == '\\')
diff --git a/Engine/Source/Programs/UnrealBuildAccelerator/Detours/Private/UbaDetoursShared.cpp b/Engine/Source/Programs/UnrealBuildAccelerator/Detours/Private/UbaDetoursShared.cpp
--- a/Engine/Source/Programs/UnrealBuildAccelerator/Detours/Private/UbaDetoursShared.cpp
+++ b/Engine/Source/Programs/UnrealBuildAccelerator/Detours/Private/UbaDetoursShared.cpp
@@ -32,6 +32,8 @@ namespace uba
 	bool g_allowKeepFilesInMemory = IsWindows;
 	bool g_allowOutputFiles = IsWindows;
 	bool g_suppressLogging = false;
+	u32 g_vfsEntryCount;
+	u32 g_vfsMatchingLength;
 
 	void InitSharedVariables()
 	{
@@ -48,6 +50,9 @@ namespace uba
 		g_systemRootMem.Create();
 		g_systemTempMem.Create();
 
+		g_vfsEntryCount = 0;
+		g_vfsMatchingLength = 0;
+
 		u64 reserveSizeMb = IsWindows ? 256 : 1024; // The sync primitives on linux/macos is much bigger
 		g_memoryBlockMem.Create(reserveSizeMb * 1024 * 1024);
 		g_directoryTableMem.Create(g_memoryBlock);
@@ -176,8 +181,6 @@ namespace uba
 
 	struct VfsEntry { StringView vfs; StringView local; VfsEntry() : vfs(NoInit), local(NoInit) {}; };
 	VfsEntry g_vfsEntries[32];
-	u32 g_vfsEntryCount;
-	u32 g_vfsMatchingLength;
 
 	void PopulateVfs(BinaryReader& vfsReader)
 	{
diff --git a/Engine/Source/Programs/UnrealBuildAccelerator/Detours/Private/Windows/UbaDetoursFunctionsNtDll.inl b/Engine/Source/Programs/UnrealBuildAccelerator/Detours/Private/Windows/UbaDetoursFunctionsNtDll.inl
--- a/Engine/Source/Programs/UnrealBuildAccelerator/Detours/Private/Windows/UbaDetoursFunctionsNtDll.inl
+++ b/Engine/Source/Programs/UnrealBuildAccelerator/Detours/Private/Windows/UbaDetoursFunctionsNtDll.inl
@@ -1288,7 +1288,7 @@ NTSTATUS NTAPI Shared_NtCreateFile(bool IsCreateFunc, PHANDLE hFileHandle, ACCES
 				if (!isWrite)
 				{
 					*hFileHandle = INVALID_HANDLE_VALUE;
-					DEBUG_LOG_DETOURED(funcName, L"NOTEXISTS (%ls) -> Error", fileName.data);
+					DEBUG_LOG_DETOURED(funcName, L"NOTEXISTS1 (%ls) -> Error", fileName.data);
 					return STATUS_OBJECT_NAME_NOT_FOUND;
 				}
 
@@ -1305,7 +1305,7 @@ NTSTATUS NTAPI Shared_NtCreateFile(bool IsCreateFunc, PHANDLE hFileHandle, ACCES
 				if (res2 != STATUS_SUCCESS)
 				{
 					*hFileHandle = INVALID_HANDLE_VALUE;
-					DEBUG_LOG_DETOURED(funcName, L"NOTEXISTS (%ls) -> Error", fileName.data);
+					DEBUG_LOG_DETOURED(funcName, L"NOTEXISTS2 (%ls) -> Error", fileName.data);
 					return res2;
 				}
 
diff --git a/Engine/Source/Programs/UnrealBuildAccelerator/Detours/Private/Windows/UbaDetoursFunctionsWin.cpp b/Engine/Source/Programs/UnrealBuildAccelerator/Detours/Private/Windows/UbaDetoursFunctionsWin.cpp
--- a/Engine/Source/Programs/UnrealBuildAccelerator/Detours/Private/Windows/UbaDetoursFunctionsWin.cpp
+++ b/Engine/Source/Programs/UnrealBuildAccelerator/Detours/Private/Windows/UbaDetoursFunctionsWin.cpp
@@ -1047,7 +1047,10 @@ void Init(const DetoursPayload& payload, u64 startTime)
 	UBA_ASSERT(g_systemRoot.count);
 
 	wchar_t systemTemp[256];
-	True_GetEnvironmentVariableW(L"TEMP", systemTemp, 256);
+	DWORD tempLen = True_GetEnvironmentVariableW(L"TEMP", systemTemp, sizeof_array(systemTemp));(void)tempLen;
+	UBA_ASSERTF(tempLen, TC("TEMP environment variable not found?"));
+	UBA_ASSERTF(tempLen < sizeof_array(systemTemp), TC("systemTemp buffer not large enough"));
+
 	FixPath(g_systemTemp, systemTemp);
 
 	StringBuffer<512> applicationBuffer;
