[UBA]
* Fixed clang compile errors

diff --git a/Engine/Source/Programs/UnrealBuildAccelerator/Core/Private/UbaApplicationRules.cpp b/Engine/Source/Programs/UnrealBuildAccelerator/Core/Private/UbaApplicationRules.cpp
--- a/Engine/Source/Programs/UnrealBuildAccelerator/Core/Private/UbaApplicationRules.cpp
+++ b/Engine/Source/Programs/UnrealBuildAccelerator/Core/Private/UbaApplicationRules.cpp
@@ -702,7 +702,7 @@ namespace uba
 		using Super = ApplicationRulesClangClExe;
 		virtual bool IsOutputFile(StringView fileName, StringView systemTemp) const override
 		{
-			return Super::IsOutputFile(fileName, systemTemp) || fileName.count > systemTemp.count && fileName.StartsWith(systemTemp);
+			return Super::IsOutputFile(fileName, systemTemp) || (fileName.count > systemTemp.count && fileName.StartsWith(systemTemp));
 		}
 	};
 
