[UBA]
* Added graceful handling of failing FixPath in NtCreateFile... pvs-studio is providing some really messed up paths and the times there are ~ in there in odd places uba FixPath fails. Fail is logged out
* Fixed so GetModuleFileNameA does not assert in debug when provided n size is larger than 1024 bytes

diff --git a/Engine/Source/Programs/UnrealBuildAccelerator/Detours/Private/Windows/UbaDetoursFunctionsKernelBase.inl b/Engine/Source/Programs/UnrealBuildAccelerator/Detours/Private/Windows/UbaDetoursFunctionsKernelBase.inl
--- a/Engine/Source/Programs/UnrealBuildAccelerator/Detours/Private/Windows/UbaDetoursFunctionsKernelBase.inl
+++ b/Engine/Source/Programs/UnrealBuildAccelerator/Detours/Private/Windows/UbaDetoursFunctionsKernelBase.inl
@@ -1373,7 +1373,7 @@ DWORD Shared_GetModuleFileNameA(const tchar* func, HMODULE hModule, LPSTR lpFile
 {
 	// Verified called from used applications.. and does not automatically call W version
 	StringBuffer<> temp;
-	UBA_ASSERT(nSize < temp.capacity);
+	nSize = Min(nSize, (DWORD)temp.capacity);
 	DWORD res = Shared_GetModuleFileNameW(func, hModule, temp.data, nSize);
 	temp.Parse(lpFilename, res);
 	return res;
diff --git a/Engine/Source/Programs/UnrealBuildAccelerator/Detours/Private/Windows/UbaDetoursFunctionsNtDll.inl b/Engine/Source/Programs/UnrealBuildAccelerator/Detours/Private/Windows/UbaDetoursFunctionsNtDll.inl
--- a/Engine/Source/Programs/UnrealBuildAccelerator/Detours/Private/Windows/UbaDetoursFunctionsNtDll.inl
+++ b/Engine/Source/Programs/UnrealBuildAccelerator/Detours/Private/Windows/UbaDetoursFunctionsNtDll.inl
@@ -774,7 +774,10 @@ NTSTATUS NTAPI Shared_NtCreateFile(bool IsCreateFunc, PHANDLE hFileHandle, ACCES
 			if (!CanDetour(createFileName))
 				suppressCreateFileDetour = true;
 			else if (!FixPath(fileName, createFileName))
-				UBA_ASSERTF(false, L"FixPath failed for string '%ls'", createFileName);
+			{
+				DEBUG_LOG(L"FixPath failed for string '%ls'", createFileName);
+				suppressCreateFileDetour = true;
+			}
 
 			// TODO: Instead of using t_createFileFileName.. should we just resolve the path from ObjectAttributes->RootDirectory?
 			ObjectAttributes->RootDirectory = nullptr;
