[UBA]
* Fixed timer around scope in cache server connect that logs out if takes long time
* Changed message names to be prefixed with Storage/Cache/Session/System
* Added logging when active store is removed because of disconnect
* Reduced scope of lock for active stores

diff --git a/Engine/Source/Programs/UnrealBuildAccelerator/Common/Private/UbaCacheServer.cpp b/Engine/Source/Programs/UnrealBuildAccelerator/Common/Private/UbaCacheServer.cpp
--- a/Engine/Source/Programs/UnrealBuildAccelerator/Common/Private/UbaCacheServer.cpp
+++ b/Engine/Source/Programs/UnrealBuildAccelerator/Common/Private/UbaCacheServer.cpp
@@ -1503,29 +1503,37 @@ namespace uba
 			if (reader.GetLeft())
 				hint = reader.ReadString();
 
+			u64 startTime = GetTime();
 			SCOPED_FUTEX(m_connectionsLock, lock);
 			if (m_isRunningMaintenance)
 			{
 				writer.WriteBool(false);
-				writer.WriteString(TC("Running maintenance..."));
+				writer.WriteString(TCV("Running maintenance..."));
 			}
-
 			writer.WriteBool(true);
 			auto insres = m_connections.try_emplace(connectionInfo.GetId());
 			auto& connection = insres.first->second;
 			connection.clientVersion = clientVersion;
 			connection.connectTime = GetTime();
 			m_peakConnectionCount = Max(m_peakConnectionCount, u32(m_connections.size()));
+			lock.Leave();
 
-			StringBuffer<> logStr;
-			logStr.Appendf(TC("Client %u connected"), connectionInfo.GetId());
-			if (clientVersion != CacheNetworkVersion)
-				logStr.Appendf(TC(" (v%u)"), clientVersion);
-			if (IsReadOnly(connectionInfo))
-				logStr.Appendf(TC(" (readonly)"));
-			if (!hint.empty())
-				logStr.Appendf(TC(" %s"), hint.c_str());
-			m_logger.Log(LogEntryType_Info, logStr);
+			u64 endTime = GetTime();
+			if (TimeToMs(endTime - startTime) > 60*1000)
+				m_logger.Warning(TC("Took %s to connect client"), TimeToText(endTime - startTime).str);
+
+			if (insres.second)
+			{
+				StringBuffer<> logStr;
+				logStr.Appendf(TC("Client %u connected"), connectionInfo.GetId());
+				if (clientVersion != CacheNetworkVersion)
+					logStr.Appendf(TC(" (v%u)"), clientVersion);
+				if (IsReadOnly(connectionInfo))
+					logStr.Appendf(TC(" (readonly)"));
+				if (!hint.empty())
+					logStr.Appendf(TC(" %s"), hint.c_str());
+				m_logger.Log(LogEntryType_Info, logStr);
+			}
 
 			return true;
 		}
diff --git a/Engine/Source/Programs/UnrealBuildAccelerator/Common/Private/UbaNetwork.cpp b/Engine/Source/Programs/UnrealBuildAccelerator/Common/Private/UbaNetwork.cpp
--- a/Engine/Source/Programs/UnrealBuildAccelerator/Common/Private/UbaNetwork.cpp
+++ b/Engine/Source/Programs/UnrealBuildAccelerator/Common/Private/UbaNetwork.cpp
@@ -8,7 +8,7 @@ namespace uba
 	{
 		switch (type)
 		{
-			#define UBA_SYSTEM_MESSAGE(x) case SystemMessageType_##x: return AsView(TC("")#x);
+			#define UBA_SYSTEM_MESSAGE(x) case SystemMessageType_##x: return AsView(TC("System_")#x);
 			UBA_SYSTEM_MESSAGES
 			#undef UBA_SYSTEM_MESSAGE
 		default:
@@ -20,7 +20,7 @@ namespace uba
 	{
 		switch (type)
 		{
-			#define UBA_STORAGE_MESSAGE(x) case StorageMessageType_##x: return AsView(TC("")#x);
+			#define UBA_STORAGE_MESSAGE(x) case StorageMessageType_##x: return AsView(TC("Storage_")#x);
 			UBA_STORAGE_MESSAGES
 			#undef UBA_STORAGE_MESSAGE
 		default:
@@ -32,7 +32,7 @@ namespace uba
 	{
 		switch (type)
 		{
-			#define UBA_SESSION_MESSAGE(x) case SessionMessageType_##x: return AsView(TC("")#x);
+			#define UBA_SESSION_MESSAGE(x) case SessionMessageType_##x: return AsView(TC("Session_")#x);
 			UBA_SESSION_MESSAGES
 			#undef UBA_SESSION_MESSAGE
 		default:
@@ -44,7 +44,7 @@ namespace uba
 	{
 		switch (type)
 		{
-			#define UBA_CACHE_MESSAGE(x) case CacheMessageType_##x: return AsView(TC("")#x);
+			#define UBA_CACHE_MESSAGE(x) case CacheMessageType_##x: return AsView(TC("Cache_")#x);
 			UBA_CACHE_MESSAGES
 			#undef UBA_CACHE_MESSAGE
 		default:
diff --git a/Engine/Source/Programs/UnrealBuildAccelerator/Common/Private/UbaStorageClient.cpp b/Engine/Source/Programs/UnrealBuildAccelerator/Common/Private/UbaStorageClient.cpp
--- a/Engine/Source/Programs/UnrealBuildAccelerator/Common/Private/UbaStorageClient.cpp
+++ b/Engine/Source/Programs/UnrealBuildAccelerator/Common/Private/UbaStorageClient.cpp
@@ -71,7 +71,7 @@ namespace uba
 	{
 		m_client.RegisterOnConnected([this]()
 			{
-				StackBinaryWriter<1024> writer;
+				StackBinaryWriter<2048> writer;
 				NetworkMessage msg(m_client, ServiceId, StorageMessageType_Connect, writer);
 				writer.WriteString(TC("Client"));
 				writer.WriteU32(StorageNetworkVersion);
diff --git a/Engine/Source/Programs/UnrealBuildAccelerator/Common/Private/UbaStorageServer.cpp b/Engine/Source/Programs/UnrealBuildAccelerator/Common/Private/UbaStorageServer.cpp
--- a/Engine/Source/Programs/UnrealBuildAccelerator/Common/Private/UbaStorageServer.cpp
+++ b/Engine/Source/Programs/UnrealBuildAccelerator/Common/Private/UbaStorageServer.cpp
@@ -192,6 +192,14 @@ namespace uba
 			}
 		}
 		{
+			struct StoreToRemove
+			{
+				FileAccessor* fileAccessor;
+				CasEntry* casEntry;
+				MappedView mappedView;
+			};
+			List<StoreToRemove> storesToRemove;
+
 			SCOPED_WRITE_LOCK(m_activeStoresLock, lock);
 			for (auto it=m_activeStores.begin(); it!=m_activeStores.end();)
 			{
@@ -202,14 +210,16 @@ namespace uba
 					continue;
 				}
 
-				{
-					SCOPED_WRITE_LOCK(store.casEntry->lock, entryLock);
-					store.casEntry->verified = false;
-					store.casEntry->beingWritten = false;
-					if (m_traceStore)
-						m_trace->FileStoreEnd(clientId, store.casEntry->key);
-				}
+				m_logger.Detail(TC("Cancelled store id %u because of disconnect of client with id %u"), u32(it->first), clientId);
+				
+				storesToRemove.push_back({store.fileAccessor, store.casEntry, store.mappedView});
 
+				it = m_activeStores.erase(it);
+			}
+			lock.Leave();
+
+			for (StoreToRemove& store : storesToRemove)
+			{
 				if (auto fa = store.fileAccessor)
 				{
 					const tchar* filename = fa->GetFileName();
@@ -220,7 +230,15 @@ namespace uba
 				m_casDataBuffer.UnmapView(store.mappedView, TC("OnDisconnected"));
 				if (store.mappedView.isIndependent)
 					CloseFileMapping(m_logger, store.mappedView.handle, TC("StoreDone"));
-				it = m_activeStores.erase(it);
+
+				SCOPED_WRITE_LOCK(store.casEntry->lock, entryLock);
+				if (!store.casEntry->beingWritten)
+					m_logger.Warning(TC("CasEntry does not have being written flag?")); // Searching for a bug
+
+				store.casEntry->verified = false;
+				store.casEntry->beingWritten = false;
+				if (m_traceStore)
+					m_trace->FileStoreEnd(clientId, store.casEntry->key);
 			}
 		}
 		{
@@ -244,6 +262,10 @@ namespace uba
 				it = m_activeFetches.erase(it);
 			}
 		}
+		{
+			SCOPED_WRITE_LOCK(m_connectionInfoLock, lock);
+			m_connectionInfo.erase(clientId);
+		}
 	}
 
 	bool StorageServer::IsDisallowedPath(const tchar* fileName)
@@ -339,7 +361,8 @@ namespace uba
 	{
 		SCOPED_READ_LOCK(m_connectionInfoLock, lock);
 		auto findIt = m_connectionInfo.find(clientId);
-		UBA_ASSERT(findIt != m_connectionInfo.end());
+		if (findIt == m_connectionInfo.end())
+			return m_logger.Error(TC("Couldn't find connection info for client %u"), clientId);
 		Info& info = findIt->second;
 		lock.Leave();
 
