[UBA]
* Added code for scenario where messages timeout because they don't get a response. Theory is that Send succeeds after failing recv have closed the connection. Since they are running on different threads there might be a race here

diff --git a/Engine/Source/Programs/UnrealBuildAccelerator/Common/Private/UbaNetworkClient.cpp b/Engine/Source/Programs/UnrealBuildAccelerator/Common/Private/UbaNetworkClient.cpp
--- a/Engine/Source/Programs/UnrealBuildAccelerator/Common/Private/UbaNetworkClient.cpp
+++ b/Engine/Source/Programs/UnrealBuildAccelerator/Common/Private/UbaNetworkClient.cpp
@@ -834,6 +834,17 @@ namespace uba
 			}
 		}
 
+		// Theory for long timeouts waiting for gotResponse..
+		// There could have been a disconnect that happened just before the message above was added to m_activeMessages (thus calling Done on all active messages)
+		// But still backend Send succeeded on same socket... this would mean that gotResponse is never set and we get the timeout
+		// We need to explicitly check here so we are still connected and if not, call OnDisconnected
+		if (!connection.connected)
+		{
+			message.m_error = 9;
+			OnDisconnected(connection, 10);
+			return false;
+		}
+
 		if (async)
 			return true;
 
