[UBA]
* Added so we set error mode to SEM_NOOPENFILEERRORBOX before spawning first process since we want to prevent the application error dialog from popping up stalling everything

diff --git a/Engine/Source/Programs/UnrealBuildAccelerator/Common/Private/UbaProcess.cpp b/Engine/Source/Programs/UnrealBuildAccelerator/Common/Private/UbaProcess.cpp
--- a/Engine/Source/Programs/UnrealBuildAccelerator/Common/Private/UbaProcess.cpp
+++ b/Engine/Source/Programs/UnrealBuildAccelerator/Common/Private/UbaProcess.cpp
@@ -1572,6 +1572,15 @@ namespace uba
 
 		if (!m_parentProcess)
 		{
+			// We want to make sure no crash dialogs pop up when failing to spawn processes since that will stall everything forever until someone clicks the dialog
+			static bool once = []()
+				{
+					u32 old = SetErrorMode(SEM_FAILCRITICALERRORS | SEM_NOOPENFILEERRORBOX);
+					SetErrorMode(old | SEM_FAILCRITICALERRORS | SEM_NOOPENFILEERRORBOX);
+					return true;
+				}();(void)once;
+
+
 			if (IsRunningArm())
 			{
 				SCOPED_FUTEX(m_session.m_isX64ApplicationLock, lock);
