UnrealBuildAccelerator: Track peak memory for Windows actions using the job object

#rnx

diff --git a/Engine/Source/Programs/UnrealBuildAccelerator/Common/Private/UbaProcess.cpp b/Engine/Source/Programs/UnrealBuildAccelerator/Common/Private/UbaProcess.cpp
--- a/Engine/Source/Programs/UnrealBuildAccelerator/Common/Private/UbaProcess.cpp
+++ b/Engine/Source/Programs/UnrealBuildAccelerator/Common/Private/UbaProcess.cpp
@@ -446,14 +446,19 @@ namespace uba
 		return m_messageThread.Wait(millisecondsTimeout);
 	}
 
+	u64 ProcessImpl::GetTotalProcessorTime() const
+	{
+		return m_processStats.cpuTime;
+	}
+
 	u64 ProcessImpl::GetTotalWallTime() const
 	{
 		return m_processStats.wallTime;
 	}
 
-	u64 ProcessImpl::GetTotalProcessorTime() const
+	u64 ProcessImpl::GetPeakMemory() const
 	{
-		return m_processStats.cpuTime;
+		return m_processStats.peakJobMemory > 0 ? m_processStats.peakJobMemory : m_processStats.peakMemory;
 	}
 
 	bool ProcessImpl::Cancel()
@@ -639,6 +644,12 @@ namespace uba
 			JOBOBJECT_BASIC_ACCOUNTING_INFORMATION accountingInformation = {};
 			if (QueryInformationJobObject(m_accountingJobObject, JOBOBJECTINFOCLASS::JobObjectBasicAccountingInformation, &accountingInformation, sizeof(accountingInformation), NULL))
 				m_processStats.cpuTime = accountingInformation.TotalUserTime.QuadPart + accountingInformation.TotalKernelTime.QuadPart;
+			JOBOBJECT_EXTENDED_LIMIT_INFORMATION limitInformation = {};
+			if (QueryInformationJobObject(m_accountingJobObject, JOBOBJECTINFOCLASS::JobObjectExtendedLimitInformation, &limitInformation, sizeof(limitInformation), NULL))
+			{
+				m_processStats.peakProcessMemory = limitInformation.PeakProcessMemoryUsed;
+				m_processStats.peakJobMemory = limitInformation.PeakJobMemoryUsed;
+			}
 			CloseHandle(m_accountingJobObject);
 		}
 		#endif
diff --git a/Engine/Source/Programs/UnrealBuildAccelerator/Common/Private/UbaSession.cpp b/Engine/Source/Programs/UnrealBuildAccelerator/Common/Private/UbaSession.cpp
--- a/Engine/Source/Programs/UnrealBuildAccelerator/Common/Private/UbaSession.cpp
+++ b/Engine/Source/Programs/UnrealBuildAccelerator/Common/Private/UbaSession.cpp
@@ -155,6 +155,11 @@ namespace uba
 		UBA_ASSERT(m_process);
 		return m_process->GetTotalWallTime();
 	}
+	u64 ProcessHandle::GetPeakMemory() const
+	{
+		UBA_ASSERT(m_process);
+		return m_process->GetPeakMemory();
+	}
 	bool ProcessHandle::Cancel() const
 	{
 		UBA_ASSERT(m_process);
diff --git a/Engine/Source/Programs/UnrealBuildAccelerator/Common/Private/UbaSessionServer.cpp b/Engine/Source/Programs/UnrealBuildAccelerator/Common/Private/UbaSessionServer.cpp
--- a/Engine/Source/Programs/UnrealBuildAccelerator/Common/Private/UbaSessionServer.cpp
+++ b/Engine/Source/Programs/UnrealBuildAccelerator/Common/Private/UbaSessionServer.cpp
@@ -38,6 +38,7 @@ namespace uba
 		virtual bool WaitForExit(u32 millisecondsTimeout) override { return m_done.IsSet(millisecondsTimeout); }
 		virtual u64 GetTotalProcessorTime() const override { return m_processorTime; }
 		virtual u64 GetTotalWallTime() const override { return m_wallTime; }
+		virtual u64 GetPeakMemory() const override { return m_peakMemory;  }
 		virtual const Vector<ProcessLogLine>& GetLogLines() const override { return m_logLines; }
 		virtual const Vector<u8>& GetTrackedInputs() const override { return m_trackedInputs; }
 		virtual const Vector<u8>& GetTrackedOutputs() const override { return m_trackedOutputs; };
@@ -91,6 +92,7 @@ namespace uba
 		u32 m_exitCode = ~0u;
 		u64 m_processorTime = 0;
 		u64 m_wallTime = 0;
+		u64 m_peakMemory = 0;
 		Event m_done;
 		Vector<ProcessLogLine> m_logLines;
 		Vector<u8> m_trackedInputs;
@@ -1609,6 +1611,7 @@ namespace uba
 
 		process.m_processorTime = processStats.cpuTime;
 		process.m_wallTime = processStats.wallTime;
+		process.m_peakMemory = processStats.peakJobMemory > 0 ? processStats.peakJobMemory : processStats.peakMemory;
 		process.m_server = nullptr;
 		process.m_done.Set();
 		process.CallProcessExit(h);
diff --git a/Engine/Source/Programs/UnrealBuildAccelerator/Common/Private/UbaStats.cpp b/Engine/Source/Programs/UnrealBuildAccelerator/Common/Private/UbaStats.cpp
--- a/Engine/Source/Programs/UnrealBuildAccelerator/Common/Private/UbaStats.cpp
+++ b/Engine/Source/Programs/UnrealBuildAccelerator/Common/Private/UbaStats.cpp
@@ -54,6 +54,10 @@ namespace uba
 		if (iopsOther)
 			logger.Info(TC("  IopsOther                   %9llu"), iopsOther.load());
 
+		if (peakProcessMemory)
+			logger.Info(TC("  PeakProcessMem (Win)        %9s"), BytesToText(peakProcessMemory).str);
+		if (peakJobMemory)
+			logger.Info(TC("  PeakJobMem (Win)            %9s"), BytesToText(peakJobMemory).str);
 		if (peakMemory)
 			logger.Info(TC("  PeakMem                     %9s"), BytesToText(peakMemory).str);
 		if (cpuTime)
@@ -103,6 +107,11 @@ namespace uba
 			exitTime = reader.Read7BitEncoded();
 			wallTime = reader.Read7BitEncoded();
 			cpuTime = reader.Read7BitEncoded();
+			if (version >= 49)
+			{
+				peakProcessMemory = reader.Read7BitEncoded();
+				peakJobMemory = reader.Read7BitEncoded();
+			}
 			detoursMemory = reader.Read7BitEncoded();
 			peakMemory = reader.Read7BitEncoded();
 			if (version >= 39)
diff --git a/Engine/Source/Programs/UnrealBuildAccelerator/Common/Public/UbaProcess.h b/Engine/Source/Programs/UnrealBuildAccelerator/Common/Public/UbaProcess.h
--- a/Engine/Source/Programs/UnrealBuildAccelerator/Common/Public/UbaProcess.h
+++ b/Engine/Source/Programs/UnrealBuildAccelerator/Common/Public/UbaProcess.h
@@ -19,6 +19,7 @@ namespace uba
 		virtual bool WaitForExit(u32 millisecondsTimeout) { UBA_ASSERT(false); return false; }
 		virtual u64 GetTotalProcessorTime() const { UBA_ASSERT(false); return 0; }
 		virtual u64 GetTotalWallTime() const { UBA_ASSERT(false); return 0; }
+		virtual u64 GetPeakMemory() const { UBA_ASSERT(false); return 0; }
 		virtual const Vector<ProcessLogLine>& GetLogLines() const = 0;
 		virtual const Vector<u8>& GetTrackedInputs() const = 0;
 		virtual const Vector<u8>& GetTrackedOutputs() const = 0;
@@ -50,6 +51,7 @@ namespace uba
 		virtual bool WaitForExit(u32 millisecondsTimeout) override;
 		virtual u64 GetTotalProcessorTime() const override;
 		virtual u64 GetTotalWallTime() const override;
+		virtual u64 GetPeakMemory() const override;
 		virtual const Vector<ProcessLogLine>& GetLogLines() const override { return m_logLines; }
 		virtual const Vector<u8>& GetTrackedInputs() const override { return m_trackedInputs; }
 		virtual const Vector<u8>& GetTrackedOutputs() const override { return m_trackedOutputs; }
diff --git a/Engine/Source/Programs/UnrealBuildAccelerator/Common/Public/UbaProcessHandle.h b/Engine/Source/Programs/UnrealBuildAccelerator/Common/Public/UbaProcessHandle.h
--- a/Engine/Source/Programs/UnrealBuildAccelerator/Common/Public/UbaProcessHandle.h
+++ b/Engine/Source/Programs/UnrealBuildAccelerator/Common/Public/UbaProcessHandle.h
@@ -44,6 +44,7 @@ namespace uba
 		const Vector<u8>& GetTrackedOutputs() const;		// If ProcessStartInfo.trackInputs was true, this will return a buffer to wchar/char strings 
 		u64 GetTotalProcessorTime() const;					// Total used cpu time (in time units. Use TimeToMs etc)
 		u64 GetTotalWallTime() const;						// Total wall time (in time units. Use TimeToMs etc)
+		u64 GetPeakMemory() const;							// Peak memory usage (in bytes)
 		bool Cancel() const;								// Request to cancel process. Returns false if failed to cancel (might be at the end of exiting process)
 		const tchar* GetExecutingHost() const;				// Host that is executing process.
 		bool IsRemote() const;								// Returns true if process is a remote process
diff --git a/Engine/Source/Programs/UnrealBuildAccelerator/Common/Public/UbaTrace.h b/Engine/Source/Programs/UnrealBuildAccelerator/Common/Public/UbaTrace.h
--- a/Engine/Source/Programs/UnrealBuildAccelerator/Common/Public/UbaTrace.h
+++ b/Engine/Source/Programs/UnrealBuildAccelerator/Common/Public/UbaTrace.h
@@ -81,7 +81,7 @@ namespace uba
 		#undef UBA_TRACE_TYPE
 	};
 
-	static constexpr u32 TraceVersion = 48;
+	static constexpr u32 TraceVersion = 49;
 	static constexpr u32 TraceReadCompatibilityVersion = 6;
 
 	////////////////////////////////////////////////////////////////////////////////////////////////////
diff --git a/Engine/Source/Programs/UnrealBuildAccelerator/Core/Public/UbaProcessStats.h b/Engine/Source/Programs/UnrealBuildAccelerator/Core/Public/UbaProcessStats.h
--- a/Engine/Source/Programs/UnrealBuildAccelerator/Core/Public/UbaProcessStats.h
+++ b/Engine/Source/Programs/UnrealBuildAccelerator/Core/Public/UbaProcessStats.h
@@ -62,7 +62,11 @@ namespace uba
 
 		// Don't add in GetTotalTime()
 		AtomicU64 wallTime;
+
+		// From job object (windows only)
 		AtomicU64 cpuTime;
+		AtomicU64 peakProcessMemory;
+		AtomicU64 peakJobMemory;
 
 		AtomicU64 detoursMemory;
 		AtomicU64 peakMemory;
@@ -99,6 +103,8 @@ namespace uba
 			writer.Write7BitEncoded(exitTime);
 			writer.Write7BitEncoded(wallTime);
 			writer.Write7BitEncoded(cpuTime);
+			writer.Write7BitEncoded(peakProcessMemory);
+			writer.Write7BitEncoded(peakJobMemory);
 			writer.Write7BitEncoded(detoursMemory);
 			writer.Write7BitEncoded(peakMemory);
 			writer.Write7BitEncoded(iopsRead);
diff --git a/Engine/Source/Programs/UnrealBuildAccelerator/Host/Private/UbaExports.cpp b/Engine/Source/Programs/UnrealBuildAccelerator/Host/Private/UbaExports.cpp
--- a/Engine/Source/Programs/UnrealBuildAccelerator/Host/Private/UbaExports.cpp
+++ b/Engine/Source/Programs/UnrealBuildAccelerator/Host/Private/UbaExports.cpp
@@ -400,6 +400,11 @@ extern "C"
 		return TimeToTick(handle->GetTotalWallTime());
 	}
 
+	uba::u64 ProcessHandle_GetPeakMemory(const uba::ProcessHandle* handle)
+	{
+		return handle->GetPeakMemory();
+	}
+
 	void ProcessHandle_TraverseOutputFiles(const uba::ProcessHandle* handle, ProcessHandle_OutputFileCallback* callback, void* userData)
 	{
 		handle->TraverseOutputFiles([&](uba::StringView file) { callback(userData, file.data); });
diff --git a/Engine/Source/Programs/UnrealBuildAccelerator/Host/Public/UbaExports.h b/Engine/Source/Programs/UnrealBuildAccelerator/Host/Public/UbaExports.h
--- a/Engine/Source/Programs/UnrealBuildAccelerator/Host/Public/UbaExports.h
+++ b/Engine/Source/Programs/UnrealBuildAccelerator/Host/Public/UbaExports.h
@@ -95,6 +95,7 @@ extern "C"
 	UBA_API uba::u64 ProcessHandle_GetHash(const uba::ProcessHandle* handle);
 	UBA_API uba::u64 ProcessHandle_GetTotalProcessorTime(const uba::ProcessHandle* handle);
 	UBA_API uba::u64 ProcessHandle_GetTotalWallTime(const uba::ProcessHandle* handle);
+	UBA_API uba::u64 ProcessHandle_GetPeakMemory(const uba::ProcessHandle* handle);
 	UBA_API void ProcessHandle_TraverseOutputFiles(const uba::ProcessHandle* handle, ProcessHandle_OutputFileCallback* callback, void* userData);
 	UBA_API bool ProcessHandle_WaitForExit(const uba::ProcessHandle* handle, uba::u32 millisecondsTimeout);
 	UBA_API void ProcessHandle_Cancel(const uba::ProcessHandle* handle, bool terminate);
diff --git a/Engine/Source/Programs/UnrealBuildAccelerator/Visualizer/Private/UbaVisualizer.cpp b/Engine/Source/Programs/UnrealBuildAccelerator/Visualizer/Private/UbaVisualizer.cpp
--- a/Engine/Source/Programs/UnrealBuildAccelerator/Visualizer/Private/UbaVisualizer.cpp
+++ b/Engine/Source/Programs/UnrealBuildAccelerator/Visualizer/Private/UbaVisualizer.cpp
@@ -1473,15 +1473,15 @@ namespace uba
 									if (accumulatingValue)
 									{
 										double invScaleY = duration * scale;
-										if (invScaleY != 0 && prevValue)
+										if (invScaleY != 0 && prevValue != 0)
 											y -= int(double(value - prevValue) / invScaleY) + offsetY;
 									}
 									else
 									{
 										double div = 1.0f;
-										if (maxValue)
+										if (maxValue != 0)
 											div = double(maxValue);
-										y -= int(double(maxValue - value*scale)*graphHeight/div) + offsetY;
+										y -= int(double((double)maxValue - (double)value*scale)*graphHeight/div) + offsetY;
 									}
 
 									if (prevValue <= 0)
@@ -2870,7 +2870,7 @@ namespace uba
 				drawText(L"Active Helpers: %u", Max(1u, m_traceView.activeSessionCount) - 1);
 			}
 
-			if (session.highestSendPerS || session.highestRecvPerS)
+			if (session.highestSendPerS > 0.0f || session.highestRecvPerS > 0.0f)
 			{
 				auto updateTime = session.updates.back();
 				auto updateSend = session.networkSend.back();
