UBA: Fix incorrect define

#rnx

diff --git a/Engine/Source/Programs/UnrealBuildAccelerator/Common/Private/UbaThread.cpp b/Engine/Source/Programs/UnrealBuildAccelerator/Common/Private/UbaThread.cpp
--- a/Engine/Source/Programs/UnrealBuildAccelerator/Common/Private/UbaThread.cpp
+++ b/Engine/Source/Programs/UnrealBuildAccelerator/Common/Private/UbaThread.cpp
@@ -119,7 +119,7 @@ namespace uba
 		UBA_ASSERT(err == 0); (void)err;
 		if (!description)
 			description = "UbaUnknown";
-		#if PLATFORM_MACOS
+		#if PLATFORM_MAC
 		pthread_setname_np(description);
 		#else
 		pthread_setname_np(pth, description);
diff --git a/Engine/Source/Programs/UnrealBuildAccelerator/Core/Private/UbaPlatform.cpp b/Engine/Source/Programs/UnrealBuildAccelerator/Core/Private/UbaPlatform.cpp
--- a/Engine/Source/Programs/UnrealBuildAccelerator/Core/Private/UbaPlatform.cpp
+++ b/Engine/Source/Programs/UnrealBuildAccelerator/Core/Private/UbaPlatform.cpp
@@ -174,7 +174,7 @@ namespace uba
 			outPretty.Append(TC("Linux/Wine-")).Append(GetWineVersion());
 		else
 			outPretty.Append(TC("Windows"));
-#elif PLATFORM_MACOS
+#elif PLATFORM_MAC
 		FILE* osRelease = popen("sw_vers -productVersion", "r");
 		if (!osRelease)
 			return false;
