[UBA]
* Added code to force maintenance if caskey/path tables are getting too big (close to absolute max)

diff --git a/Engine/Source/Programs/UnrealBuildAccelerator/Common/Private/UbaCacheClient.cpp b/Engine/Source/Programs/UnrealBuildAccelerator/Common/Private/UbaCacheClient.cpp
--- a/Engine/Source/Programs/UnrealBuildAccelerator/Common/Private/UbaCacheClient.cpp
+++ b/Engine/Source/Programs/UnrealBuildAccelerator/Common/Private/UbaCacheClient.cpp
@@ -1467,7 +1467,7 @@ namespace uba
 
 		UBA_TRACE_FETCH_HINT_SCOPE("FetchCasTable")
 
-		if (requiredCasTableOffset > 256*1024*1024)
+		if (requiredCasTableOffset > CasKeyTableMaxSize)
 			return m_logger.Warning(TC("Cas table offset %u too large. Cache entry corrupt (Bucket %llu)"), bucket.id);
 
 		TimerScope ts2(stats.fetchCasTable);
diff --git a/Engine/Source/Programs/UnrealBuildAccelerator/Common/Private/UbaCacheServer.cpp b/Engine/Source/Programs/UnrealBuildAccelerator/Common/Private/UbaCacheServer.cpp
--- a/Engine/Source/Programs/UnrealBuildAccelerator/Common/Private/UbaCacheServer.cpp
+++ b/Engine/Source/Programs/UnrealBuildAccelerator/Common/Private/UbaCacheServer.cpp
@@ -630,12 +630,18 @@ namespace uba
 
 	bool CacheServer::RunMaintenance(bool force, bool allowSave, const Function<bool()>& shouldExit)
 	{
+		if (m_forceMaintenance)
+		{
+			force = true;
+			m_forceMaintenance = false;
+		}
+
 		if (m_addsSinceMaintenance == 0 && !force && !m_bucketIsOverflowing)
 			return true;
 
 		bool isFirst = true;
 
-		do
+		while (true)
 		{
 			SCOPED_FUTEX(m_connectionsLock, lock2);
 			if (!force && !m_connections.empty())
@@ -651,13 +657,14 @@ namespace uba
 				lock2.Leave();
 				Sleep(200);
 				lock2.Enter();
+				// m_server.DisconnectClients(); // TODO, if we've waited enough we want to kick all clients
 			}
 
 			lock2.Leave();
 
 			PrintContentionSummary(m_logger);
 
-			auto g = MakeGuard([&]()
+			auto runningMaintenanceGuard = MakeGuard([&]()
 				{
 					SCOPED_FUTEX(m_connectionsLock, lock3);
 					m_isRunningMaintenance = false;
@@ -695,8 +702,11 @@ namespace uba
 
 			if(!RunMaintenanceInternal(shouldExit, allowSave))
 				return false;
+			if (!m_bucketIsOverflowing)
+				break;
+			if (force)
+				runningMaintenanceGuard.Cancel();
 		}
-		while (m_bucketIsOverflowing);
 
 		return true;
 	}
@@ -1641,6 +1651,11 @@ namespace uba
 
 				
 				pathOffset = bucket.m_pathTable.Add(path.data, path.count);
+				if (pathOffset == ~0u)
+				{
+					m_forceMaintenance = true;
+					return m_logger.Warning(TC("Path table in bucket %llu too big to store more keys. Forcing maintenance run (Bump max size of split up in more buckets if this happens frequently)"), bucket.m_id);
+				}
 
 				#if 0
 				StringBuffer<> test;
@@ -1650,6 +1665,12 @@ namespace uba
 
 				u32 requiredCasTableSize;
 				outCasKeyOffset = bucket.m_casKeyTable.Add(casKey, pathOffset, requiredCasTableSize);
+				if (outCasKeyOffset == ~0u)
+				{
+					m_forceMaintenance = true;
+					return m_logger.Warning(TC("CasKey table in bucket %llu too big to store more keys. Forcing maintenance run (Bump max size of split up in more buckets if this happens frequently)"), bucket.m_id);
+				}
+
 				auto insres = entries.insert(outCasKeyOffset);
 				if (!insres.second)
 					return m_logger.Warning(TC("Client %u %s file %s exists more than once in cache entry (PathOffset %u, CasKey: %s CasKeyOffset: %u ClientVersion: %u PathtableVersion: %u)"), clientId, direction, path.data, pathOffset, CasKeyString(casKey).str, outCasKeyOffset, clientVersion, connectionBucket.pathTable.GetVersion());
diff --git a/Engine/Source/Programs/UnrealBuildAccelerator/Common/Private/UbaCompactTables.cpp b/Engine/Source/Programs/UnrealBuildAccelerator/Common/Private/UbaCompactTables.cpp
--- a/Engine/Source/Programs/UnrealBuildAccelerator/Common/Private/UbaCompactTables.cpp
+++ b/Engine/Source/Programs/UnrealBuildAccelerator/Common/Private/UbaCompactTables.cpp
@@ -6,8 +6,6 @@
 
 namespace uba
 {
-	static constexpr u32 PathTableMaxSize = 128*1024*1024; // Max size for string memory
-	static constexpr u32 CasKeyTableMaxSize = 256*1024*1024;
 	static constexpr u32 CasKeyLookupInitialCount = 1024;
 	static constexpr u32 CasKeyArrayMaxSize = 16;
 
@@ -47,6 +45,9 @@ namespace uba
 	{
 		InitMem();
 
+		if (m_pathTableMem.writtenSize >= PathTableMaxSize - (strLen + 5))
+			return ~0u;
+
 		const tchar* stringKeyString = str;
 		StringBuffer<MaxPath> tempStringKeyStr;
 		if (m_caseInsensitive)
@@ -642,6 +643,9 @@ namespace uba
 		if (!m_offsets.IsInitialized())
 			m_offsets.Init(CasKeyLookupInitialCount);
 
+		if (m_casKeyTableMem.writtenSize >= CasKeyTableMaxSize - sizeof(CasKey) + 5)
+			return ~0u;
+
 		bool added = false;
 		u32* casKeyOffset = InternalAdd(casKey, stringOffset, added);
 
diff --git a/Engine/Source/Programs/UnrealBuildAccelerator/Common/Public/UbaCacheServer.h b/Engine/Source/Programs/UnrealBuildAccelerator/Common/Public/UbaCacheServer.h
--- a/Engine/Source/Programs/UnrealBuildAccelerator/Common/Public/UbaCacheServer.h
+++ b/Engine/Source/Programs/UnrealBuildAccelerator/Common/Public/UbaCacheServer.h
@@ -104,6 +104,7 @@ namespace uba
 		Atomic<u64> m_cacheKeyHitCount;
 		Atomic<bool> m_isRunningMaintenance;
 		Atomic<bool> m_bucketIsOverflowing;
+		Atomic<bool> m_forceMaintenance;
 
 		Futex m_bucketsLock;
 		Map<u64, Bucket> m_buckets;
diff --git a/Engine/Source/Programs/UnrealBuildAccelerator/Common/Public/UbaCompactTables.h b/Engine/Source/Programs/UnrealBuildAccelerator/Common/Public/UbaCompactTables.h
--- a/Engine/Source/Programs/UnrealBuildAccelerator/Common/Public/UbaCompactTables.h
+++ b/Engine/Source/Programs/UnrealBuildAccelerator/Common/Public/UbaCompactTables.h
@@ -13,6 +13,9 @@ namespace uba
 	struct BinaryWriter;
 	struct StringKey;
 
+	static constexpr u32 PathTableMaxSize = 128*1024*1024; // Max size for string memory
+	static constexpr u32 CasKeyTableMaxSize = 256*1024*1024;
+
 	class CompactPathTable
 	{
 	public:
