[UBA]
* UbaCacheService - Added log entry for when crypto is enabled

diff --git a/Engine/Source/Programs/UnrealBuildAccelerator/CacheService/Private/UbaCacheService.cpp b/Engine/Source/Programs/UnrealBuildAccelerator/CacheService/Private/UbaCacheService.cpp
--- a/Engine/Source/Programs/UnrealBuildAccelerator/CacheService/Private/UbaCacheService.cpp
+++ b/Engine/Source/Programs/UnrealBuildAccelerator/CacheService/Private/UbaCacheService.cpp
@@ -311,7 +311,8 @@ namespace uba
 		#if UBA_DEBUG
 		dbgStr = TC(" (DEBUG)");
 		#endif
-		logger.Info(TC("UbaCacheService v%s(%u)%s (Workers: %u, Rootdir: \"%s\", StoreCapacity: %s, Expiration: %s)"), Version, CacheNetworkVersion, dbgStr, maxWorkerCount, g_rootDir.data, BytesToText(storageCapacity).str, TimeToText(MsToTime(expirationTimeSeconds)*1000, true).str);
+		logger.Info(TC("UbaCacheService v%s(%u)%s (Workers: %u, Rootdir: \"%s\", Crypto: %s, StoreCapacity: %s, Expiration: %s)"), Version, CacheNetworkVersion, dbgStr, maxWorkerCount,
+			g_rootDir.data, (requiresCrypto ? TC("on") : TC("off")), BytesToText(storageCapacity).str, TimeToText(MsToTime(expirationTimeSeconds)*1000, true).str);
 
 		u64 maintenanceReserveSizeMb = 196;
 
