[UBA] Fix crash in waitpid() POSIX detoured implementation when output status is null.

#rnx

diff --git a/Engine/Source/Programs/UnrealBuildAccelerator/Detours/Private/PosixOS/UbaDetoursFunctionsPosix.cpp b/Engine/Source/Programs/UnrealBuildAccelerator/Detours/Private/PosixOS/UbaDetoursFunctionsPosix.cpp
--- a/Engine/Source/Programs/UnrealBuildAccelerator/Detours/Private/PosixOS/UbaDetoursFunctionsPosix.cpp
+++ b/Engine/Source/Programs/UnrealBuildAccelerator/Detours/Private/PosixOS/UbaDetoursFunctionsPosix.cpp
@@ -2292,10 +2292,14 @@ UBA_EXPORT pid_t UBA_WRAPPER(waitpid)(pid_t pid, int* status, int options)
 {
 	UBA_INIT_DETOUR(waitpid, pid, status, options);
 	// TODO: Should probably report id to session
-	auto res = TRUE_WRAPPER(waitpid)(pid, status, options);
-	DEBUG_LOG_TRUE("waitpid", "(%i) -> %i (%i)", pid, res, *status);
-	if (WIFEXITED(*status))
+	// Always pass status to waitpid() so that we know whether to call Rpc_UpdateTables()
+	int resStatus = 0;
+	pid_t res = TRUE_WRAPPER(waitpid)(pid, &resStatus, options);
+	DEBUG_LOG_TRUE("waitpid", "(%i) -> %i (%i)", pid, res, resStatus);
+	if (WIFEXITED(resStatus))
 		Rpc_UpdateTables();
+	if (status)
+		*status = resStatus;
 	return res;
 }
 
