[UBA]
* Fixed bug in GetModuleFileNameA reported by external team. Not sure exactly what caused it but either the module has a name that is not ansi or there were parameters sent into GetModuleFileName that was not properly handled. Code has been replaced to use WideCharToMultiByte which fails without crashing instead of wcstombs_s

#FailedValidations Preflight

diff --git a/Engine/Source/Programs/UnrealBuildAccelerator/Core/Private/UbaStringBuffer.cpp b/Engine/Source/Programs/UnrealBuildAccelerator/Core/Private/UbaStringBuffer.cpp
--- a/Engine/Source/Programs/UnrealBuildAccelerator/Core/Private/UbaStringBuffer.cpp
+++ b/Engine/Source/Programs/UnrealBuildAccelerator/Core/Private/UbaStringBuffer.cpp
@@ -448,10 +448,7 @@ namespace uba
 	u32 StringBufferBase::Parse(char* out, u64 outCapacity) const
 	{
 		#if PLATFORM_WINDOWS
-		size_t destLen;
-		if (wcstombs_s(&destLen, out, outCapacity, data, outCapacity-1) != 0)
-			return 0;
-		return (u32)destLen;
+		return WideCharToMultiByte(CP_ACP, 0, data, count+1, out, int(outCapacity), nullptr, nullptr);
 		#else
 		if (outCapacity == 0)
 			return 0;
diff --git a/Engine/Source/Programs/UnrealBuildAccelerator/Detours/Private/Windows/UbaDetoursFunctionsKernelBase.inl b/Engine/Source/Programs/UnrealBuildAccelerator/Detours/Private/Windows/UbaDetoursFunctionsKernelBase.inl
--- a/Engine/Source/Programs/UnrealBuildAccelerator/Detours/Private/Windows/UbaDetoursFunctionsKernelBase.inl
+++ b/Engine/Source/Programs/UnrealBuildAccelerator/Detours/Private/Windows/UbaDetoursFunctionsKernelBase.inl
@@ -1375,8 +1375,12 @@ DWORD Shared_GetModuleFileNameA(const tchar* func, HMODULE hModule, LPSTR lpFile
 	StringBuffer<> temp;
 	nSize = Min(nSize, (DWORD)temp.capacity);
 	DWORD res = Shared_GetModuleFileNameW(func, hModule, temp.data, nSize);
-	temp.Parse(lpFilename, res);
-	return res;
+	if (res == 0)
+		return 0;
+	u32 res2 = temp.Parse(lpFilename, res);
+	if (res2 == 0)
+		return 0;
+	return res2 - 1;
 }
 
 DWORD Detoured_GetModuleFileNameA(HMODULE hModule, LPSTR lpFilename, DWORD nSize)
