[UBA]
* Fixed bug in previous change for GetModuleFileNameA. Code was not setting count on stringbuilder after calling wide version

diff --git a/Engine/Source/Programs/UnrealBuildAccelerator/Detours/Private/Windows/UbaDetoursFunctionsKernelBase.inl b/Engine/Source/Programs/UnrealBuildAccelerator/Detours/Private/Windows/UbaDetoursFunctionsKernelBase.inl
--- a/Engine/Source/Programs/UnrealBuildAccelerator/Detours/Private/Windows/UbaDetoursFunctionsKernelBase.inl
+++ b/Engine/Source/Programs/UnrealBuildAccelerator/Detours/Private/Windows/UbaDetoursFunctionsKernelBase.inl
@@ -1242,7 +1242,11 @@ DWORD Detoured_GetFullPathNameA(LPCSTR lpFileName, DWORD nBufferLength, LPSTR lp
 	u64 requiredSize = fullPath.count + 1;
 	if (nBufferLength < requiredSize)
 		return DWORD(requiredSize);
-	fullPath.Parse(lpBuffer, requiredSize);
+	if (!fullPath.Parse(lpBuffer, requiredSize))
+	{
+		DEBUG_LOG(L"GetFullPathNameA failed to parse %s to ansi", fullPath.data);
+	}
+
 	if (lpFilePart)
 		*lpFilePart = strrchr(lpBuffer, '\\') + 1;
 	auto res = DWORD(fullPath.count);
@@ -1377,9 +1381,13 @@ DWORD Shared_GetModuleFileNameA(const tchar* func, HMODULE hModule, LPSTR lpFile
 	DWORD res = Shared_GetModuleFileNameW(func, hModule, temp.data, nSize);
 	if (res == 0)
 		return 0;
-	u32 res2 = temp.Parse(lpFilename, res);
+	temp.count = res;
+	u32 res2 = temp.Parse(lpFilename, nSize);
 	if (res2 == 0)
+	{
+		DEBUG_LOG(L"GetModuleFileNameA failed to parse %s to ansi", temp.data);
 		return 0;
+	}
 	return res2 - 1;
 }
 
