[UBA]
* Changed code writing out times/bytes to use while loop to write suffix

diff --git a/Engine/Source/Programs/UnrealBuildAccelerator/Common/Private/UbaLogger.cpp b/Engine/Source/Programs/UnrealBuildAccelerator/Common/Private/UbaLogger.cpp
--- a/Engine/Source/Programs/UnrealBuildAccelerator/Common/Private/UbaLogger.cpp
+++ b/Engine/Source/Programs/UnrealBuildAccelerator/Common/Private/UbaLogger.cpp
@@ -387,7 +387,7 @@ namespace uba
 	void Format(tchar* out, u32 value, const tchar* suffix)
 	{
 		out += UToA(out, value);
-		TStrcpy_s(out, 4, suffix);
+		do { *out = *suffix; if (!*suffix) break; ++out; ++suffix; } while (true);
 	}
 
 	void Format1dp(tchar* out, double value, const tchar* suffix)
@@ -396,7 +396,7 @@ namespace uba
 		out += UToA(out, scaled / 10);
 		*out++ = '.';
 		*out++ = '0' + scaled % 10;
-		TStrcpy_s(out, 4, suffix);
+		do { *out = *suffix; if (!*suffix) break; ++out; ++suffix; } while (true);
 	}
 
 	BytesToText::BytesToText(u64 bytes, bool base2)
