[UBA]
* Fixed temporary code that is supposed to output where stall is to actually work. Before it always showed index just before exiting function

diff --git a/Engine/Source/Programs/UnrealBuildAccelerator/Common/Private/UbaStorageServer.cpp b/Engine/Source/Programs/UnrealBuildAccelerator/Common/Private/UbaStorageServer.cpp
--- a/Engine/Source/Programs/UnrealBuildAccelerator/Common/Private/UbaStorageServer.cpp
+++ b/Engine/Source/Programs/UnrealBuildAccelerator/Common/Private/UbaStorageServer.cpp
@@ -464,14 +464,27 @@ namespace uba
 			case StorageMessageType_FetchBegin:
 			{
 				// TODO: Remove this when we have tracked down the issue where clients time out
-				u32 todoRemoveMe = 0;
+				u32 todoRemoveMeIndex = 0;
+				u64 todoRemoveMeLastTime = GetTime();
+				u64 todoRemoveMeLongestIndex = 0;
+				auto todoRemoveMe = [&](int index)
+					{
+						u64 t = GetTime();
+						u64 delta = t - todoRemoveMeLastTime;
+						todoRemoveMeLastTime = t;
+						if (delta <= todoRemoveMeLongestIndex)
+							return;
+						todoRemoveMeIndex = index;
+						todoRemoveMeLongestIndex = delta;
+					};
+
 				auto timeoutGuard = MakeGuard([&, timeoutStartTime = GetTime()]()
 					{
 						u64 timeSpentMs = TimeToMs(GetTime() - timeoutStartTime);
 						if (timeSpentMs > 8 * 60 * 1000)
 						{
 							// Took more than 5 minutes to respond
-							m_logger.Warning(TC("Took more than 8 minutes to respond to FetchBegin to client %u (Hang code %u).."), connectionInfo.GetId(), todoRemoveMe);
+							m_logger.Warning(TC("Took more than 8 minutes to respond to FetchBegin to client %u (Hang code %u).."), connectionInfo.GetId(), todoRemoveMeIndex);
 						}
 					});
 
@@ -479,12 +492,12 @@ namespace uba
 				bool wantsProxy = (recvFlags & 1) != 0;
 				if (wantsProxy) // Wants proxy
 				{
-					todoRemoveMe = 1;
+					todoRemoveMe(1);
 					if (GetProxy(writer, connectionInfo.GetId(), true))
 						return true;
 				}
 
-				todoRemoveMe = 2;
+				todoRemoveMe(2);
 
 				u64 start = GetTime();
 				CasKey casKey = reader.ReadCasKey();
@@ -505,7 +518,7 @@ namespace uba
 				bool has = HasCasFile(casKey, &casEntry); // HasCasFile also writes deferred cas entries if in queue
 				if (!has)
 				{
-					todoRemoveMe = 3;
+					todoRemoveMe(3);
 					if (!EnsureCasFile(casKey, nullptr) && m_allowHintAsFallback)
 					{
 						// Last resort.. use hint to load file into cas (hint should be renamed since it is now a critical parameter)
@@ -572,7 +585,7 @@ namespace uba
 				entryLock.Leave();
 				auto rlg = MakeGuard([&]() { entryLock.Enter(); --casEntry->readCount; });
 
-				todoRemoveMe = 4;
+				todoRemoveMe(4);
 
 				if (casEntry->disallowed)
 				{
@@ -585,7 +598,7 @@ namespace uba
 				FileHandle readFileHandle = InvalidFileHandle;
 				auto rfg = MakeGuard([&]()
 					{
-						todoRemoveMe = 13;
+						todoRemoveMe(13);
 						if (!CloseFile(nullptr, readFileHandle))
 							m_logger.Error(TC("Failed to close file %s (%s)"), casFile.data, LastErrorToText().data);
 					});
@@ -599,7 +612,7 @@ namespace uba
 					{
 						if (ownsMapping)
 						{
-							todoRemoveMe = 11;
+							todoRemoveMe(11);
 							if (mappedView.memory)
 								if (!UnmapViewOfFile(m_logger, mappedView.memory, fileSize, TC("FetchBegin")))
 									m_logger.Error(TC("Failed to unmap memory %s at 0x%llx with size %llu (%s)"), hint.data, mappedView.memory, fileSize, LastErrorToText().data);
@@ -609,12 +622,12 @@ namespace uba
 						}
 						else
 						{
-							todoRemoveMe = 12;
+							todoRemoveMe(12);
 							m_casDataBuffer.UnmapView(mappedView, TC("FetchBegin"));
 						}
 					});
 
-				todoRemoveMe = 5;
+				todoRemoveMe(5);
 
 				bool useFileMapping = casEntry->mappingHandle.IsValid();
 				if (useFileMapping)
@@ -654,7 +667,7 @@ namespace uba
 					}
 				}
 
-				todoRemoveMe = 6;
+				todoRemoveMe(6);
 
 				if (detailedTrace)
 					m_trace->FileFetchSize(connectionInfo.GetId(), casKey, fileSize);
@@ -675,7 +688,7 @@ namespace uba
 				u32 toWrite = u32(Min(left, capacityLeft));
 				void* writeBuffer = writer.AllocWrite(toWrite);
 
-				todoRemoveMe = 7;
+				todoRemoveMe(7);
 
 				if (useFileMapping)
 				{
@@ -709,7 +722,7 @@ namespace uba
 					readFileHandle = InvalidFileHandle;
 				}
 
-				todoRemoveMe = 8;
+				todoRemoveMe(8);
 
 				u64 actualSize = fileSize;
 				if (m_storeCompressed)
@@ -731,7 +744,7 @@ namespace uba
 					return true;
 				}
 
-				todoRemoveMe = 9;
+				todoRemoveMe(9);
 
 				rlg.Cancel();
 				mvg.Cancel();
@@ -746,7 +759,7 @@ namespace uba
 				fetch.clientId = connectionInfo.GetId();
 				lock.Leave();
 
-				todoRemoveMe = 10;
+				todoRemoveMe(10);
 
 				mappedView.size = fileSize;
 
