[UBA]
* Fixed crash when cancelling build while uba was killing processes due to memory pressure

diff --git a/Engine/Source/Programs/UnrealBuildAccelerator/Common/Private/UbaScheduler.cpp b/Engine/Source/Programs/UnrealBuildAccelerator/Common/Private/UbaScheduler.cpp
--- a/Engine/Source/Programs/UnrealBuildAccelerator/Common/Private/UbaScheduler.cpp
+++ b/Engine/Source/Programs/UnrealBuildAccelerator/Common/Private/UbaScheduler.cpp
@@ -486,11 +486,13 @@ namespace uba
 	{
 		auto& si = process.GetStartInfo();
 		auto& ei = *(ExitProcessInfo*)si.userData;
-		u32 processIndex = ei.processIndex;
 		float processWeight = 0;
 
 		if (isLocal)
 		{
+			if (!si.userData)
+				return false;
+
 			processWeight = ei.startInfo->weight;
 
 			SCOPED_CRITICAL_SECTION(m_wasReturnedLock, lock);
@@ -506,6 +508,8 @@ namespace uba
 			process.Cancel(); // Cancel will call ProcessExited
 		}
 		
+		u32 processIndex = ei.processIndex;
+
 		if (processIndex == ~0u)
 		{
 			UBA_ASSERT(!isLocal);
