[UBA]
* Fixed so LoadCasTable is always thread safe

diff --git a/Engine/Source/Programs/UnrealBuildAccelerator/Common/Private/UbaStorage.cpp b/Engine/Source/Programs/UnrealBuildAccelerator/Common/Private/UbaStorage.cpp
--- a/Engine/Source/Programs/UnrealBuildAccelerator/Common/Private/UbaStorage.cpp
+++ b/Engine/Source/Programs/UnrealBuildAccelerator/Common/Private/UbaStorage.cpp
@@ -981,6 +981,10 @@ namespace uba
 		if (!m_writeToDisk)
 			return true;
 
+		SCOPED_FUTEX(m_casTableLoadSaveLock, loadSaveLock);
+		if (m_casTableLoaded)
+			return true;
+
 		if (m_exclusiveMutex == InvalidMutexHandle)
 		{
 			m_exclusiveMutex = GetExclusiveAccess(m_logger, m_rootDir, true);
@@ -990,9 +994,6 @@ namespace uba
 
 		CreateDirectory(m_rootDir.data);
 
-		SCOPED_FUTEX(m_casTableLoadSaveLock, loadSaveLock);
-
-		UBA_ASSERTF(!m_casTableLoaded, TC("CasTable already loaded..."));
 		m_casTableLoaded = true;
 		u64 startTime = GetTime();
 		tchar isRunningName[256];
diff --git a/Engine/Source/Programs/UnrealBuildAccelerator/Common/Private/UbaStorageServer.cpp b/Engine/Source/Programs/UnrealBuildAccelerator/Common/Private/UbaStorageServer.cpp
--- a/Engine/Source/Programs/UnrealBuildAccelerator/Common/Private/UbaStorageServer.cpp
+++ b/Engine/Source/Programs/UnrealBuildAccelerator/Common/Private/UbaStorageServer.cpp
@@ -45,9 +45,7 @@ namespace uba
 
 		m_server.RegisterOnClientConnected(ServiceId, [this](const Guid& clientUid, u32 clientId)
 			{
-				SCOPED_FUTEX(m_loadCasTableLock, lock);
-				if (!m_casTableLoaded)
-					LoadCasTable(true);
+				LoadCasTable(true);
 			});
 
 		m_server.RegisterOnClientDisconnected(ServiceId, [this](const Guid& clientUid, u32 clientId)
diff --git a/Engine/Source/Programs/UnrealBuildAccelerator/Common/Public/UbaStorageServer.h b/Engine/Source/Programs/UnrealBuildAccelerator/Common/Public/UbaStorageServer.h
--- a/Engine/Source/Programs/UnrealBuildAccelerator/Common/Public/UbaStorageServer.h
+++ b/Engine/Source/Programs/UnrealBuildAccelerator/Common/Public/UbaStorageServer.h
@@ -134,8 +134,6 @@ namespace uba
 		ReaderWriterLock m_connectionInfoLock;
 		UnorderedMap<u32, Info> m_connectionInfo;
 
-		Futex m_loadCasTableLock;
-
 		Trace* m_trace = nullptr;
 
 		Vector<TString> m_disallowedPaths;
