[UBA]
* Changed so all access violations are caught when UbaObjTool is detoured

diff --git a/Engine/Source/Programs/UnrealBuildAccelerator/Core/Private/UbaApplicationRules.cpp b/Engine/Source/Programs/UnrealBuildAccelerator/Core/Private/UbaApplicationRules.cpp
--- a/Engine/Source/Programs/UnrealBuildAccelerator/Core/Private/UbaApplicationRules.cpp
+++ b/Engine/Source/Programs/UnrealBuildAccelerator/Core/Private/UbaApplicationRules.cpp
@@ -682,6 +682,11 @@ namespace uba
 		{
 			return DependencyCrawlerType_MsvcLinker;
 		}
+
+		virtual bool ReportAllExceptions() const override
+		{
+			return true;
+		}
 	};
 
 	class ApplicationRulesUbaTestApp : public ApplicationRules
