[UBA]
* Changed "error" to "err-or" in socket errors when running in debug so you can still run uba in debug and use submittool
* Fixed non unity compile error

diff --git a/Engine/Source/Programs/UnrealBuildAccelerator/Common/Private/UbaNetworkBackendTcp.cpp b/Engine/Source/Programs/UnrealBuildAccelerator/Common/Private/UbaNetworkBackendTcp.cpp
--- a/Engine/Source/Programs/UnrealBuildAccelerator/Common/Private/UbaNetworkBackendTcp.cpp
+++ b/Engine/Source/Programs/UnrealBuildAccelerator/Common/Private/UbaNetworkBackendTcp.cpp
@@ -2,6 +2,7 @@
 
 #include "UbaNetworkBackendTcp.h"
 #include "UbaConfig.h"
+#include "UbaEnvironment.h"
 #include "UbaEvent.h"
 #include "UbaHash.h"
 #include "UbaPlatform.h"
@@ -1203,7 +1204,7 @@ namespace uba
 				if (res == SOCKET_ERROR)
 				{
 					#if UBA_LOG_SOCKET_ERRORS
-					connection.logger.Info(TC("RecvSocket - WSAPoll returned an error for socket %i and connection %s: %s (%s%s)"), socket, GuidToString(connection.uid).str, LastErrorToText(WSAGetLastError()).data, connection.recvHint, hint);
+					connection.logger.Info(TC("RecvSocket - WSAPoll returned an err-or for socket %i and connection %s: %s (%s%s)"), socket, GuidToString(connection.uid).str, LastErrorToText(WSAGetLastError()).data, connection.recvHint, hint);
 					#endif
 					return false;
 				}
@@ -1212,11 +1213,11 @@ namespace uba
 
 			#if !PLATFORM_WINDOWS
 			if (!isFirstCall && errno != ECONNRESET)
-				return connection.logger.Warning(TC("RecvSocket - recv error on socket %i and connection %s: %s (%s%s)"), socket, GuidToString(connection.uid).str, strerror(errno), connection.recvHint, hint);
+				return connection.logger.Warning(TC("RecvSocket - recv err-or on socket %i and connection %s: %s (%s%s)"), socket, GuidToString(connection.uid).str, strerror(errno), connection.recvHint, hint);
 			#endif
 
 			#if UBA_LOG_SOCKET_ERRORS
-			connection.logger.Info(TC("RecvSocket - read returned an error for socket %i and connection %s: %s (%s%s)"), socket, GuidToString(connection.uid).str, LastErrorToText(WSAGetLastError()).data, connection.recvHint, hint);
+			connection.logger.Info(TC("RecvSocket - read returned an err-or for socket %i and connection %s: %s (%s%s)"), socket, GuidToString(connection.uid).str, LastErrorToText(WSAGetLastError()).data, connection.recvHint, hint);
 			#endif
 			return false;
 		}
