[UBA]
* Fixed access violation in uba scheduler when returning processes from remotes

diff --git a/Engine/Source/Programs/UnrealBuildAccelerator/Common/Private/UbaScheduler.cpp b/Engine/Source/Programs/UnrealBuildAccelerator/Common/Private/UbaScheduler.cpp
--- a/Engine/Source/Programs/UnrealBuildAccelerator/Common/Private/UbaScheduler.cpp
+++ b/Engine/Source/Programs/UnrealBuildAccelerator/Common/Private/UbaScheduler.cpp
@@ -487,6 +487,7 @@ namespace uba
 		auto& si = process.GetStartInfo();
 		auto& ei = *(ExitProcessInfo*)si.userData;
 		float processWeight = 0;
+		u32 processIndex;
 
 		if (isLocal)
 		{
@@ -494,6 +495,7 @@ namespace uba
 				return false;
 
 			processWeight = ei.startInfo->weight;
+			processIndex = ei.processIndex; // Must be fetched before cancel
 
 			SCOPED_CRITICAL_SECTION(m_wasReturnedLock, lock);
 			if (ei.hasCheckedWasReturned)
@@ -504,11 +506,11 @@ namespace uba
 		}
 		else
 		{
+			processIndex = ei.processIndex; // Must be fetched before cancel
 			ei.wasReturned = true;
 			process.Cancel(); // Cancel will call ProcessExited
 		}
 		
-		u32 processIndex = ei.processIndex;
 
 		if (processIndex == ~0u)
 		{
