[UBA]
* Added so compressed filemappings are not always kept in memory. It is now trying to respect keepOutputFileMemoryMapsThreshold
* Added application rule CloseFileMappingAfterUse and set it to true to all linkers we use that (almost) only write monolithics. Since we never link these remotely this doesn't matter and .obj files will always be released
* Added support for keepOutputFileMemoryMapsThreshold on session client side. Still defaults to keeping all files in memory but can now be controlled from config
* Fixed so more session client creation options are tweakable from config
* Fixed bug where Config_Load did not really work if specific config file was provided
* Added StorageServer_PreloadCasTable which can be used to load castable in parallel with doing other things (connecting to cache servers etc)
* WIP - Added support for transient virtual files (Not added to file mapping table).
* Fixed outdated code inside UBA_DEBUG_FILE_MAPPING (which is disabled by default)

diff --git a/Engine/Source/Programs/UnrealBuildAccelerator/Common/Private/UbaFileMapping.cpp b/Engine/Source/Programs/UnrealBuildAccelerator/Common/Private/UbaFileMapping.cpp
--- a/Engine/Source/Programs/UnrealBuildAccelerator/Common/Private/UbaFileMapping.cpp
+++ b/Engine/Source/Programs/UnrealBuildAccelerator/Common/Private/UbaFileMapping.cpp
@@ -264,9 +264,9 @@ namespace uba
 		#if UBA_DEBUG_FILE_MAPPING
 		{
 			SCOPED_FUTEX(g_fileMappingsLock, lock);
-			auto findIt = g_fileMappings.find(fileMappingObject.handle);
-			UBA_ASSERTF(findIt != g_fileMappings.end(), TC("Mapping nonexisting handle: %llu"), u64(fileMappingObject.handle));
-			auto insres = g_viewMappings.try_emplace(res, fileMappingObject.handle);
+			auto findIt = g_fileMappings.find(fileMappingObject.mh);
+			UBA_ASSERTF(findIt != g_fileMappings.end(), TC("Mapping nonexisting handle: %llu"), u64(fileMappingObject.mh));
+			auto insres = g_viewMappings.try_emplace(res, fileMappingObject.mh);
 			UBA_ASSERT(insres.second);
 			++findIt->second.viewCount;
 		}
@@ -323,10 +323,10 @@ namespace uba
 
 		#if UBA_DEBUG_FILE_MAPPING
 		SCOPED_FUTEX(g_fileMappingsLock, lock);
-		auto findIt = g_fileMappings.find(h.handle);
-		UBA_ASSERTF(findIt != g_fileMappings.end(), TC("Handle: %llu"), u64(h.handle));
+		auto findIt = g_fileMappings.find(h.mh);
+		UBA_ASSERTF(findIt != g_fileMappings.end(), TC("Handle: %llu"), u64(h.mh));
 		UBA_ASSERTF(!findIt->second.viewCount.load(), TC("Closing file mapping with %llu open views (%s)"), findIt->second.viewCount.load(), findIt->second.hint.c_str());
-		//logger.Info(TC("FILEMAPPING CLOSE: %s %llu (%llu)"), findIt->second.hint.c_str(), u64(h.handle), g_fileMappings.size() - 1);
+		logger.Info(TC("FILEMAPPING CLOSE: %s %llu (%llu)"), findIt->second.hint.c_str(), u64(h.mh), g_fileMappings.size() - 1);
 		g_fileMappings.erase(findIt);
 		#endif
 
@@ -742,7 +742,7 @@ namespace uba
 #endif
 	}
 
-	void FileMappingBuffer::UnmapView(MappedView view, const tchar* hint_, u64 newSize)
+	void FileMappingBuffer::UnmapView(MappedView view, const tchar* hint_, u64 newSize, bool allowDeferredUnmap)
 	{
 		if (!view.handle.IsValid())
 			return;
@@ -783,7 +783,12 @@ namespace uba
 
 				//--file.activeMapCount;
 			};
-		if (m_workManager && newSize == InvalidValue)
+
+		#if UBA_DEBUG_FILE_MAPPING
+		allowDeferredUnmap = false;
+		#endif
+
+		if (allowDeferredUnmap && m_workManager && newSize == InvalidValue)
 			m_workManager->AddWork([=, h = TString(hint_)](const WorkContext&) { unmap(h.c_str()); }, 1, TC("UnmapView"));
 		else
 			unmap(hint_);
diff --git a/Engine/Source/Programs/UnrealBuildAccelerator/Common/Private/UbaSession.cpp b/Engine/Source/Programs/UnrealBuildAccelerator/Common/Private/UbaSession.cpp
--- a/Engine/Source/Programs/UnrealBuildAccelerator/Common/Private/UbaSession.cpp
+++ b/Engine/Source/Programs/UnrealBuildAccelerator/Common/Private/UbaSession.cpp
@@ -226,7 +226,7 @@ namespace uba
 		table.GetValueAsU32(writeFilesFileMapMaxMb, TC("WriteFilesFileMapMaxMB"));
 		table.GetValueAsU32(writeFilesNoBufferingMinMb, TC("WriteFilesNoBufferingMinMB"));
 		table.GetValueAsU32(traceIntervalMs, TC("TraceIntervalMs"));
-
+		table.GetValueAsU64(keepOutputFileMemoryMapsThreshold, TC("KeepOutputFileMemoryMapsThreshold"));
 	}
 
 	bool Session::WriteFileToDisk(ProcessImpl& process, WrittenFile& file)
@@ -244,6 +244,7 @@ namespace uba
 		}
 
 		bool shouldEvictFromMemory = IsRarelyReadAfterWritten(process, file.name) || file.mappingWritten > m_keepOutputFileMemoryMapsThreshold;
+		bool tryEvictFromMemory = false;
 
 		u64 writtenSize = 0;
 		bool shouldWriteToDisk = ShouldWriteToDisk(file.name);
@@ -334,8 +335,11 @@ namespace uba
 				if (!m_storage.WriteCompressedFile(res, TC("MemoryMap"), InvalidFileHandle, mem, fileSize, file.name.c_str(), &header, sizeof(header), file.lastWriteTime))
 					return false;
 
-				shouldEvictFromMemory = false; // Can't evict without properly update filemappingtable.. the file on disk does now not match what was registered for write
 				writtenSize = res.storedSize;
+
+				// Can't evict without properly update filemappingtable.. the file on disk does now not match what was registered for write
+				tryEvictFromMemory = shouldEvictFromMemory;
+				shouldEvictFromMemory = false;
 			}
 			else
 			{
@@ -388,22 +392,7 @@ namespace uba
 		file.mappingHandle = {};
 		file.originalMappingHandle = {};
 
-		if (shouldEvictFromMemory)
-		{
-			#if UBA_DEBUG_TRACK_MAPPING
-			m_debugLogger->Info(TC("Mapping eviction queued 0x%llx (%s)"), u64(mh.mh), file.name.c_str());
-			#endif
-
-			m_workManager.AddWork([mh, n = file.name, this](const WorkContext&)
-				{
-					#if UBA_DEBUG_TRACK_MAPPING
-					m_debugLogger->Info(TC("Mapping evicted 0x%llx (%s)"), u64(mh.mh), n.c_str());
-					#endif
-
-					CloseFileMapping(m_logger, mh, n.c_str());
-				}, 1, TC("CloseFileMapping"));
-		}
-		else
+		if (!shouldEvictFromMemory)
 		{
 			StringBuffer<> name;
 			Storage::GetMappingString(name, mh, 0);
@@ -416,34 +405,58 @@ namespace uba
 			if (!insres.second && entry.canBeFreed) // It could be that this file has been read as input and that is fine.
 				m_logger.Error(TC("Trying to write the same file twice (%s)"), file.name.c_str());
 			UBA_ASSERT(mh.IsValid());
-			entry.handled = true;
-			entry.mapping = mh;
-			entry.mappingOffset = 0;
-			entry.contentSize = file.mappingWritten;
-			entry.lastWriteTime = file.lastWriteTime;
-			entry.isDir = false;
-			entry.success = true;
-			entry.canBeFreed = true;
-			entry.usedCount = 0;
-			entry.usedCountBeforeFree = g_globalRules.GetUsedCountBeforeFree(file.name);
-			entry.storedSize = writtenSize;
-			entry.isInvisible = rules.IsInvisible(file.name);
 
+			if (!entry.handled && tryEvictFromMemory)
+			{
+				shouldEvictFromMemory = true;
+			}
+			else
+			{
+				entry.handled = true;
+				entry.mapping = mh;
+				entry.mappingOffset = 0;
+				entry.contentSize = file.mappingWritten;
+				entry.lastWriteTime = file.lastWriteTime;
+				entry.isDir = false;
+				entry.success = true;
+				entry.canBeFreed = true;
+				entry.usedCount = 0;
+				entry.usedCountBeforeFree = g_globalRules.GetUsedCountBeforeFree(file.name);
+				entry.storedSize = writtenSize;
+				entry.isInvisible = rules.IsInvisible(file.name);
+
+				#if UBA_DEBUG_TRACK_MAPPING
+				entry.name = file.name;
+				m_debugLogger->Info(TC("Mapping kept 0x%llx (%s) from detoured process (UsedCountBeforeFree: %u)"), u64(mh.mh), entry.name.c_str(), u32(entry.usedCountBeforeFree));
+				#endif
+
+				if (!entry.isInvisible)
+				{
+					SCOPED_WRITE_LOCK(m_fileMappingTableMemLock, lock);
+					BinaryWriter writer(m_fileMappingTableMem, m_fileMappingTableSize);
+					writer.WriteStringKey(file.key);
+					writer.WriteString(name);
+					writer.Write7BitEncoded(file.mappingWritten);
+					u32 newSize = (u32)writer.GetPosition();
+					m_fileMappingTableSize = (u32)newSize;
+				}
+			}
+		}
+
+		if (shouldEvictFromMemory)
+		{
 			#if UBA_DEBUG_TRACK_MAPPING
-			entry.name = file.name;
-			m_debugLogger->Info(TC("Mapping kept 0x%llx (%s) from detoured process (UsedCountBeforeFree: %u)"), u64(mh.mh), entry.name.c_str(), u32(entry.usedCountBeforeFree));
+			m_debugLogger->Info(TC("Mapping eviction queued 0x%llx (%s)"), u64(mh.mh), file.name.c_str());
 			#endif
 
-			if (!entry.isInvisible)
-			{
-				SCOPED_WRITE_LOCK(m_fileMappingTableMemLock, lock);
-				BinaryWriter writer(m_fileMappingTableMem, m_fileMappingTableSize);
-				writer.WriteStringKey(file.key);
-				writer.WriteString(name);
-				writer.Write7BitEncoded(file.mappingWritten);
-				u32 newSize = (u32)writer.GetPosition();
-				m_fileMappingTableSize = (u32)newSize;
-			}
+			m_workManager.AddWork([mh, n = file.name, this](const WorkContext&)
+				{
+					#if UBA_DEBUG_TRACK_MAPPING
+					m_debugLogger->Info(TC("Mapping evicted 0x%llx (%s)"), u64(mh.mh), n.c_str());
+					#endif
+
+					CloseFileMapping(m_logger, mh, n.c_str());
+				}, 1, TC("CloseFileMapping"));
 		}
 
 		if (shouldWriteToDisk)
@@ -811,6 +824,8 @@ namespace uba
 	{
 		TimerScope ts(Stats().waitMmapFromFile);
 
+		StringView fileNameView = ToView(fileName);
+
 		SCOPED_FUTEX(m_fileMappingTableLookupLock, lookupLock);
 		auto insres = m_fileMappingTableLookup.try_emplace(fileNameKey);
 		FileMappingEntry& entry = insres.first->second;
@@ -821,8 +836,9 @@ namespace uba
 			{
 				if (requestingProcess && entry.canBeFreed && !m_runningRemote)
 				{
+					ProcessImpl::UsedFileMapping usedFileMapping { requestingProcess->m_startInfo.rules->CloseFileMappingAfterUse(fileNameView) };
 					SCOPED_FUTEX(requestingProcess->m_usedFileMappingsLock, lock);
-					if (!requestingProcess->m_hasExited && requestingProcess->m_usedFileMappings.insert(fileNameKey).second)
+					if (!requestingProcess->m_hasExited && requestingProcess->m_usedFileMappings.insert({fileNameKey, usedFileMapping}).second)
 						++entry.refCount;
 				}
 			};
@@ -882,6 +898,9 @@ namespace uba
 			{
 				if (!ReadFile(m_logger, fileName, fileHandle, &contentSize, sizeof(u64)))
 					return m_logger.Error(TC("Failed to read first bytes from file %s (%s)"), fileName, LastErrorToText().data);
+
+				if (contentSize > InvalidValue)
+					return m_logger.Error(TC("Compressed cas has content size larger than %s. File %s is corrupt"), BytesToText(InvalidValue).str, fileName);
 			}
 			else
 			{
@@ -893,7 +912,7 @@ namespace uba
 				contentSize = info.size;
 				lastWriteTime = info.lastWriteTime;
 
-				if (m_readIntermediateFilesCompressed && info.size > sizeof(CompressedFileHeader) && g_globalRules.FileCanBeCompressed(ToView(fileName)))
+				if (m_readIntermediateFilesCompressed && info.size > sizeof(CompressedFileHeader) && g_globalRules.FileCanBeCompressed(fileNameView))
 				{
 					CompressedFileHeader header(CasKeyZero);
 					if (!ReadFile(m_logger, fileName, fileHandle, &header, sizeof(header)))
@@ -904,6 +923,8 @@ namespace uba
 						isCompressed = true;
 						if (!ReadFile(m_logger, fileName, fileHandle, &contentSize, sizeof(u64)))
 							return m_logger.Error(TC("Failed to read first bytes from file %s (%s)"), fileName, LastErrorToText().data);
+						if (contentSize > InvalidValue)
+							return m_logger.Error(TC("Compressed cas has content size larger than %s. File %s is corrupt"), BytesToText(InvalidValue).str, fileName);
 					}
 					else
 					{
@@ -933,7 +954,7 @@ namespace uba
 			MappedView mappedView;
 
 			auto ownedGuard1 = MakeGuard([&]() { CloseFileMapping(m_logger, mapping, fileName); });
-			auto ownedGuard2 = MakeGuard([&]() { m_workManager.AddWork([this, mappingMemory, contentSize, fn = TString(fileName)](const WorkContext&) { UnmapViewOfFile(m_logger, mappingMemory, contentSize, fn.c_str()); }, 1, TC("UnmapViewOfFile")); });
+			auto ownedGuard2 = MakeGuard([&]() { m_workManager.AddWork([this, mappingMemory, contentSize, fn = fileNameView.ToString()](const WorkContext&) { UnmapViewOfFile(m_logger, mappingMemory, contentSize, fn.c_str()); }, 1, TC("UnmapViewOfFile")); });
 			auto viewGuard = MakeGuard([&]() { m_fileMappingBuffer.UnmapView(mappedView, fileName); });
 
 			if (canBeFreed)
@@ -983,7 +1004,7 @@ namespace uba
 			if (canBeFreed)
 			{
 				entry.usedCount = 0;
-				entry.usedCountBeforeFree = g_globalRules.GetUsedCountBeforeFree(ToView(fileName));
+				entry.usedCountBeforeFree = g_globalRules.GetUsedCountBeforeFree(fileNameView);
 			}
 
 			updateRequestingProcess();
@@ -2022,12 +2043,12 @@ namespace uba
 		return RegisterCreateFileForWrite(fileNameKey, fixedFilePath, false, sourceSize, 0, false);
 	}
 
-	bool Session::CreateVirtualFile(const tchar* filePath, const void* memory, u64 memorySize)
+	bool Session::CreateVirtualFile(const tchar* filePath, const void* memory, u64 memorySize, bool transient)
 	{
 		UBA_ASSERT(!m_runningRemote);
 		StringBuffer<> fixedFilePath;
 		auto fileNameKey = GetKeyAndFixedName(fixedFilePath, filePath);
-		if (!CreateVirtualFileInternal(fileNameKey, fixedFilePath, memory, memorySize))
+		if (!CreateVirtualFileInternal(fileNameKey, fixedFilePath, memory, memorySize, transient))
 			return false;
 		return RegisterCreateFileForWrite(fileNameKey, fixedFilePath, false, memorySize, 0, false);
 	}
@@ -2419,8 +2440,9 @@ namespace uba
 		auto usedFileMappings = std::move(process.m_usedFileMappings); // Might be that crawler is still working
 		usedFileMappingsLock.Leave();
 
-		for (auto& fileNameKey : usedFileMappings)
+		for (auto& kv : usedFileMappings)
 		{
+			const StringKey& fileNameKey = kv.first;
 			SCOPED_FUTEX(m_fileMappingTableLookupLock, lookupLock);
 			auto findIt = m_fileMappingTableLookup.find(fileNameKey);
 			UBA_ASSERT(findIt != m_fileMappingTableLookup.end());
@@ -2435,8 +2457,12 @@ namespace uba
 			if (--entry.refCount)
 				continue;
 
-			if (entry.usedCountBeforeFree == 255 || entry.usedCount < entry.usedCountBeforeFree)
-				continue;
+			auto& usedFileMapping = kv.second;
+			if (!usedFileMapping.closeAfterWritten)
+			{
+				if (entry.usedCountBeforeFree == 255 || entry.usedCount < entry.usedCountBeforeFree)
+					continue;
+			}
 
 			UBA_ASSERT(entry.canBeFreed);
 
@@ -2897,6 +2923,21 @@ namespace uba
 				}
 				return true;
 			}
+			else // Still need to check if file exists since it can be a virtual file
+			{
+				SCOPED_FUTEX(m_fileMappingTableLookupLock, lookupLock);
+				auto insres = m_fileMappingTableLookup.try_emplace(fileNameKey);
+				FileMappingEntry& entry = insres.first->second;
+				lookupLock.Leave();
+				SCOPED_FUTEX(entry.lock, entryLock);
+				if (entry.handled && entry.success && entry.mapping.IsValid())
+				{
+					UBA_ASSERT(entry.isInvisible);
+					out.size = entry.contentSize;
+					Storage::GetMappingString(out.fileName, entry.mapping, entry.mappingOffset);
+					return true;
+				}
+			}
 		}
 
 		if (!IsRarelyRead(process, fileName))
@@ -3524,10 +3565,10 @@ namespace uba
 		if (sourceSize + sourceOffset > virtualFile.size)
 			return m_logger.Error(TC("Virtual file offset(%llu)+size(%llu) outside source file size(%llu)"), sourceOffset, sourceSize, virtualFile.size, filePath.data);
 
-		return RegisterVirtualFileInternal(fileNameKey, filePath, virtualFile.mappingHandle, sourceSize, sourceOffset);
+		return RegisterVirtualFileInternal(fileNameKey, filePath, virtualFile.mappingHandle, sourceSize, sourceOffset, false);
 	}
 
-	bool Session::CreateVirtualFileInternal(const StringKey& fileNameKey, const StringView& filePath, const void* memory, u64 memorySize)
+	bool Session::CreateVirtualFileInternal(const StringKey& fileNameKey, const StringView& filePath, const void* memory, u64 memorySize, bool transient)
 	{
 		TimerScope ts(Stats().createMmapFromFile);
 
@@ -3547,10 +3588,10 @@ namespace uba
 
 		virtualFile.mappingHandle = mapping;
 		virtualFile.size = memorySize;
-		return RegisterVirtualFileInternal(fileNameKey, filePath, virtualFile.mappingHandle, memorySize, 0);
+		return RegisterVirtualFileInternal(fileNameKey, filePath, virtualFile.mappingHandle, memorySize, 0, transient);
 	}
 
-	bool Session::RegisterVirtualFileInternal(const StringKey& fileNameKey, const StringView& filePath, FileMappingHandle mappingHandle, u64 mappingSize, u64 mappingOffset)
+	bool Session::RegisterVirtualFileInternal(const StringKey& fileNameKey, const StringView& filePath, FileMappingHandle mappingHandle, u64 mappingSize, u64 mappingOffset, bool transient)
 	{
 		SCOPED_FUTEX(m_fileMappingTableLookupLock, lookupLock);
 		auto insres = m_fileMappingTableLookup.try_emplace(fileNameKey);
@@ -3568,6 +3609,7 @@ namespace uba
 		entry.handled = true;
 		entry.lastWriteTime = 0; // TODO: Take lastwritetime of source file?
 		entry.createIndependentMapping = true; // We want independent mapping in casdb so they can be deleted
+		entry.isInvisible = transient;
 
 		StringBuffer<> mappingName;
 		Storage::GetMappingString(mappingName, mappingHandle, mappingOffset);
@@ -3578,13 +3620,15 @@ namespace uba
 		m_debugLogger->Info(TC("Mapping created 0x%llx (%s) from virtual file"), u64(entry.mapping.mh), entry.name.c_str());
 		#endif
 
-		SCOPED_WRITE_LOCK(m_fileMappingTableMemLock, lock);
-		BinaryWriter writer(m_fileMappingTableMem, m_fileMappingTableSize);
-		writer.WriteStringKey(fileNameKey);
-		writer.WriteString(mappingName);
-		writer.Write7BitEncoded(mappingSize);
-		m_fileMappingTableSize = (u32)writer.GetPosition();
-
+		if (!transient)
+		{
+			SCOPED_WRITE_LOCK(m_fileMappingTableMemLock, lock);
+			BinaryWriter writer(m_fileMappingTableMem, m_fileMappingTableSize);
+			writer.WriteStringKey(fileNameKey);
+			writer.WriteString(mappingName);
+			writer.Write7BitEncoded(mappingSize);
+			m_fileMappingTableSize = (u32)writer.GetPosition();
+		}
 		return true;
 	}
 
diff --git a/Engine/Source/Programs/UnrealBuildAccelerator/Common/Private/UbaSessionClient.cpp b/Engine/Source/Programs/UnrealBuildAccelerator/Common/Private/UbaSessionClient.cpp
--- a/Engine/Source/Programs/UnrealBuildAccelerator/Common/Private/UbaSessionClient.cpp
+++ b/Engine/Source/Programs/UnrealBuildAccelerator/Common/Private/UbaSessionClient.cpp
@@ -30,14 +30,29 @@ namespace uba
 		}
 	}
 
+	SessionClientCreateInfo::SessionClientCreateInfo(Storage& s, NetworkClient& c, LogWriter& writer) : SessionCreateInfo(s, writer), client(c)
+	{
+		// By default we keep all files in memory
+		keepOutputFileMemoryMapsThreshold = InvalidValue;
+	}
+
 	void SessionClientCreateInfo::Apply(const Config& config)
 	{
 		SessionCreateInfo::Apply(config);
 
 		if (const ConfigTable* table = config.GetTable(TC("Session")))
 		{
+			table->GetValueAsU32(maxProcessCount, TC("MaxProcessCount"));
 			table->GetValueAsBool(useDependencyCrawler, TC("UseDependencyCrawler"));
 			table->GetValueAsU32(pingTimeoutSecondsPrintCallstacks, TC("PingTimeoutSecondsPrintCallstacks"));
+			table->GetValueAsU32(maxIdleSeconds, TC("MaxIdleSeconds"));
+			table->GetValueAsBool(killRandom, TC("KillRandom"));
+			table->GetValueAsBool(useStorage, TC("UseStorage"));
+			u32 temp;
+			if (table->GetValueAsU32(temp, TC("MemWaitLoadPercent")))
+				memWaitLoadPercent = u8(temp);
+			if (table->GetValueAsU32(temp, TC("MemKillLoadPercent")))
+				memKillLoadPercent = u8(temp);
 		}
 	}
 
@@ -609,7 +624,7 @@ namespace uba
 					bool storeUncompressed = memoryMapAlignment == 0;
 					bool allowProxy = rules.AllowStorageProxy(fileName);
 					if (!RetrieveCasFile(newCasKey, fileSize, casKey, fileName.data, storeUncompressed, allowProxy))
-						return m_logger.Error(TC("Error retrieving cas entry %s (%s)"), CasKeyString(casKey).str, fileName.data);
+						return m_logger.Error(TC("Error retrieving cas entry %s (%s) (1)"), CasKeyString(casKey).str, fileName.data);
 
 					if (!m_storage.GetCasFileName(newName, newCasKey))
 						return false;
@@ -640,7 +655,7 @@ namespace uba
 						Storage::RetrieveResult result;
 						bool allowProxy = rules.AllowStorageProxy(fileName);
 						if (!m_storage.RetrieveCasFile(result, casKey, fileName.data, &m_fileMappingBuffer, memoryMapAlignment, allowProxy))
-							return m_logger.Error(TC("Error retrieving cas entry %s (%s)"), CasKeyString(casKey).str, fileName.data);
+							return m_logger.Error(TC("Error retrieving cas entry %s (%s) (2)"), CasKeyString(casKey).str, fileName.data);
 						entry.success = true;
 						entry.contentSize = result.size;
 						entry.mapping = result.view.handle;
@@ -727,7 +742,7 @@ namespace uba
 			if (!pair.second.mappingHandle.IsValid())
 				m_logger.Warning(TC("%s is not using file mapping"), file.name.c_str());
 #endif
-			bool keepMappingInMemory = IsWindows && !IsRarelyReadAfterWritten(process, file.name);
+			bool keepMappingInMemory = IsWindows && !IsRarelyReadAfterWritten(process, file.name) && file.mappingWritten < m_keepOutputFileMemoryMapsThreshold;
 			bool compressed = rules.SendFileCompressedFromClient(file.name);
 			if (!SendFile(file, process.GetId(), keepMappingInMemory, compressed))
 				return false;
diff --git a/Engine/Source/Programs/UnrealBuildAccelerator/Common/Private/UbaStorage.cpp b/Engine/Source/Programs/UnrealBuildAccelerator/Common/Private/UbaStorage.cpp
--- a/Engine/Source/Programs/UnrealBuildAccelerator/Common/Private/UbaStorage.cpp
+++ b/Engine/Source/Programs/UnrealBuildAccelerator/Common/Private/UbaStorage.cpp
@@ -96,6 +96,7 @@ namespace uba
 		table.GetValueAsBool(storeCompressed, TC("StoreCompressed"));
 		table.GetValueAsU32(fileEntryTableWipeThreshold, TC("FileEntryTableWipeThreshold"));
 		table.GetValueAsU32(maxParallelCopyOrLink, TC("MaxParallelCopyOrLink"));
+		table.GetValueAsBool(dropIndependentMappingsOnCopy, TC("DropIndependentMappingsOnCopy"));
 
 		const tchar* compressor;
 		if (table.GetValueAsString(compressor, TC("Compressor")))
@@ -882,6 +883,7 @@ namespace uba
 		m_asyncUnmapViewOfFile = info.asyncUnmapViewOfFile;
 		m_allowDeleteVerified = info.allowDeleteVerified;
 		m_writeToDisk = info.writeToDisk;
+		m_dropIndependentMappingsOnCopy = info.dropIndependentMappingsOnCopy;
 		if (!m_writeToDisk)
 			m_casCapacityBytes = 0;
 
@@ -2094,6 +2096,7 @@ namespace uba
 			}
 			casEntry.verified = true;
 			casEntry.exists = false;
+			casEntry.dropped = false;
 			entryLock.Leave();
 	
 			CasEntryDeleted(casEntry, sizeDeleted);
@@ -2153,8 +2156,8 @@ namespace uba
 		bool testCompressed = !writeCompressed;
 		while (true)
 		{
-			CasEntry* casEntry = nullptr;
-			if (!HasCasFile(actualKey, &casEntry))
+			CasEntry* casEntryPtr = nullptr;
+			if (!HasCasFile(actualKey, &casEntryPtr))
 			{
 				if (!testCompressed)
 					return m_logger.Error(TC("[CopyOrLink] Trying to copy cas %s to %s but can't find neither compressed or uncompressed version"), CasKeyString(actualKey).str, destination);
@@ -2164,9 +2167,10 @@ namespace uba
 				continue;
 			}
 
-			SCOPED_READ_LOCK(casEntry->lock, casEntryLock);
-			UBA_ASSERT(casEntry->verified);
-			UBA_ASSERT(casEntry->exists);
+			CasEntry& casEntry = *casEntryPtr;
+			SCOPED_READ_LOCK(casEntry.lock, casEntryLock);
+			UBA_ASSERT(casEntry.verified);
+			UBA_ASSERT(casEntry.exists);
 
 			if (IsCompressed(actualKey))
 			{
@@ -2181,10 +2185,10 @@ namespace uba
 				StringBuffer<512> casFile;
 				u64 decompressedSize;
 
-				if (casEntry->mappingHandle.IsValid())
+				if (casEntry.mappingHandle.IsValid())
 				{
 					casFile.Append(CasKeyString(actualKey).str);
-					mappedView = m_casDataBuffer.MapView(casEntry->mappingHandle, casEntry->mappingOffset, casEntry->mappingSize, casFile.data, casEntry->mappingIsIndependent);
+					mappedView = m_casDataBuffer.MapView(casEntry.mappingHandle, casEntry.mappingOffset, casEntry.mappingSize, casFile.data, casEntry.mappingIsIndependent);
 					compressedData = mappedView.memory;
 					if (!compressedData)
 						return m_logger.Error(TC("[CopyOrLink] Failed to map view of mapping %s (%s)"), casFile.data, LastErrorToText().data);
@@ -2204,7 +2208,7 @@ namespace uba
 						return m_logger.Error(TC("[CopyOrLink] Failed to read first 8 bytes from compressed file %s (%s)"), casFile.data, LastErrorToText().data);
 				}
 
-				bool writeDirectlyToFile = casEntry->mappingHandle.IsValid() && false; // Experiment to try to fix bottlenecks on cloud
+				bool writeDirectlyToFile = casEntry.mappingHandle.IsValid() && false; // Experiment to try to fix bottlenecks on cloud
 				bool useNoBuffering = writeDirectlyToFile && false; // Experiment to try to fix bottlenecks on cloud
 
 				u32 writeFlags = useNoBuffering ? FILE_FLAG_NO_BUFFERING : 0;
@@ -2314,15 +2318,23 @@ namespace uba
 					fileEntry.size = sizeOnDisk;
 					fileEntry.verified = true;
 				}
+
+				if (m_dropIndependentMappingsOnCopy && casEntry.mappingIsIndependent && casEntry.mappingHandle.IsValid())
+				{
+					mapViewGuard.Execute();
+					casEntryLock.Leave();
+					DropCasFile(casKey, true, destination);
+				}
+
 				return true;
 			}
 
-			if (casEntry->mappingHandle.IsValid())
+			if (casEntry.mappingHandle.IsValid())
 			{
 				SCOPED_FUTEX(fileEntry.lock, entryLock);
 				fileEntry.verified = false;
 
-				MappedView mappedView = m_casDataBuffer.MapView(casEntry->mappingHandle, casEntry->mappingOffset, casEntry->mappingSize, TC(""));
+				MappedView mappedView = m_casDataBuffer.MapView(casEntry.mappingHandle, casEntry.mappingOffset, casEntry.mappingSize, TC(""));
 				auto mapViewGuard = MakeGuard([&](){ m_casDataBuffer.UnmapView(mappedView, destination); });
 
 				FileAccessor destinationFile(m_logger, destination);
diff --git a/Engine/Source/Programs/UnrealBuildAccelerator/Common/Public/UbaFileMapping.h b/Engine/Source/Programs/UnrealBuildAccelerator/Common/Public/UbaFileMapping.h
--- a/Engine/Source/Programs/UnrealBuildAccelerator/Common/Public/UbaFileMapping.h
+++ b/Engine/Source/Programs/UnrealBuildAccelerator/Common/Public/UbaFileMapping.h
@@ -69,7 +69,7 @@ namespace uba
 		MappedView MapView(FileMappingHandle handle, u64 offset, u64 size, const tchar* hint, bool isIndependent = false);
 		MappedView MapViewWrite(FileMappingHandle handle, u64 offset, u64 size, const tchar* hint, bool isIndependent = false);
 		MappedView MapView(const StringView& str, u64 size, const tchar* hint);
-		void UnmapView(MappedView view, const tchar* hint, u64 newSize = InvalidValue);
+		void UnmapView(MappedView view, const tchar* hint, u64 newSize = InvalidValue, bool allowDeferredUnmap = true);
 		
 		void GetSizeAndCount(FileMappingType type, u64& outSize, u32& outCount);
 		FileMappingHandle GetPersistentHandle(u32 index) { return m_storage[MappedView_Persistent].files[index].handle; }
diff --git a/Engine/Source/Programs/UnrealBuildAccelerator/Common/Public/UbaProcess.h b/Engine/Source/Programs/UnrealBuildAccelerator/Common/Public/UbaProcess.h
--- a/Engine/Source/Programs/UnrealBuildAccelerator/Common/Public/UbaProcess.h
+++ b/Engine/Source/Programs/UnrealBuildAccelerator/Common/Public/UbaProcess.h
@@ -176,8 +176,12 @@ namespace uba
 		ProcessStats m_processStats;
 		KernelStats m_kernelStats;
 
+		struct UsedFileMapping
+		{
+			bool closeAfterWritten;
+		};
 		Futex m_usedFileMappingsLock;
-		UnorderedSet<StringKey> m_usedFileMappings;
+		UnorderedMap<StringKey, UsedFileMapping> m_usedFileMappings;
 
 		Thread m_messageThread;
 
diff --git a/Engine/Source/Programs/UnrealBuildAccelerator/Common/Public/UbaSession.h b/Engine/Source/Programs/UnrealBuildAccelerator/Common/Public/UbaSession.h
--- a/Engine/Source/Programs/UnrealBuildAccelerator/Common/Public/UbaSession.h
+++ b/Engine/Source/Programs/UnrealBuildAccelerator/Common/Public/UbaSession.h
@@ -65,7 +65,7 @@ namespace uba
 		bool RefreshDirectory(const tchar* dirName, bool forceRegister = false); // Tell uba a directory on disk has been changed by some other system while session is running
 		bool RegisterNewFile(const tchar* filePath); // Tell uba a new file on disk has been added by some other system while session is running
 		bool RegisterVirtualFile(const tchar* filePath, const tchar* sourceFile, u64 sourceOffset, u64 sourceSize);
-		bool CreateVirtualFile(const tchar* filePath, const void* memory, u64 memorySize);
+		bool CreateVirtualFile(const tchar* filePath, const void* memory, u64 memorySize, bool transient);
 		bool DeleteVirtualFile(const tchar* filePath);
 		void RegisterDeleteFile(const tchar* filePath); // Tell uba a file on disk has been deleted by some other system while session is running
 		bool RegisterNewDirectory(const tchar* directoryPath); // Tell uba a directory on disk has been added by some other system while session is running
@@ -151,8 +151,8 @@ namespace uba
 		virtual void PrintSessionStats(Logger& logger);
 
 		virtual bool RegisterVirtualFileInternal(const StringKey& fileNameKey, const StringView& filePath, const tchar* sourceFile, u64 sourceOffset, u64 sourceSize);
-		virtual bool CreateVirtualFileInternal(const StringKey& fileNameKey, const StringView& filePath, const void* memory, u64 memorySize);
-		virtual bool RegisterVirtualFileInternal(const StringKey& fileNameKey, const StringView& filePath, FileMappingHandle mappingHandle, u64 mappingSize, u64 mappingOffset);
+		virtual bool CreateVirtualFileInternal(const StringKey& fileNameKey, const StringView& filePath, const void* memory, u64 memorySize, bool transient);
+		virtual bool RegisterVirtualFileInternal(const StringKey& fileNameKey, const StringView& filePath, FileMappingHandle mappingHandle, u64 mappingSize, u64 mappingOffset, bool transient);
 		virtual bool GetOutputFileSizeInternal(u64& outSize, const StringKey& fileNameKey, StringView filePath);
 		virtual bool GetOutputFileDataInternal(void* outData, const StringKey& fileNameKey, StringView filePath, bool deleteInternalMapping);
 		virtual bool WriteOutputFileInternal(const StringKey& fileNameKey, StringView filePath, bool deleteInternalMapping);
diff --git a/Engine/Source/Programs/UnrealBuildAccelerator/Common/Public/UbaSessionClient.h b/Engine/Source/Programs/UnrealBuildAccelerator/Common/Public/UbaSessionClient.h
--- a/Engine/Source/Programs/UnrealBuildAccelerator/Common/Public/UbaSessionClient.h
+++ b/Engine/Source/Programs/UnrealBuildAccelerator/Common/Public/UbaSessionClient.h
@@ -11,8 +11,7 @@ namespace uba
 
 	struct SessionClientCreateInfo : SessionCreateInfo
 	{
-		SessionClientCreateInfo(Storage& s, NetworkClient& c, LogWriter& writer = g_consoleLogWriter) : SessionCreateInfo(s, writer), client(c) {}
-
+		SessionClientCreateInfo(Storage& s, NetworkClient& c, LogWriter& writer = g_consoleLogWriter);
 		void Apply(const Config& config);
 
 		NetworkClient& client;
diff --git a/Engine/Source/Programs/UnrealBuildAccelerator/Common/Public/UbaStorage.h b/Engine/Source/Programs/UnrealBuildAccelerator/Common/Public/UbaStorage.h
--- a/Engine/Source/Programs/UnrealBuildAccelerator/Common/Public/UbaStorage.h
+++ b/Engine/Source/Programs/UnrealBuildAccelerator/Common/Public/UbaStorage.h
@@ -113,6 +113,7 @@ namespace uba
 		bool writeToDisk = true;
 		bool keepTransientDataMapped = true;	// Will keep the transient data mapped in instead of doing map/unmap everytime we access it
 		bool allowDeleteVerified = false;
+		bool dropIndependentMappingsOnCopy = true; // Will drop independent mappings once they have been copyorlink to destination
 		MutexHandle exclusiveMutex = InvalidMutexHandle; // Set this to hand over pre-exclusive access to storage. Storage will release mutex on shutdown
 		u8 casCompressor = 0;
 		u8 casCompressionLevel = 0;
@@ -283,6 +284,7 @@ namespace uba
 		bool m_asyncUnmapViewOfFile = true;
 		bool m_allowDeleteVerified = false;
 		bool m_writeToDisk;
+		bool m_dropIndependentMappingsOnCopy = true;
 
 		MutexHandle m_exclusiveMutex = InvalidMutexHandle;
 
diff --git a/Engine/Source/Programs/UnrealBuildAccelerator/Core/Private/UbaApplicationRules.cpp b/Engine/Source/Programs/UnrealBuildAccelerator/Core/Private/UbaApplicationRules.cpp
--- a/Engine/Source/Programs/UnrealBuildAccelerator/Core/Private/UbaApplicationRules.cpp
+++ b/Engine/Source/Programs/UnrealBuildAccelerator/Core/Private/UbaApplicationRules.cpp
@@ -250,6 +250,11 @@ namespace uba
 				|| (fileName.EndsWith(TCV(".manifest")) && fileName.Contains(systemTemp)) // Used between lld-link.exe and mt.exe (test compiling AutoRTFMEngineTests)
 				;
 		}
+
+		virtual bool CloseFileMappingAfterUse(StringView fileName) const override
+		{
+			return true;
+		}
 	};
 
 	class ApplicationRulesRadLinkExe : public ApplicationRulesVcLink
@@ -394,6 +399,11 @@ namespace uba
 		{
 			return true;
 		}
+
+		virtual bool CloseFileMappingAfterUse(StringView fileName) const override
+		{
+			return true;
+		}
 	};
 
 	class ApplicationRulesLlvmObjCopyExe : public ApplicationRulesClang
@@ -463,6 +473,11 @@ namespace uba
 		{
 			return true;
 		}
+
+		virtual bool CloseFileMappingAfterUse(StringView fileName) const override
+		{
+			return true;
+		}
 	};
 
 	class ApplicationRulesClangPlusPlusExePlatform2 : public ApplicationRulesClangPlusPlusExe
@@ -505,6 +520,11 @@ namespace uba
 		{
 			return true;
 		}
+
+		virtual bool CloseFileMappingAfterUse(StringView fileName) const override
+		{
+			return true;
+		}
 	};
 
 	class ApplicationRulesStageSplitMergeApplication : public ApplicationRules
diff --git a/Engine/Source/Programs/UnrealBuildAccelerator/Core/Public/UbaApplicationRules.h b/Engine/Source/Programs/UnrealBuildAccelerator/Core/Public/UbaApplicationRules.h
--- a/Engine/Source/Programs/UnrealBuildAccelerator/Core/Public/UbaApplicationRules.h
+++ b/Engine/Source/Programs/UnrealBuildAccelerator/Core/Public/UbaApplicationRules.h
@@ -185,6 +185,14 @@ namespace uba
 			return false;
 		}
 
+		// Return true if process should be allowed to close file mappings of files after used.
+		// For example, when linker is done with obj files they are likely never used again so we can unmap them.
+		// .. but if we have both .lib and .dll creation we might end up unmap just to remap obj files (and if they are compressed that means decompressing twice)
+		virtual bool CloseFileMappingAfterUse(StringView fileName) const
+		{
+			return false;
+		}
+
 		u32 index = ~0u;
 	};
 
diff --git a/Engine/Source/Programs/UnrealBuildAccelerator/Host/Private/UbaExports.cpp b/Engine/Source/Programs/UnrealBuildAccelerator/Host/Private/UbaExports.cpp
--- a/Engine/Source/Programs/UnrealBuildAccelerator/Host/Private/UbaExports.cpp
+++ b/Engine/Source/Programs/UnrealBuildAccelerator/Host/Private/UbaExports.cpp
@@ -89,20 +89,23 @@ namespace uba
 
 	Config& GetConfig(const tchar* fileName = nullptr)
 	{
-		static Config config;
-		if (config.IsLoaded())
-			return config;
 		LoggerWithWriter logger(g_nullLogWriter);
-		StringBuffer<> temp;
+
 		if (!fileName || !*fileName)
 		{
+			static Config config;
+			if (config.IsLoaded())
+				return config;
+			StringBuffer<> temp;
 			GetDirectoryOfCurrentModule(logger, temp);
 			temp.EnsureEndsWithSlash().Append(TCV("UbaHost.toml"));
-			fileName = temp.data;
+			config.LoadFromFile(logger, temp.data);
+			return config;
 		}
 
-		config.LoadFromFile(logger, fileName);
-		return config;
+		auto config = new Config();
+		config->LoadFromFile(logger, fileName);
+		return *config;
 	}
 
 	struct ExportsDowngradedLogger : public LoggerWithWriter
@@ -336,6 +339,12 @@ extern "C"
 		return StorageServer_InternalCreate(info);
 	}
 
+	void StorageServer_PreloadCasTable(uba::StorageServer* storageServer)
+	{
+		if (!storageServer->IsCasTableLoaded())
+			storageServer->LoadCasTable();
+	}
+
 	void StorageServer_Destroy(uba::StorageServer* storageServer)
 	{
 		delete storageServer;
@@ -498,9 +507,9 @@ extern "C"
 	{
 		return server->RegisterVirtualFile(filename, sourceFile, sourceOffset, sourceSize);
 	}
-	bool SessionServer_CreateVirtualFile(uba::SessionServer* server, const uba::tchar* filename, const void* memory, uba::u64 memorySize)
+	bool SessionServer_CreateVirtualFile(uba::SessionServer* server, const uba::tchar* filename, const void* memory, uba::u64 memorySize, bool transient)
 	{
-		return server->CreateVirtualFile(filename, memory, memorySize);
+		return server->CreateVirtualFile(filename, memory, memorySize, transient);
 	}
 	bool SessionServer_DeleteVirtualFile(uba::SessionServer* server, const uba::tchar* filename)
 	{
diff --git a/Engine/Source/Programs/UnrealBuildAccelerator/Host/Public/UbaExports.h b/Engine/Source/Programs/UnrealBuildAccelerator/Host/Public/UbaExports.h
--- a/Engine/Source/Programs/UnrealBuildAccelerator/Host/Public/UbaExports.h
+++ b/Engine/Source/Programs/UnrealBuildAccelerator/Host/Public/UbaExports.h
@@ -72,6 +72,7 @@ extern "C"
 	// StorageServer
 	UBA_API uba::StorageServer* StorageServer_Create(uba::NetworkServer& server, const uba::tchar* rootDir, uba::u64 casCapacityBytes, bool storeCompressed, uba::LogWriter& writer = uba::g_consoleLogWriter, const uba::tchar* zone = TC(""));
 	UBA_API uba::StorageServer* StorageServer_Create2(uba::NetworkServer& server, const uba::Config& config, uba::LogWriter& writer = uba::g_consoleLogWriter);
+	UBA_API void StorageServer_PreloadCasTable(uba::StorageServer* storageServer);
 	UBA_API void StorageServer_Destroy(uba::StorageServer* storageServer);
 	UBA_API void StorageServer_SaveCasTable(uba::StorageServer* storageServer);
 	UBA_API void StorageServer_RegisterDisallowedPath(uba::StorageServer* storageServer, const uba::tchar* path);
@@ -121,7 +122,7 @@ extern "C"
 	UBA_API void SessionServer_RegisterDeleteFile(uba::SessionServer* server, const uba::tchar* filePath);
 	UBA_API bool SessionServer_RegisterNewDirectory(uba::SessionServer* server, const uba::tchar* directoryPath);
 	UBA_API bool SessionServer_RegisterVirtualFile(uba::SessionServer* server, const uba::tchar* filename, const uba::tchar* sourceFile, uba::u64 sourceOffset, uba::u64 sourceSize);
-	UBA_API bool SessionServer_CreateVirtualFile(uba::SessionServer* server, const uba::tchar* filename, const void* memory, uba::u64 memorySize);
+	UBA_API bool SessionServer_CreateVirtualFile(uba::SessionServer* server, const uba::tchar* filename, const void* memory, uba::u64 memorySize, bool transient);
 	UBA_API bool SessionServer_DeleteVirtualFile(uba::SessionServer* server, const uba::tchar* filename);
 
 	UBA_API bool SessionServer_GetOutputFileSize(uba::SessionServer* server, uba::u64& outSize, const uba::tchar* filename);
