[UBA]
* Fixed so session test functions can take in extra data
* Added test for transient virtual files (Not added to file mapping table)

diff --git a/Engine/Source/Programs/UnrealBuildAccelerator/Test/Private/UbaTest.h b/Engine/Source/Programs/UnrealBuildAccelerator/Test/Private/UbaTest.h
--- a/Engine/Source/Programs/UnrealBuildAccelerator/Test/Private/UbaTest.h
+++ b/Engine/Source/Programs/UnrealBuildAccelerator/Test/Private/UbaTest.h
@@ -24,8 +24,8 @@ namespace uba
 	bool FileExists(Logger& logger, StringView dir, StringView fileName);
 
 	using RunProcessFunction = Function<ProcessHandle(const ProcessStartInfo&)>;
-	using TestSessionFunction = Function<bool(LoggerWithWriter& logger, SessionServer& session, const tchar* workingDir, const RunProcessFunction& runProcess)>;
-	bool RunLocal(LoggerWithWriter& logger, const StringBufferBase& testRootDir, const TestSessionFunction& testFunc, bool enableDetour = true);
+	using TestSessionFunction = Function<bool(LoggerWithWriter& logger, SessionServer& session, const tchar* workingDir, const RunProcessFunction& runProcess, void* extraData)>;
+	bool RunLocal(LoggerWithWriter& logger, const StringBufferBase& testRootDir, const TestSessionFunction& testFunc, void* extraData = nullptr, bool enableDetour = true);
 	bool RunRemote(LoggerWithWriter& logger, const StringBufferBase& testRootDir, const TestSessionFunction& testFunc, bool deleteAll = true, bool serverShouldListen = true);
 	void GetTestAppPath(LoggerWithWriter& logger, StringBufferBase& out);
 	const tchar* GetSystemApplication();
diff --git a/Engine/Source/Programs/UnrealBuildAccelerator/Test/Private/UbaTestScheduler.cpp b/Engine/Source/Programs/UnrealBuildAccelerator/Test/Private/UbaTestScheduler.cpp
--- a/Engine/Source/Programs/UnrealBuildAccelerator/Test/Private/UbaTestScheduler.cpp
+++ b/Engine/Source/Programs/UnrealBuildAccelerator/Test/Private/UbaTestScheduler.cpp
@@ -10,7 +10,7 @@ namespace uba
 {
 	bool TestLocalSchedule(LoggerWithWriter& logger, const StringBufferBase& testRootDir)
 	{
-		return RunLocal(logger, testRootDir, [](LoggerWithWriter& logger, SessionServer& session, const tchar* workingDir, const RunProcessFunction& runProcess)
+		return RunLocal(logger, testRootDir, [](LoggerWithWriter& logger, SessionServer& session, const tchar* workingDir, const RunProcessFunction& runProcess, void* extraData)
 			{
 				SchedulerCreateInfo info(session);
 				Scheduler scheduler(info);
@@ -34,7 +34,7 @@ namespace uba
 
 	bool TestLocalScheduleReuse(LoggerWithWriter& logger, const StringBufferBase& testRootDir)
 	{
-		return RunLocal(logger, testRootDir, [](LoggerWithWriter& logger, SessionServer& session, const tchar* workingDir, const RunProcessFunction& runProcess)
+		return RunLocal(logger, testRootDir, [](LoggerWithWriter& logger, SessionServer& session, const tchar* workingDir, const RunProcessFunction& runProcess, void* extraData)
 			{
 				SchedulerCreateInfo info(session);
 				info.enableProcessReuse = true;
@@ -62,7 +62,7 @@ namespace uba
 
 	bool TestRemoteScheduleReuse(LoggerWithWriter& logger, const StringBufferBase& testRootDir)
 	{
-		return RunRemote(logger, testRootDir, [](LoggerWithWriter& logger, SessionServer& session, const tchar* workingDir, const RunProcessFunction& runProcess)
+		return RunRemote(logger, testRootDir, [](LoggerWithWriter& logger, SessionServer& session, const tchar* workingDir, const RunProcessFunction& runProcess, void* extraData)
 			{
 				SchedulerCreateInfo info(session);
 				info.enableProcessReuse = true;
diff --git a/Engine/Source/Programs/UnrealBuildAccelerator/Test/Private/UbaTestSession.cpp b/Engine/Source/Programs/UnrealBuildAccelerator/Test/Private/UbaTestSession.cpp
--- a/Engine/Source/Programs/UnrealBuildAccelerator/Test/Private/UbaTestSession.cpp
+++ b/Engine/Source/Programs/UnrealBuildAccelerator/Test/Private/UbaTestSession.cpp
@@ -21,7 +21,7 @@ namespace uba
 		return 10'000;
 	}
 
-	bool RunLocal(LoggerWithWriter& logger, const StringBufferBase& testRootDir, const Config& config, const TestSessionFunction& testFunc, bool enableDetour = true)
+	bool RunLocal(LoggerWithWriter& logger, const StringBufferBase& testRootDir, const Config& config, const TestSessionFunction& testFunc, void* extraData = nullptr, bool enableDetour = true)
 	{
 		LogWriter& logWriter = logger.m_writer;
 
@@ -58,12 +58,12 @@ namespace uba
 		if (!DeleteAllFiles(logger, workingDir.data, false))
 			return false;
 		workingDir.EnsureEndsWithSlash();
-		return testFunc(logger, session, workingDir.data, [&](const ProcessStartInfo& pi) { return session.RunProcess(pi, true, enableDetour); });
+		return testFunc(logger, session, workingDir.data, [&](const ProcessStartInfo& pi) { return session.RunProcess(pi, true, enableDetour); }, extraData);
 	}
 
-	bool RunLocal(LoggerWithWriter& logger, const StringBufferBase& testRootDir, const TestSessionFunction& testFunc, bool enableDetour)
+	bool RunLocal(LoggerWithWriter& logger, const StringBufferBase& testRootDir, const TestSessionFunction& testFunc, void* extraData, bool enableDetour)
 	{
-		return RunLocal(logger, testRootDir, {}, testFunc, enableDetour);
+		return RunLocal(logger, testRootDir, {}, testFunc, extraData, enableDetour);
 	}
 
 	void GetTestAppPath(LoggerWithWriter& logger, StringBufferBase& out)
@@ -243,18 +243,18 @@ namespace uba
 			});
 	}
 
-	bool RunRemote(LoggerWithWriter& logger, const StringBufferBase& testRootDir, const Config& serverConfig, const TestSessionFunction& testFunc, bool deleteAll = true, bool serverShouldListen = true)
+	bool RunRemote(LoggerWithWriter& logger, const StringBufferBase& testRootDir, const Config& serverConfig, const TestSessionFunction& testFunc, void* extraData = nullptr, bool deleteAll = true, bool serverShouldListen = true)
 	{
 		return SetupServerClientSession(logger, testRootDir, deleteAll, serverShouldListen, serverConfig,
 			[&](LoggerWithWriter& logger, const StringView& workingDir, SessionServer& sessionServer, SessionClient& sessionClient)
 			{
-				return testFunc(logger, sessionServer, workingDir.data, [&](const ProcessStartInfo& pi) { return sessionServer.RunProcessRemote(pi); });
+				return testFunc(logger, sessionServer, workingDir.data, [&](const ProcessStartInfo& pi) { return sessionServer.RunProcessRemote(pi); }, extraData);
 			});
 	}
 
 	bool RunRemote(LoggerWithWriter& logger, const StringBufferBase& testRootDir, const TestSessionFunction& testFunc, bool deleteAll, bool serverShouldListen)
 	{
-		return RunRemote(logger, testRootDir, {}, testFunc, deleteAll, serverShouldListen);
+		return RunRemote(logger, testRootDir, {}, testFunc, nullptr, deleteAll, serverShouldListen);
 	}
 	bool CreateTextFile(StringBufferBase& outPath, LoggerWithWriter& logger, const tchar* workingDir, const tchar* fileName, const char* text)
 	{
@@ -266,7 +266,7 @@ namespace uba
 		return fr.Close();
 	}
 
-	bool RunTestAppTests(LoggerWithWriter& logger, SessionServer& session, const tchar* workingDir, const RunProcessFunction& runProcess)
+	bool RunTestAppTests(LoggerWithWriter& logger, SessionServer& session, const tchar* workingDir, const RunProcessFunction& runProcess, void* extraData)
 	{
 		StringBuffer<MaxPath> fileR;
 		if (!CreateTextFile(fileR, logger, workingDir, TC("FileR.h"), "Foo"))
@@ -333,7 +333,7 @@ namespace uba
 	}
 #endif
 
-	bool RunClang(LoggerWithWriter& logger, SessionServer& session, const tchar* workingDir, const RunProcessFunction& runProcess)
+	bool RunClang(LoggerWithWriter& logger, SessionServer& session, const tchar* workingDir, const RunProcessFunction& runProcess, void* extraData)
 	{
 		StringBuffer<MaxPath> sourceFile;
 		sourceFile.Append(workingDir).Append(TCV("Code.cpp"));
@@ -389,7 +389,7 @@ namespace uba
 		return true;
 	}
 
-	bool RunCustomService(LoggerWithWriter& logger, SessionServer& session, const tchar* workingDir, const RunProcessFunction& runProcess)
+	bool RunCustomService(LoggerWithWriter& logger, SessionServer& session, const tchar* workingDir, const RunProcessFunction& runProcess, void* extraData)
 	{
 		bool gotMessage = false;
 		session.RegisterCustomService([&](uba::Process& process, const void* recv, u32 recvSize, void* send, u32 sendCapacity)
@@ -451,7 +451,7 @@ namespace uba
 
 	bool TestDetouredTouch(LoggerWithWriter& logger, const StringBufferBase& testRootDir)
 	{
-		return RunLocal(logger, testRootDir, [](LoggerWithWriter& logger, SessionServer& session, const tchar* workingDir, const RunProcessFunction& runProcess)
+		return RunLocal(logger, testRootDir, [](LoggerWithWriter& logger, SessionServer& session, const tchar* workingDir, const RunProcessFunction& runProcess, void* extraData)
 			{
 				StringBuffer<> file;
 				file.Append(workingDir).Append(TCV("TouchFile.h"));
@@ -555,7 +555,7 @@ namespace uba
 
 	bool TestMultipleDetouredProcesses(LoggerWithWriter& logger, const StringBufferBase& testRootDir)
 	{
-		return RunLocal(logger, testRootDir, [](LoggerWithWriter& logger, SessionServer& session, const tchar* workingDir, const RunProcessFunction& runProcess)
+		return RunLocal(logger, testRootDir, [](LoggerWithWriter& logger, SessionServer& session, const tchar* workingDir, const RunProcessFunction& runProcess, void* extraData)
 			{
 				ProcessStartInfo processInfo;
 				processInfo.application = GetSystemApplication();
@@ -580,7 +580,7 @@ namespace uba
 			});
 	}
 
-	bool RunSystemApplicationAndLookForLog(LoggerWithWriter& logger, SessionServer& session, const tchar* workingDir, const RunProcessFunction& runProcess)
+	bool RunSystemApplicationAndLookForLog(LoggerWithWriter& logger, SessionServer& session, const tchar* workingDir, const RunProcessFunction& runProcess, void* extraData)
 	{
 		ProcessStartInfo processInfo;
 		processInfo.application = GetSystemApplication();
@@ -613,10 +613,10 @@ namespace uba
 
 	bool TestLogLinesNoDetour(LoggerWithWriter& logger, const StringBufferBase& testRootDir)
 	{
-		return RunLocal(logger, testRootDir, RunSystemApplicationAndLookForLog, false);
+		return RunLocal(logger, testRootDir, RunSystemApplicationAndLookForLog, nullptr, false);
 	}
 
-	bool CheckAttributes(LoggerWithWriter& logger, SessionServer& session, const tchar* workingDir, const RunProcessFunction& runProcess)
+	bool CheckAttributes(LoggerWithWriter& logger, SessionServer& session, const tchar* workingDir, const RunProcessFunction& runProcess, void* extraData)
 	{
 		StringBuffer<MaxPath> testApp;
 		GetTestAppPath(logger, testApp);
@@ -702,7 +702,7 @@ namespace uba
 
 	bool TestSharedReservedMemory(LoggerWithWriter& logger, const StringBufferBase& testRootDir)
 	{
-		return RunLocal(logger, testRootDir, [](LoggerWithWriter& logger, SessionServer& session, const tchar* workingDir, const RunProcessFunction& runProcess)
+		return RunLocal(logger, testRootDir, [](LoggerWithWriter& logger, SessionServer& session, const tchar* workingDir, const RunProcessFunction& runProcess, void* extraData)
 			{
 				StringBuffer<MaxPath> testApp;
 				GetTestAppPath(logger, testApp);
@@ -754,11 +754,13 @@ namespace uba
 #endif
 	}
 
-	bool RunVirtualFileTest(LoggerWithWriter& logger, SessionServer& session, const tchar* workingDir, const RunProcessFunction& runProcess)
+	bool RunVirtualFileTest(LoggerWithWriter& logger, SessionServer& session, const tchar* workingDir, const RunProcessFunction& runProcess, void* extraData)
 	{
+		bool transient = *(bool*)extraData;
+
 		StringBuffer<> inFile(workingDir);
 		inFile.Append(TCV("VirtualFile.in"));
-		if (!session.CreateVirtualFile(inFile.data, "FOO", 3))
+		if (!session.CreateVirtualFile(inFile.data, "FOO", 3, transient))
 			return false;
 		ProcessHandle ph;
 		if (!RunTestApp(logger, workingDir, runProcess, TC("-virtualFile"), &ph))
@@ -809,7 +811,13 @@ namespace uba
 	{
 		Config config;
 		config.AddTable(TC("Session")).AddValue(TC("ShouldWriteToDisk"), false);
-		return RunLocal(logger, testRootDir, config, RunVirtualFileTest);
+		bool transient = true;
+		if (!RunLocal(logger, testRootDir, config, RunVirtualFileTest, &transient))
+			return false;
+		transient = false;
+		if (!RunLocal(logger, testRootDir, config, RunVirtualFileTest, &transient))
+			return false;
+		return true;
 	}
 
 	bool TestRemoteVirtualFile(LoggerWithWriter& logger, const StringBufferBase& testRootDir)
@@ -817,7 +825,13 @@ namespace uba
 		Config config;
 		config.AddTable(TC("Session")).AddValue(TC("ShouldWriteToDisk"), false);
 		config.AddTable(TC("Storage")).AddValue(TC("CreateIndependentMappings"), true);
-		return RunRemote(logger, testRootDir, config, RunVirtualFileTest);
+		bool transient = true;
+		if (!RunRemote(logger, testRootDir, config, RunVirtualFileTest, &transient))
+			return false;
+		transient = false;
+		if (!RunRemote(logger, testRootDir, config, RunVirtualFileTest, &transient))
+			return false;
+		return true;
 	}
 
 	#if PLATFORM_MAC
