[UBA]
* Fixed so wait entries are cleared with success = false if client that is in the middle of store is disconnecting

diff --git a/Engine/Source/Programs/UnrealBuildAccelerator/Common/Private/UbaStorageServer.cpp b/Engine/Source/Programs/UnrealBuildAccelerator/Common/Private/UbaStorageServer.cpp
--- a/Engine/Source/Programs/UnrealBuildAccelerator/Common/Private/UbaStorageServer.cpp
+++ b/Engine/Source/Programs/UnrealBuildAccelerator/Common/Private/UbaStorageServer.cpp
@@ -138,7 +138,7 @@ namespace uba
 					return m_logger.Info(TC("Timed out waiting %s for cas %s to be transferred from remote to storage (%s)"), TimeToText(waited).str, CasKeyString(casKey).str, hint).ToFalse();
 			}
 		}
-		return waitEntry.Success;
+		return waitEntry.success;
 	}
 
 	void StorageServer::ActiveFetch::Release(StorageServer& server, const tchar* reason)
@@ -237,6 +237,8 @@ namespace uba
 				store.casEntry->beingWritten = false;
 				if (m_traceStore)
 					m_trace->FileStoreEnd(clientId, store.casEntry->key);
+
+				UpdateWaitEntries(store.casEntry->key, false);
 			}
 		}
 		{
@@ -266,6 +268,17 @@ namespace uba
 		}
 	}
 
+	void StorageServer::UpdateWaitEntries(const CasKey& casKey, bool success)
+	{
+		SCOPED_FUTEX(m_waitEntriesLock, waitLock);
+		auto waitFindIt = m_waitEntries.find(casKey);
+		if (waitFindIt == m_waitEntries.end())
+			return;
+		WaitEntry& waitEntry = waitFindIt->second;
+		waitEntry.success = success;
+		waitEntry.done.Set();
+	}
+
 	bool StorageServer::IsDisallowedPath(const tchar* fileName)
 	{
 		for (auto& path : m_disallowedPaths)
@@ -461,6 +474,9 @@ namespace uba
 
 			case StorageMessageType_FetchBegin:
 			{
+				CasKey casKey;
+				StringBuffer<> hint;
+
 				// TODO: Remove this when we have tracked down the issue where clients time out
 				u32 todoRemoveMeIndex = 0;
 				u64 todoRemoveMeLastTime = GetTime();
@@ -482,7 +498,7 @@ namespace uba
 						if (timeSpentMs > 8 * 60 * 1000)
 						{
 							// Took more than 5 minutes to respond
-							m_logger.Warning(TC("Took more than 8 minutes to respond to FetchBegin to client %u (Hang code %u).."), connectionInfo.GetId(), todoRemoveMeIndex);
+							m_logger.Warning(TC("Took more than 8 minutes to respond to FetchBegin to client %u (Hang code %u, cas %s, hint %s)"), connectionInfo.GetId(), todoRemoveMeIndex, CasKeyString(casKey).str, hint.data);
 						}
 					});
 
@@ -498,8 +514,7 @@ namespace uba
 				todoRemoveMe(2);
 
 				u64 start = GetTime();
-				CasKey casKey = reader.ReadCasKey();
-				StringBuffer<> hint;
+				casKey = reader.ReadCasKey();
 				reader.ReadString(hint);
 
 				workContext.tracker.AddHint(StringView(hint).GetFileName());
@@ -1056,15 +1071,7 @@ namespace uba
 					stats.recvCasBytesComp += activeStore.fileSize;
 					stats.recvCasBytesRaw += activeStore.actualSize;
 
-					SCOPED_FUTEX(m_waitEntriesLock, waitLock);
-					auto waitFindIt = m_waitEntries.find(casEntry.key);
-					if (waitFindIt != m_waitEntries.end())
-					{
-						WaitEntry& waitEntry = waitFindIt->second;
-						waitEntry.Success = true;
-						waitEntry.done.Set();
-					}
-					waitLock.Leave();
+					UpdateWaitEntries(casEntry.key, true);
 
 					if (!firstStore)
 					{
diff --git a/Engine/Source/Programs/UnrealBuildAccelerator/Common/Public/UbaStorageServer.h b/Engine/Source/Programs/UnrealBuildAccelerator/Common/Public/UbaStorageServer.h
--- a/Engine/Source/Programs/UnrealBuildAccelerator/Common/Public/UbaStorageServer.h
+++ b/Engine/Source/Programs/UnrealBuildAccelerator/Common/Public/UbaStorageServer.h
@@ -44,6 +44,7 @@ namespace uba
 		virtual bool HasProxy(u32 clientId) override;
 
 		void OnDisconnected(u32 clientId);
+		void UpdateWaitEntries(const CasKey& casKey, bool success);
 		bool HandleMessage(const ConnectionInfo& connectionInfo, const WorkContext& workContext, u8 messageType, BinaryReader& reader, BinaryWriter& writer);
 		bool WaitForWritten(CasEntry& casEntry, ScopedWriteLock& entryLock, const ConnectionInfo& connectionInfo, const tchar* hint);
 		bool VerifyExists(bool& outExists, CasEntry& casEntry, ScopedWriteLock& entryLock, const CasKey& casKey);
@@ -67,7 +68,7 @@ namespace uba
 		{
 			WaitEntry() : done(true) {}
 			Event done;
-			bool Success = false;
+			bool success = false;
 			u32 refCount = 0;
 		};
 		UnorderedMap<CasKey, WaitEntry> m_waitEntries;
