[UBA]
* Fixed posix compile error that must have been caused by change of toolchain

diff --git a/Engine/Source/Programs/UnrealBuildAccelerator/Detours/Private/PosixOS/UbaDetoursFunctionsPosix.cpp b/Engine/Source/Programs/UnrealBuildAccelerator/Detours/Private/PosixOS/UbaDetoursFunctionsPosix.cpp
--- a/Engine/Source/Programs/UnrealBuildAccelerator/Detours/Private/PosixOS/UbaDetoursFunctionsPosix.cpp
+++ b/Engine/Source/Programs/UnrealBuildAccelerator/Detours/Private/PosixOS/UbaDetoursFunctionsPosix.cpp
@@ -2482,9 +2482,10 @@ int Internal_execve(const char* pathname, char* const _Nullable argv[], char* co
 		TRUE_WRAPPER(_exit)(res);
 	}
 	else
+	{
 		TRUE_WRAPPER(exit)(res);
-
-	return -1;
+	}
+	//return -1;
 }
 
 UBA_EXPORT int UBA_WRAPPER(execve)(const char* pathname, char* const _Nullable argv[], char* const _Nullable envp[])
