[UBA]
* Fixed linux compile error in unit tests

diff --git a/Engine/Source/Programs/UnrealBuildAccelerator/Test/Private/UbaTestSession.cpp b/Engine/Source/Programs/UnrealBuildAccelerator/Test/Private/UbaTestSession.cpp
--- a/Engine/Source/Programs/UnrealBuildAccelerator/Test/Private/UbaTestSession.cpp
+++ b/Engine/Source/Programs/UnrealBuildAccelerator/Test/Private/UbaTestSession.cpp
@@ -488,7 +488,7 @@ namespace uba
 	bool TestDetouredPopen(LoggerWithWriter& logger, const StringBufferBase& testRootDir)
 	{
 		#if PLATFORM_LINUX
-		return RunLocal(logger, testRootDir, [](LoggerWithWriter& logger, SessionServer& session, const tchar* workingDir, const RunProcessFunction& runProcess)
+		return RunLocal(logger, testRootDir, [](LoggerWithWriter& logger, SessionServer& session, const tchar* workingDir, const RunProcessFunction& runProcess, void* extraData)
 			{
 				StringBuffer<> testApp;
 				GetTestAppPath(logger, testApp);
