[UBA]
* Disabled unreachable code warning for UbaMimalloc since mimalloc code cause warnings

diff --git a/Engine/Source/Programs/UnrealBuildAccelerator/Mimalloc/UbaMimalloc.Build.cs b/Engine/Source/Programs/UnrealBuildAccelerator/Mimalloc/UbaMimalloc.Build.cs
--- a/Engine/Source/Programs/UnrealBuildAccelerator/Mimalloc/UbaMimalloc.Build.cs
+++ b/Engine/Source/Programs/UnrealBuildAccelerator/Mimalloc/UbaMimalloc.Build.cs
@@ -10,6 +10,7 @@ public class UbaMimalloc : ModuleRules
 	public UbaMimalloc(ReadOnlyTargetRules Target) : base(Target)
 	{
 		PCHUsage = PCHUsageMode.NoPCHs;
+		CppCompileWarningSettings.UnreachableCodeWarningLevel = WarningLevel.Off;
 
 		if (Target.Platform.IsInGroup(UnrealPlatformGroup.Windows) || (!Target.bShouldCompileAsDLL && Target.Platform.IsInGroup(UnrealPlatformGroup.Linux)))
 		{
