[UBA]
* Added more info to when memcpy inside WriteFile cause an access violation
* Changed returned error code when failing to create file

diff --git a/Engine/Source/Programs/UnrealBuildAccelerator/Detours/Private/Windows/UbaDetoursFunctionsKernelBase.inl b/Engine/Source/Programs/UnrealBuildAccelerator/Detours/Private/Windows/UbaDetoursFunctionsKernelBase.inl
--- a/Engine/Source/Programs/UnrealBuildAccelerator/Detours/Private/Windows/UbaDetoursFunctionsKernelBase.inl
+++ b/Engine/Source/Programs/UnrealBuildAccelerator/Detours/Private/Windows/UbaDetoursFunctionsKernelBase.inl
@@ -1,5 +1,7 @@
 // Copyright Epic Games, Inc. All Rights Reserved.
 
+#define UBA_FIND_MEMCPY_ACCESS_VIOLATION 1
+
 DWORD Local_GetLongPathNameW(LPCWSTR lpszShortPath, LPWSTR lpszLongPath, DWORD cchBuffer)
 {
 	SCOPED_WRITE_LOCK(g_longPathNameCacheLock, lock);
@@ -654,6 +656,29 @@ HANDLE Detoured_CreateToolhelp32Snapshot(DWORD dwFlags, DWORD th32ProcessID)
 	//return True_CreateToolhelp32Snapshot(dwFlags, th32ProcessID);
 }
 
+#if UBA_FIND_MEMCPY_ACCESS_VIOLATION
+struct MemoryFileData
+{
+	u8* baseAddress = nullptr;
+	u64 reserveSize = 0;
+	u64 mappedSize = 0;
+	u64 committedSize = 0;
+	u64 writtenSize = 0;
+};
+int LogAccessViolation(PEXCEPTION_POINTERS ep, MemoryFile& mf, MemoryFileData& data, HANDLE hFile, DWORD nNumberOfBytesToWrite)
+{
+	if (ep->ExceptionRecord->ExceptionCode == EXCEPTION_ACCESS_VIOLATION)// && r->NumberParameters >= 2 && r->ExceptionInformation[0] == 1)
+	{
+		StringBuffer<> buf;
+		buf.Appendf(L"Failed to write %llu bytes to memory file %s\r\n", nNumberOfBytesToWrite, HandleToName(hFile));
+		buf.Appendf(L" Before %p %llu %llu %llu %llu\r\n", data.baseAddress, data.reserveSize, data.mappedSize, data.committedSize, data.writtenSize);
+		buf.Appendf(L"  After %p %llu %llu %llu %llu", mf.baseAddress, mf.reserveSize, mf.mappedSize, mf.committedSize, mf.writtenSize);
+		Rpc_WriteLog(buf.data, buf.count, true, true);
+	}
+	return EXCEPTION_CONTINUE_SEARCH;
+}
+#endif
+
 BOOL Detoured_WriteFile(HANDLE hFile, LPCVOID lpBuffer, DWORD nNumberOfBytesToWrite, LPDWORD lpNumberOfBytesWritten, LPOVERLAPPED lpOverlapped)
 {
 	DWORD temp = 0;
@@ -699,7 +724,29 @@ BOOL Detoured_WriteFile(HANDLE hFile, LPCVOID lpBuffer, DWORD nNumberOfBytesToWr
 					SetEvent(lpOverlapped->hEvent);
 			}
 			else
-				mf->Write(dh, lpBuffer, nNumberOfBytesToWrite);
+			{
+				SCOPED_WRITE_LOCK(mf->lock, lock);
+
+#if UBA_FIND_MEMCPY_ACCESS_VIOLATION
+				MemoryFileData data;
+				data.baseAddress = mf->baseAddress;
+				data.reserveSize = mf->reserveSize;
+				data.mappedSize = mf->mappedSize;
+				data.committedSize = mf->committedSize;
+				data.writtenSize = mf->writtenSize;
+
+				__try
+				{
+					//*(int*)nullptr = 123;
+#endif
+					mf->Write(dh, lpBuffer, nNumberOfBytesToWrite);
+#if UBA_FIND_MEMCPY_ACCESS_VIOLATION
+				}
+				__except(LogAccessViolation(GetExceptionInformation(), *mf, data, hFile, nNumberOfBytesToWrite))
+				{
+				}
+#endif
+			}
 
 			*lpNumberOfBytesWritten = nNumberOfBytesToWrite;
 			SetLastError(ERROR_SUCCESS);
diff --git a/Engine/Source/Programs/UnrealBuildAccelerator/Detours/Private/Windows/UbaDetoursFunctionsNtDll.inl b/Engine/Source/Programs/UnrealBuildAccelerator/Detours/Private/Windows/UbaDetoursFunctionsNtDll.inl
--- a/Engine/Source/Programs/UnrealBuildAccelerator/Detours/Private/Windows/UbaDetoursFunctionsNtDll.inl
+++ b/Engine/Source/Programs/UnrealBuildAccelerator/Detours/Private/Windows/UbaDetoursFunctionsNtDll.inl
@@ -1201,7 +1201,7 @@ NTSTATUS NTAPI Shared_NtCreateFile(bool IsCreateFunc, PHANDLE hFileHandle, ACCES
 		{
 			Rpc_WriteLogf(TC("Mapped file %s cant be open for write (%s)"), lpFileName, info.originalName);
 			UBA_ASSERTF(false, TC("Mapped file %s cant be open for write (%s)"), lpFileName, info.originalName);
-			return STATUS_UNSUCCESSFUL;
+			return STATUS_ACCESS_DENIED;
 		}
 
 		const wchar_t* handleStr = lpFileName + 1;
@@ -1224,7 +1224,7 @@ NTSTATUS NTAPI Shared_NtCreateFile(bool IsCreateFunc, PHANDLE hFileHandle, ACCES
 		{
 			Rpc_WriteLogf(L"Can't duplicate handle 0x%llx (%ls) for file %ls (Error %u)", uintptr_t(mappingHandle), lpFileName, info.originalName, GetLastError());
 			UBA_ASSERTF(info.trueFileMapHandle, L"Can't duplicate handle 0x%llx (%ls) for file %ls (Error %u)", uintptr_t(mappingHandle), lpFileName, info.originalName, GetLastError());
-			return STATUS_UNSUCCESSFUL;
+			return STATUS_ACCESS_DENIED;
 		}
 
 		lookupLock.Leave();
