[UBA]
* Fixed macos unit test compile errors

diff --git a/Engine/Source/Programs/UnrealBuildAccelerator/Test/Private/UbaTestSession.cpp b/Engine/Source/Programs/UnrealBuildAccelerator/Test/Private/UbaTestSession.cpp
--- a/Engine/Source/Programs/UnrealBuildAccelerator/Test/Private/UbaTestSession.cpp
+++ b/Engine/Source/Programs/UnrealBuildAccelerator/Test/Private/UbaTestSession.cpp
@@ -848,11 +848,11 @@ namespace uba
 	}
 	bool TestXcodeSelect(LoggerWithWriter& logger, const StringBufferBase& testRootDir)
 	{
-		return RunLocal(logger, testRootDir, RunXCodeSelect);
+		return RunLocal(logger, testRootDir, RunXCodeSelect, nullptr);
 	}
 	bool TestRemoteXcodeSelect(LoggerWithWriter& logger, const StringBufferBase& testRootDir)
 	{
-		return RunRemote(logger, testRootDir, RunXCodeSelect);
+		return RunRemote(logger, testRootDir, RunXCodeSelect, nullptr);
 	}
 	#endif
 
