[UBA]
* Fixed bug where default config was not used when instantiating network/storage/session instances

diff --git a/Engine/Source/Programs/UnrealBuildAccelerator/Host/Private/UbaExports.cpp b/Engine/Source/Programs/UnrealBuildAccelerator/Host/Private/UbaExports.cpp
--- a/Engine/Source/Programs/UnrealBuildAccelerator/Host/Private/UbaExports.cpp
+++ b/Engine/Source/Programs/UnrealBuildAccelerator/Host/Private/UbaExports.cpp
@@ -91,21 +91,29 @@ namespace uba
 	{
 		LoggerWithWriter logger(g_nullLogWriter);
 
-		if (!fileName || !*fileName)
+		static Config defaultConfig;
+
+		if (fileName && *fileName && !defaultConfig.IsLoaded() && Contains(fileName, TC("UbaHost.toml")))
+		{
+			defaultConfig.LoadFromFile(logger, fileName);
+			return defaultConfig;
+		}
+		else if (!fileName || !*fileName)
 		{
-			static Config config;
-			if (config.IsLoaded())
-				return config;
+			if (defaultConfig.IsLoaded())
+				return defaultConfig;
 			StringBuffer<> temp;
 			GetDirectoryOfCurrentModule(logger, temp);
 			temp.EnsureEndsWithSlash().Append(TCV("UbaHost.toml"));
-			config.LoadFromFile(logger, temp.data);
-			return config;
+			defaultConfig.LoadFromFile(logger, temp.data);
+			return defaultConfig;
+		}
+		else
+		{
+			auto config = new Config();
+			config->LoadFromFile(logger, fileName);
+			return *config;
 		}
-
-		auto config = new Config();
-		config->LoadFromFile(logger, fileName);
-		return *config;
 	}
 
 	struct ExportsDowngradedLogger : public LoggerWithWriter
@@ -214,11 +222,10 @@ extern "C"
 
 		NetworkServerCreateInfo info(writer);
 
-		info.Apply(GetConfig());
-
 		info.workerCount = workerCount;
 		info.sendSize = sendSize;
 		info.receiveTimeoutSeconds = receiveTimeoutSeconds;
+		info.Apply(GetConfig());
 
 		bool success = true;
 		auto server = new NetworkServerWithBackend(success, info, networkBackend);
