[UBA]
* Renamed scheduler memtrack to memwatchdog to prevent confusion with what it does
* Added separate option to enable tracking page file activity and kill if activity go crazy

diff --git a/Engine/Source/Programs/UnrealBuildAccelerator/Common/Private/UbaScheduler.cpp b/Engine/Source/Programs/UnrealBuildAccelerator/Common/Private/UbaScheduler.cpp
--- a/Engine/Source/Programs/UnrealBuildAccelerator/Common/Private/UbaScheduler.cpp
+++ b/Engine/Source/Programs/UnrealBuildAccelerator/Common/Private/UbaScheduler.cpp
@@ -156,7 +156,10 @@ namespace uba
 			table->GetValueAsU32(maxParallelCacheQueries, TC("MaxParallelCacheQueries"));
 			table->GetValueAsBool(useThreadToExitRemoteProcess, TC("UseThreadToExitRemoteProcess"));
 			
-			table->GetValueAsBool(trackMemory, TC("MemTrack"));
+			table->GetValueAsBool(memWatchdog, TC("MemTrack")); // Legacy
+
+			table->GetValueAsBool(memWatchdog, TC("MemWatchdog"));
+
 			u32 temp = 0;
 			if (table->GetValueAsU32(temp, TC("MemTracingLevel")))
 				memTracingLevel = u8(temp);
@@ -164,6 +167,7 @@ namespace uba
 				memStartWaitPercent = u8(temp);
 			if (table->GetValueAsU32(temp, TC("MemStartKillPercent")))
 				memStartKillPercent = u8(temp);
+			table->GetValueAsU32(memPagefileActivityThreshold, TC("MemPagefileActivityThreshold"));
 		}
 	}
 
@@ -187,12 +191,13 @@ namespace uba
 				return HandleReuseMessage(process, outNextProcess, prevExitCode);
 			});
 
-		if (info.trackMemory && IsWindows && !IsRunningWine()) // Only implemented for normal windows for now
+		if (info.memWatchdog && IsWindows && !IsRunningWine()) // Only implemented for normal windows for now
 		{
-			m_memTrack = true;
+			m_memWatchdog = true;
 			m_memTracingLevel = info.memTracingLevel;
 			m_memStartWaitPercent = info.memStartWaitPercent;
 			m_memStartKillPercent = info.memStartKillPercent;
+			m_memPagefileActivityThreshold = info.memPagefileActivityThreshold;
 		}
 
 		m_session.SetOuterScheduler(this);
@@ -213,7 +218,7 @@ namespace uba
 		if (!m_cacheClients.empty())
 			m_cacheWorkManager = new WorkManagerImpl(m_maxParallelCacheQueries, TC("UbaSchedCach"));
 
-		if (m_memTrack)
+		if (m_memWatchdog)
 		{
 			LoadMemTrackTable();
 			m_session.SetNativeProcessCreatedFunc([this](ProcessImpl& process) { NativeProcessCreated(process); });
@@ -251,7 +256,7 @@ namespace uba
 			m_session.WaitOnAllClients();
 		}
 
-		if (m_memTrack)
+		if (m_memWatchdog)
 		{
 			m_memThreadEvent.Set();
 			m_memThread.Wait();
@@ -1163,7 +1168,7 @@ namespace uba
 
 	void Scheduler::NativeProcessCreated(ProcessImpl& process)
 	{
-		if (!m_memTrack)
+		if (!m_memWatchdog)
 			return;
 
 		auto& impl = (ProcessImpl&)process;
@@ -1191,7 +1196,7 @@ namespace uba
 
 	void Scheduler::LocalProcessStart(ProcessStartInfo2& info)
 	{
-		if (!m_memTrack || !info.memoryGroupId)
+		if (!m_memWatchdog || !info.memoryGroupId)
 			return;
 
 		SCOPED_FUTEX(m_memGroupLookupLock, statsLock);
@@ -1260,7 +1265,7 @@ namespace uba
 
 	void Scheduler::LocalProcessExit(ProcessStartInfo2& info, const ProcessHandle& handle)
 	{
-		if (!m_memTrack)
+		if (!m_memWatchdog)
 			return;
 
 		auto& impl = *(ProcessImpl*)handle.m_process;
@@ -1428,7 +1433,7 @@ namespace uba
 		u64 maxPageFile;
 		if (!GetMemoryInfo(logger, memAvail, memTotal, &maxPageFile))
 		{
-			m_memTrack = false;
+			m_memWatchdog = false;
 			logger.Warning(TC("GetMemoryInfo failed. Memory tracking disabled."));
 			m_updateThreadWakeup.Set();
 			return;
@@ -1473,7 +1478,7 @@ namespace uba
 					continue;
 				}
 
-				m_memTrack = false;
+				m_memWatchdog = false;
 				logger.Warning(TC("NtQuerySystemInformation failed. Memory tracking disabled."));
 				m_updateThreadWakeup.Set();
 				break;
@@ -1483,7 +1488,7 @@ namespace uba
 			ULONG nextEntryOffset = *(ULONG*)it;
 			if (nextEntryOffset < PagefileUsageOffset + 8)
 			{
-				m_memTrack = false;
+				m_memWatchdog = false;
 				logger.Warning(TC("NtQuerySystemInformation does not contain PageFileUsage. Memory tracking disabled."));
 				m_updateThreadWakeup.Set();
 				break;
@@ -1574,43 +1579,46 @@ namespace uba
 			}
 
 			#if PLATFORM_WINDOWS
-			auto& currPerfInfo = *(UBA_SYSTEM_PERFORMANCE_INFORMATION*)queryBuffer.data();
-			if (NtQuerySystemInformation(SystemPerformanceInformation, &currPerfInfo, sizeof(SYSTEM_PERFORMANCE_INFORMATION), NULL) == STATUS_SUCCESS && !firstPerfInfo)
+			if (m_memPagefileActivityThreshold)
 			{
-				ULONG reads = currPerfInfo.PagesRead - prevPerfInfo.PagesRead;
-				ULONG writes = currPerfInfo.PagefilePagesWritten - prevPerfInfo.PagefilePagesWritten;
-
-				if (m_memTracingLevel >= 1)
+				auto& currPerfInfo = *(UBA_SYSTEM_PERFORMANCE_INFORMATION*)queryBuffer.data();
+				if (NtQuerySystemInformation(SystemPerformanceInformation, &currPerfInfo, sizeof(SYSTEM_PERFORMANCE_INFORMATION), NULL) == STATUS_SUCCESS && !firstPerfInfo)
 				{
-					trace.StatusUpdate(StartRow+6, 6, CountToText(reads), LogEntryType_Info);
-					trace.StatusUpdate(StartRow+7, 6, CountToText(writes), LogEntryType_Info);
-				}
+					ULONG reads = currPerfInfo.PagesRead - prevPerfInfo.PagesRead;
+					ULONG writes = currPerfInfo.PagefilePagesWritten - prevPerfInfo.PagefilePagesWritten;
 
-				if (memUsed > memPhys && activeProcessCount > 6 && writes > 1000)
-				{
-					memStartWaitPercent = u8(double(memUsed)/double(memTotal) * 100.0f);
-					str.Clear().Appendf(TC("Pagefile activity high (%lu writes). MemUsed %s, MemPhys %s, MemTot %s"), writes, BytesToText(memUsed).str, BytesToText(memPhys).str, BytesToText(memTotal).str);
-					KillNewestLocalProcess(str);
-					lastPageFileWriteTimeMs = TimeToMs(GetTime());
-				}
-				else if (writes > 0)
-				{
-					if (memStartWaitPercent > memStartWaitPercentMin)
-						--memStartWaitPercent;
-					lastPageFileWriteTimeMs = TimeToMs(GetTime());
-				}
-				else if (TimeToMs(GetTime()) - lastPageFileWriteTimeMs > 5*1000)
-				{
-					if (memStartWaitPercent < m_memStartWaitPercent)
+					if (m_memTracingLevel >= 1)
 					{
-						++memStartWaitPercent;
-						if (memStartWaitPercent == m_memStartWaitPercent)
-							m_updateThreadWakeup.Set(); // Just in case.. this should never happen in practice
+						trace.StatusUpdate(StartRow+6, 6, CountToText(reads), LogEntryType_Info);
+						trace.StatusUpdate(StartRow+7, 6, CountToText(writes), LogEntryType_Info);
+					}
+
+					if (memUsed > memPhys && activeProcessCount > 6 && writes > m_memPagefileActivityThreshold)
+					{
+						memStartWaitPercent = u8(double(memUsed)/double(memTotal) * 100.0f);
+						str.Clear().Appendf(TC("Pagefile activity high (%lu writes). MemUsed %s, MemPhys %s, MemTot %s"), writes, BytesToText(memUsed).str, BytesToText(memPhys).str, BytesToText(memTotal).str);
+						KillNewestLocalProcess(str);
+						lastPageFileWriteTimeMs = TimeToMs(GetTime());
+					}
+					else if (writes > 0)
+					{
+						if (memStartWaitPercent > memStartWaitPercentMin)
+							--memStartWaitPercent;
+						lastPageFileWriteTimeMs = TimeToMs(GetTime());
+					}
+					else if (TimeToMs(GetTime()) - lastPageFileWriteTimeMs > 5*1000)
+					{
+						if (memStartWaitPercent < m_memStartWaitPercent)
+						{
+							++memStartWaitPercent;
+							if (memStartWaitPercent == m_memStartWaitPercent)
+								m_updateThreadWakeup.Set(); // Just in case.. this should never happen in practice
+						}
 					}
 				}
+				firstPerfInfo = false;
+				prevPerfInfo = currPerfInfo;
 			}
-			firstPerfInfo = false;
-			prevPerfInfo = currPerfInfo;
 			#endif
 
 			//if (randomKillTime < GetTime())
diff --git a/Engine/Source/Programs/UnrealBuildAccelerator/Common/Public/UbaScheduler.h b/Engine/Source/Programs/UnrealBuildAccelerator/Common/Public/UbaScheduler.h
--- a/Engine/Source/Programs/UnrealBuildAccelerator/Common/Public/UbaScheduler.h
+++ b/Engine/Source/Programs/UnrealBuildAccelerator/Common/Public/UbaScheduler.h
@@ -37,10 +37,11 @@ namespace uba
 		bool forceNative = false; // Force all processes to run native (not detoured)
 		bool writeToCache = false; // Set to true in combination with setting cacheClient to populate cache
 		bool useThreadToExitRemoteProcess = false; // Set this to true to move exit process call from job to thread. 
-		bool trackMemory = false;
+		bool memWatchdog = false; // Will watch memory and throttle/kill processes if needed
 		u8 memTracingLevel = 0;
 		u8 memStartWaitPercent = 85; // When memory usage goes above this percent, no new processes will be spawned until back below
 		u8 memStartKillPercent = 95; // When memory usage goes above this percent, newest processes will be killed to bring it back below
+		u32 memPagefileActivityThreshold = 0; // When page file activity goes over this number between each watchdog update we start killing processes. 0=disabled
 		ConfigTable* processConfigs = nullptr;
 	};
 
@@ -211,7 +212,7 @@ namespace uba
 
 		CriticalSection m_wasReturnedLock;
 
-		bool m_memTrack = false;
+		bool m_memWatchdog = false;
 		Event m_memThreadEvent;
 		Thread m_memThread;
 		bool m_memDelaySpawning = false;
@@ -220,6 +221,7 @@ namespace uba
 		u8 m_memStartWaitPercent = 0;
 		u8 m_memStartKillPercent = 0;
 		u8 m_memTracingLevel = 0;
+		u32 m_memPagefileActivityThreshold = 0;
 
 		struct MemGroupStats;
 		struct MemGroupEntry;
