[UBA]
* Fixed macos compile error in unit tests

diff --git a/Engine/Source/Programs/UnrealBuildAccelerator/Test/Private/UbaTestSession.cpp b/Engine/Source/Programs/UnrealBuildAccelerator/Test/Private/UbaTestSession.cpp
--- a/Engine/Source/Programs/UnrealBuildAccelerator/Test/Private/UbaTestSession.cpp
+++ b/Engine/Source/Programs/UnrealBuildAccelerator/Test/Private/UbaTestSession.cpp
@@ -835,7 +835,7 @@ namespace uba
 	}
 
 	#if PLATFORM_MAC
-	bool RunXCodeSelect(LoggerWithWriter& logger, SessionServer& session, const tchar* workingDir, const RunProcessFunction& runProcess)
+	bool RunXCodeSelect(LoggerWithWriter& logger, SessionServer& session, const tchar* workingDir, const RunProcessFunction& runProcess, void* extraData)
 	{
 		#if UBA_DEBUG // Failing on farm for some reason... need to revisit
 		StringBuffer<> xcodeSelect;
@@ -848,11 +848,11 @@ namespace uba
 	}
 	bool TestXcodeSelect(LoggerWithWriter& logger, const StringBufferBase& testRootDir)
 	{
-		return RunLocal(logger, testRootDir, RunXCodeSelect, nullptr);
+		return RunLocal(logger, testRootDir, RunXCodeSelect);
 	}
 	bool TestRemoteXcodeSelect(LoggerWithWriter& logger, const StringBufferBase& testRootDir)
 	{
-		return RunRemote(logger, testRootDir, RunXCodeSelect, nullptr);
+		return RunRemote(logger, testRootDir, RunXCodeSelect);
 	}
 	#endif
 
