[UBA]
* Added more logging code around warning that there are still remote processes active while there are 0 connections. This should cause a hang but isn't so there must be a timing situation here where a disconnect happens at the same time as a process is finished/returned

diff --git a/Engine/Source/Programs/UnrealBuildAccelerator/Common/Private/UbaSessionServer.cpp b/Engine/Source/Programs/UnrealBuildAccelerator/Common/Private/UbaSessionServer.cpp
--- a/Engine/Source/Programs/UnrealBuildAccelerator/Common/Private/UbaSessionServer.cpp
+++ b/Engine/Source/Programs/UnrealBuildAccelerator/Common/Private/UbaSessionServer.cpp
@@ -772,12 +772,17 @@ namespace uba
 			}
 		}
 
-		if (!m_activeRemoteProcesses.empty())
-		{
-			// This path has been seen in the wild (Once over million of runs)....
-			// And the theory is that 
-			m_logger.Error(TC("No client sessions connected but there are %llu active remote processes. This should not happen, there is a bug in the code!!"), m_activeRemoteProcesses.size());
-		}
+		if (m_activeRemoteProcesses.empty())
+			return;
+
+		// This path has been seen in the wild (twice over million of runs)....
+		// It doesn't seem to have caused hangs so maybe there is a timing with disconnecting and getting results back.
+		// Let's wait 1 second and see again and log out.
+		m_logger.Warning(TC("No client sessions connected but there are %llu active remote processes. This should not happen"), m_activeRemoteProcesses.size());
+		queueLock.Leave();
+		Sleep(1000);
+		queueLock.Enter();
+		m_logger.Warning(TC("%llu active remote processes over %u connections"), m_activeRemoteProcesses.size(), m_connectionCount);
 	}
 
 	bool SessionServer::HandleConnect(const ConnectionInfo& connectionInfo, const WorkContext& workContext, BinaryReader& reader, BinaryWriter& writer)
