[UBA]
* Fixed extra obj file coff path for UbaObjTool. Was using wrong section for loopbacks, causing problems on non-windows-platform
* Removed unused code

diff --git a/Engine/Source/Programs/UnrealBuildAccelerator/Common/Private/UbaObjectFileCoff.cpp b/Engine/Source/Programs/UnrealBuildAccelerator/Common/Private/UbaObjectFileCoff.cpp
--- a/Engine/Source/Programs/UnrealBuildAccelerator/Common/Private/UbaObjectFileCoff.cpp
+++ b/Engine/Source/Programs/UnrealBuildAccelerator/Common/Private/UbaObjectFileCoff.cpp
@@ -12,6 +12,7 @@ namespace uba
 	constexpr u32 ImageSymUndefined = 0; // IMAGE_SYM_UNDEFINED
 
 	constexpr u32 ImageScnCntCode = 0x00000020; // IMAGE_SCN_CNT_CODE
+	constexpr u32 ImageScnCntInitializedData = 0x00000040; // IMAGE_SCN_CNT_INITIALIZED_DATA
 	constexpr u32 ImageScnLnkInfo = 0x00000200; // IMAGE_SCN_LNK_INFO
 	constexpr u32 ImageScnLnkRemove = 0x00000800; // IMAGE_SCN_LNK_REMOVE
 	constexpr u32 ImageScnLnkComdat = 0x00001000; // IMAGE_SCN_LNK_COMDAT
@@ -349,40 +350,37 @@ namespace uba
 		u32 loopbackCount = u32(neededLoopbacks.size());
 		UBA_ASSERT(loopbackCount < 65536);
 
-
 		auto allocate = [&](u64 size) { return memoryBlock.Allocate(size, 1, TC("")); };
 		auto write = [&](const void* data, u64 size) { memcpy(allocate(size), data, size); };
 
-
-		// Header
+		// File header (allocator returns zeroed memory)
 		auto& header = *(ImageFileHeader*)allocate(sizeof(ImageFileHeader));
 		header.Machine = 0x8664;
 
-		// Session for loopbacks
-		auto& textSection = *(ImageSectionHeader*)allocate(sizeof(ImageSectionHeader));
-		memcpy(textSection.Name, ".text$mn", 8);
-		textSection.Characteristics = ImageScnCntCode | ImageScnMemExecute | ImageScnMemRead;
-		u16 textSessionIndex = header.NumberOfSections;
+		// Section for loopback data
+		auto& rdataSection = *(ImageSectionHeader*)allocate(sizeof(ImageSectionHeader));
+		memcpy(rdataSection.Name, ".rdata\0\0", 8);
+		rdataSection.Characteristics = ImageScnCntInitializedData | ImageScnMemRead;
+		u16 rdataSectionIndex0 = header.NumberOfSections; // 0-based index
 		++header.NumberOfSections;
 
-		// Directive section
+		// Optional .drectve with exports
 		if (includeExportsInFile)
 		{
 			auto& directiveSection = *(ImageSectionHeader*)allocate(sizeof(ImageSectionHeader));
 			memcpy(directiveSection.Name, ".drectve", 8);
-			directiveSection.Characteristics = ImageScnAlign1Bytes|ImageScnLnkInfo|ImageScnLnkRemove;
+			directiveSection.Characteristics = ImageScnAlign1Bytes | ImageScnLnkInfo | ImageScnLnkRemove;
 			++header.NumberOfSections;
 
 			auto writeExport = [&](AnsiStringView symbol, bool isData)
-				{
-					write("/EXPORT:", 8);
-					write(symbol.strBegin, symbol.Length());
-					if (isData)
-						write(",DATA", 5);
-					write(" ", 1);
-				};
-
-			// Directive raw data
+			{
+				write("/EXPORT:", 8);
+				write(symbol.strBegin, symbol.Length());
+				if (isData)
+					write(",DATA", 5);
+				write(" ", 1);
+			};
+
 			u32 directiveRawDataStart = u32(memoryBlock.writtenSize);
 			directiveSection.PointerToRawData = directiveRawDataStart;
 
@@ -401,182 +399,76 @@ namespace uba
 			tmp = "ThisIsAnUnrealEngineModule";
 			StringKey isUnrealKey = ToStringKeyRaw(tmp.data(), tmp.size());
 			if (allExports.find(isUnrealKey) != allExports.end())
-				writeExport(tmp, false); // Workaround for tool not liking empty lists
+				writeExport(tmp, false);
 
 			write("", 1);
 			directiveSection.SizeOfRawData = u32(memoryBlock.writtenSize) - directiveRawDataStart;
 		}
 
-
-		// Memory for relocations and write relocations
-		u32 relocationsRawDataPos = u32(memoryBlock.writtenSize);
+		// Raw data for loopback pointers (8 bytes each, zero-initialized)
+		u32 dataRawPos = u32(memoryBlock.writtenSize);
 		allocate(loopbackCount * 8);
-		u32 relocationsMemPos = u32(memoryBlock.writtenSize);
+
+		// Relocations for each 8-byte slot (pointing at the base symbol, i.e., "foo")
+		u32 relocFilePos = u32(memoryBlock.writtenSize);
 		auto relocations = (ImageRelocation*)allocate(loopbackCount * sizeof(ImageRelocation));
-		textSection.PointerToRelocations = relocationsMemPos;
-		textSection.NumberOfRelocations = u16(loopbackCount);
-		textSection.PointerToRawData = relocationsRawDataPos;
-		textSection.SizeOfRawData = loopbackCount * 8;
-		for (u32 i=0; i!=loopbackCount; ++i)
+		rdataSection.PointerToRawData   = dataRawPos;
+		rdataSection.SizeOfRawData      = loopbackCount * 8;
+		rdataSection.PointerToRelocations = relocFilePos;
+		rdataSection.NumberOfRelocations  = u16(loopbackCount);
+
+		for (u32 i = 0; i != loopbackCount; ++i)
 		{
 			auto& relocation = relocations[i];
-			relocation.VirtualAddress = 8 * i;
-			relocation.SymbolTableIndex = i;
-			relocation.Type = ImageRelAmd64Addr64;
+			relocation.VirtualAddress   = 8 * i;   // offset in section
+			relocation.SymbolTableIndex = i;       // refers to first half of symbols (undefined "foo")
+			relocation.Type             = ImageRelAmd64Addr64;
 		}
 
-		// Memory for symbols
+		// Symbol table
 		header.PointerToSymbolTable = u32(memoryBlock.writtenSize);
 		header.NumberOfSymbols = loopbackCount * 2;
 		auto symbols = (ImageSymbol*)allocate(header.NumberOfSymbols * sizeof(ImageSymbol));
 
-		// Write string table
+		// String table: write BOTH "foo" and "__imp_foo" as independent entries
 		u64 stringStart = memoryBlock.writtenSize;
 		auto& stringTableSize = *(u32*)allocate(4);
-		Vector<u32> symbolsToAdd;
-		for (auto& str : neededLoopbacks)
-		{
-			symbolsToAdd.push_back(u32(memoryBlock.writtenSize - stringStart));
-			write(str.data(), str.size()+1);
-		}
-		stringTableSize = u32(memoryBlock.writtenSize - stringStart);
-		
-		// Write symbols
-		for (u32 i=0; i!=loopbackCount; ++i)
-		{
-			auto& symbol = symbols[i];
-			symbol.N.Name.Long = symbolsToAdd[i] + 6; // Remove __imp_
-			symbol.SectionNumber = ImageSymUndefined;
-			symbol.StorageClass = ImageSymClassExternal;
-		}
-		for (u32 i=0; i!=loopbackCount; ++i)
-		{
-			auto& symbol = symbols[i + loopbackCount];
-			symbol.N.Name.Long = symbolsToAdd[i];
-			symbol.SectionNumber = textSessionIndex + 1;
-			symbol.StorageClass = ImageSymClassExternal;
-			symbol.Value = i * 8;
-		}
+		Vector<u32> offBase;    offBase.reserve(loopbackCount);    // "foo"
+		Vector<u32> offImpBase; offImpBase.reserve(loopbackCount); // "__imp_foo"
 
-		return true;
-	}
-
-	template<typename SymbolType>
-	void ObjectFileCoff::WriteImports(Logger& logger, u8* newData, Info& newInfo, const Vector<u32>& symbolsToAdd)
-	{
-		u32 importsToFixCount = u32(symbolsToAdd.size());
-
-		// Copy header and sections
-		u32 offsetToAfterLastSection = m_info.sectionsMemOffset + m_info.sectionCount*sizeof(ImageSectionHeader);
-		memcpy(newData, m_data, offsetToAfterLastSection);
-
-		// Add another section
-		u32 newSectionIndex = m_info.sectionCount;
-		auto& newSection = *(ImageSectionHeader*)(newData + offsetToAfterLastSection);
-		memset(&newSection, 0, sizeof(ImageSectionHeader));
-		memcpy(newSection.Name, ".text$mn", 8);
-
-		// Add dummy memory for relocations
-		u32 newRelocationVirtualMemPos = offsetToAfterLastSection + sizeof(ImageSectionHeader);
-		u32 newRelocationVirtualMemSize = 8 * importsToFixCount;
-		memset(newData + newRelocationVirtualMemPos, 0, newRelocationVirtualMemSize);
-
-		u32 newRelocationsPos = newRelocationVirtualMemPos + newRelocationVirtualMemSize;
-		u32 newRelocationsSize = sizeof(ImageRelocation) * importsToFixCount;
-		newSection.PointerToRelocations = newRelocationsPos;
-		newSection.NumberOfRelocations = u16(importsToFixCount);
-		newSection.PointerToRawData = newRelocationVirtualMemPos;
-		newSection.SizeOfRawData = newRelocationVirtualMemSize;
-		newSection.Characteristics = ImageScnCntCode | ImageScnMemExecute | ImageScnMemRead;
-		UBA_ASSERT(importsToFixCount <= 65535);
-
-		// Add the new relocations and raw memory
-		u32 newSymbolIndex = m_info.symbolCount;
-		memset(newData + newRelocationsPos, 0, sizeof(ImageRelocation) *  importsToFixCount);
-		for (u32 i=0; i!=importsToFixCount; ++i)
+		for (auto& s : neededLoopbacks)
 		{
-			auto& relocation = ((ImageRelocation*)(newData + newRelocationsPos))[i];
-			relocation.VirtualAddress = 8 * i;
-			relocation.SymbolTableIndex = newSymbolIndex + i;
-			relocation.Type = ImageRelAmd64Addr64;
-		}
+			// s == "__imp_foo"
+			std::string base = s.substr(6); // "foo"
 
-		// Offset of everything after new section and relocations
-		u32 memoryOffset = newRelocationsPos + newRelocationsSize - offsetToAfterLastSection;
+			offBase.push_back(u32(memoryBlock.writtenSize - stringStart));
+			write(base.data(), base.size() + 1);
 
-		u32 symbolTablePos;
-		u32 symbolTableSize;
-		if (m_isBigObj)
-		{
-			auto& header = *(AnonObjectHeaderBigobj*)newData;
-			++header.NumberOfSections;
-			header.PointerToSymbolTable += memoryOffset;
-			symbolTablePos = header.PointerToSymbolTable;
-			symbolTableSize = header.NumberOfSymbols * sizeof(SymbolType);
-			header.NumberOfSymbols += importsToFixCount*2;
-		}
-		else
-		{
-			auto& header = *(ImageFileHeader*)newData;
-			++header.NumberOfSections;
-			header.PointerToSymbolTable += memoryOffset;
-			symbolTablePos = header.PointerToSymbolTable;
-			symbolTableSize = header.NumberOfSymbols * sizeof(SymbolType);
-			header.NumberOfSymbols += importsToFixCount*2;
+			offImpBase.push_back(u32(memoryBlock.writtenSize - stringStart));
+			write(s.data(), s.size() + 1);
 		}
+		stringTableSize = u32(memoryBlock.writtenSize - stringStart); // includes the 4-byte size field
 
-		u32 offsetToAfterSymbolTable = symbolTablePos + symbolTableSize;
-
-		// Copy everything after sections til end of symbol table
-		u64 nextToCopySize = m_info.symbolsMemPos + symbolTableSize - offsetToAfterLastSection;
-		memcpy(newData + offsetToAfterLastSection + memoryOffset, m_data + offsetToAfterLastSection, nextToCopySize);
-
-		newInfo.symbolsMemPos += memoryOffset;
-
-		u32 newSymbolsPos = offsetToAfterSymbolTable;
-		u32 newSymbolsSize = sizeof(SymbolType)*importsToFixCount*2;
-		auto newSymbols = (SymbolType*)(newData + newSymbolsPos);
-		memset(newSymbols, 0, newSymbolsSize);
-		for (u32 i=0; i!=importsToFixCount; ++i)
-		{
-			auto& symbol = newSymbols[i];
-			symbol.N.Name.Long = symbolsToAdd[i] + 6; // Remove __imp_
-			symbol.SectionNumber = ImageSymUndefined;
-			symbol.StorageClass = ImageSymClassExternal;
-		}
-		for (u32 i=0; i!=importsToFixCount; ++i)
+		// First N symbols: undefined "foo" (targets for relocations)
+		for (u32 i = 0; i != loopbackCount; ++i)
 		{
-			auto& symbol = newSymbols[i + importsToFixCount];
-			symbol.N.Name.Long = symbolsToAdd[i];
-			symbol.SectionNumber = (decltype(SymbolType::SectionNumber))(newSectionIndex + 1);
-			symbol.StorageClass = ImageSymClassExternal;
-			symbol.Value = i * 8;
+			auto& sym = symbols[i];
+			sym.N.Name.Long  = offBase[i];             // points to "foo"
+			sym.SectionNumber = ImageSymUndefined;
+			sym.StorageClass  = ImageSymClassExternal;
+			// Type/Value/Aux are 0 from allocator
 		}
 
-		u64 lastToCopySize = m_dataSize - m_info.stringTableMemPos;
-		memcpy(newData + newSymbolsPos + newSymbolsSize, m_data + m_info.stringTableMemPos, lastToCopySize);
-
-
-		// Traverse all sections and update their PointerToRawData
-		auto sections = (ImageSectionHeader*)(newData + m_info.sectionsMemOffset);
-		for (u32 i=0; i!=m_info.sectionCount; ++i)
+		// Next N symbols: defined "__imp_foo" in our .rdata section at offset i*8
+		for (u32 i = 0; i != loopbackCount; ++i)
 		{
-			auto& section = sections[i];
-			if (section.PointerToRawData)
-			{
-				UBA_ASSERT(section.PointerToRawData < symbolTablePos);
-				section.PointerToRawData += memoryOffset;
-			}
-			if (section.PointerToRelocations)
-			{
-				UBA_ASSERT(section.PointerToRelocations < symbolTablePos);
-				section.PointerToRelocations += memoryOffset;
-			}
-			if (section.PointerToLinenumbers)
-			{
-				UBA_ASSERT(section.PointerToLinenumbers < symbolTablePos);
-				section.PointerToLinenumbers += memoryOffset;
-			}
+			auto& sym = symbols[i + loopbackCount];
+			sym.N.Name.Long   = offImpBase[i];         // points to "__imp_foo"
+			sym.SectionNumber = (short)(rdataSectionIndex0 + 1); // 1-based section index
+			sym.StorageClass  = ImageSymClassExternal;
+			sym.Value         = i * 8;                 // offset of the pointer slot
 		}
+
+		return true;
 	}
 }
diff --git a/Engine/Source/Programs/UnrealBuildAccelerator/Common/Private/UbaObjectFileCoff.h b/Engine/Source/Programs/UnrealBuildAccelerator/Common/Private/UbaObjectFileCoff.h
--- a/Engine/Source/Programs/UnrealBuildAccelerator/Common/Private/UbaObjectFileCoff.h
+++ b/Engine/Source/Programs/UnrealBuildAccelerator/Common/Private/UbaObjectFileCoff.h
@@ -20,7 +20,6 @@ namespace uba
 
 		bool ParseExports();
 		template<typename SymbolType> void ParseAllSymbols();
-		template<typename SymbolType> void WriteImports(Logger& logger, u8* newData, Info& newInfo, const Vector<u32>& symbolsToAdd);
 
 		struct Info
 		{
