[UBA]
* Changed timeout for fetching cache compact table from 5min to 20min. This number must be higher than network timeout and should never timeout.

diff --git a/Engine/Source/Programs/UnrealBuildAccelerator/Common/Private/UbaCacheClient.cpp b/Engine/Source/Programs/UnrealBuildAccelerator/Common/Private/UbaCacheClient.cpp
--- a/Engine/Source/Programs/UnrealBuildAccelerator/Common/Private/UbaCacheClient.cpp
+++ b/Engine/Source/Programs/UnrealBuildAccelerator/Common/Private/UbaCacheClient.cpp
@@ -1533,13 +1533,13 @@ namespace uba
 			break;
 		}
 
-		u32 timeOutTimeMs = 5*60*1000;
+		u32 timeOutTimeMs = 20*60*1000; // Must be higher than network timeout
 
 		u64 written = 0;
 		for (Entry& e : entries)
 		{
 			if (!e.done.IsSet(timeOutTimeMs)) // This should not be possible
-				m_logger.Error(TC("FetchCompactTable timed out after 5 minutes getting async message response. This timeout will cause a crash."));
+				m_logger.Error(TC("FetchCompactTable timed out after 20 minutes getting async message response. This timeout will cause a crash."));
 			if (!e.message.ProcessAsyncResults(e.reader))
 				success = false;
 			written += e.reader.GetLeft();
