[UBA]
* Don't hook up exception handler in UbaHost.dll when running on linux. Callstacks are garbage and hoping to find some crash using core dumps or something

diff --git a/Engine/Source/Programs/UnrealBuildAccelerator/Host/Private/UbaExports.cpp b/Engine/Source/Programs/UnrealBuildAccelerator/Host/Private/UbaExports.cpp
--- a/Engine/Source/Programs/UnrealBuildAccelerator/Host/Private/UbaExports.cpp
+++ b/Engine/Source/Programs/UnrealBuildAccelerator/Host/Private/UbaExports.cpp
@@ -137,7 +137,9 @@ extern "C"
 
 	uba::LogWriter* CreateCallbackLogWriter(uba::CallbackLogWriter::BeginScopeCallback begin, uba::CallbackLogWriter::EndScopeCallback end, uba::CallbackLogWriter::LogCallback log)
 	{
+		#if !PLATFORM_LINUX
 		uba::AddExceptionHandler();
+		#endif
 		uba::InitMemory();
 
 		return new uba::CallbackLogWriter(begin, end, log);
