[UBA]
* Fixed so binary reader does not crash on too small capacity when reading strings. It still asserts in debug.

diff --git a/Engine/Source/Programs/UnrealBuildAccelerator/Core/Public/UbaBinaryReaderWriter.h b/Engine/Source/Programs/UnrealBuildAccelerator/Core/Public/UbaBinaryReaderWriter.h
--- a/Engine/Source/Programs/UnrealBuildAccelerator/Core/Public/UbaBinaryReaderWriter.h
+++ b/Engine/Source/Programs/UnrealBuildAccelerator/Core/Public/UbaBinaryReaderWriter.h
@@ -104,6 +104,7 @@ namespace uba
 
 	protected:
 		inline u64 InternalReadString(tchar* str, u64 charLen);
+		inline void InternalSkipString(u64 charLen);
 		const u8* m_begin;
 		const u8* m_pos;
 		const u8* m_end;
@@ -348,8 +349,14 @@ namespace uba
 	u64 BinaryReader::ReadString(tchar* str, u64 strCapacity)
 	{
 		u64 charLen = Read7BitEncoded();
-		UBA_ASSERTF(charLen < strCapacity - 1, TC("Strlen: %llu, Capacity: %llu"), charLen, strCapacity); (void)strCapacity;
-		return InternalReadString(str, charLen);
+		if (charLen < strCapacity - 1)
+			return InternalReadString(str, charLen);
+
+		u64 newCharLen = strCapacity - 1;
+		u64 res = InternalReadString(str, newCharLen);
+		UBA_ASSERTF(false, TC("Strlen: %llu, Capacity: %llu Read: %s"), charLen, strCapacity, str);
+		InternalSkipString(charLen - newCharLen);
+		return res;
 	}
 
 	u64 BinaryReader::InternalReadString(tchar* str, u64 charLen)
@@ -400,6 +407,38 @@ namespace uba
 		#endif
 	}
 
+	void BinaryReader::InternalSkipString(u64 charLen)
+	{
+		#if PLATFORM_WINDOWS
+		const u8* pos = m_pos;
+		u64 left = charLen;
+		while (left--)
+		{
+			UBA_ASSERT_READ(1);
+			u8 a = *pos++;
+			if (a <= 127)
+				continue;
+			UBA_ASSERT_READ(1);
+			pos++;
+			if (a >= 192 && a <= 223)
+				continue;
+			UBA_ASSERT_READ(1);
+			pos++;
+			if (a >= 224 && a <= 239)
+				continue;
+			if (a >= 240 && a <= 253)
+			{
+				UBA_ASSERT(false); // Wide chars cannot exceed 16 bits
+				continue;
+			}
+			UBA_ASSERT(false); // Wide chars cannot exceed 16 bits
+		}
+		m_pos = pos;
+		#else
+		Skip(charLen);
+		#endif
+	}
+
 	void BinaryReader::ReadString(StringBufferBase& out)
 	{
 		out.count += u32(ReadString(out.data + out.count, (out.capacity - out.count)));
@@ -416,8 +455,8 @@ namespace uba
 
 	void BinaryReader::SkipString()
 	{
-		u64 len = Read7BitEncoded();
-		Skip(len);
+		u64 charLen = Read7BitEncoded();
+		InternalSkipString(charLen);
 	}
 
 	Guid BinaryReader::ReadGuid()
diff --git a/Engine/Source/Programs/UnrealBuildAccelerator/Test/Private/UbaTestBasics.cpp b/Engine/Source/Programs/UnrealBuildAccelerator/Test/Private/UbaTestBasics.cpp
--- a/Engine/Source/Programs/UnrealBuildAccelerator/Test/Private/UbaTestBasics.cpp
+++ b/Engine/Source/Programs/UnrealBuildAccelerator/Test/Private/UbaTestBasics.cpp
@@ -466,6 +466,21 @@ namespace uba
 			return false;
 		#endif
 
+		#if !UBA_DEBUG
+		u8 mem[1024];
+		BinaryWriter writer(mem);
+		writer.WriteString(TC("Foo"));
+		writer.WriteString(TC("Bar"));
+		BinaryReader reader(mem);
+		StringBuffer<3> str;
+		reader.ReadString(str);
+		if (!str.Equals(TC("Fo")))
+			return false;
+		reader.ReadString(str.Clear());
+		if (!str.Equals(TC("Ba")))
+			return false;
+		#endif
+
 		return true;
 	}
 
