[UBA]
* Changed some reader.ReadString(stringBuffer) to TString = reader.ReadString() because in one case stringbuffer was too small  and we can just as well use TString if it is in places that are not frequently called

diff --git a/Engine/Source/Programs/UnrealBuildAccelerator/Common/Private/UbaSessionServer.cpp b/Engine/Source/Programs/UnrealBuildAccelerator/Common/Private/UbaSessionServer.cpp
--- a/Engine/Source/Programs/UnrealBuildAccelerator/Common/Private/UbaSessionServer.cpp
+++ b/Engine/Source/Programs/UnrealBuildAccelerator/Common/Private/UbaSessionServer.cpp
@@ -872,14 +872,12 @@ namespace uba
 			return true;
 		}
 
-		StringBuffer<256> machineId;
-		reader.ReadString(machineId);
+		TString machineId = reader.ReadString();
 
 		u32 processSlotCount = reader.ReadU32();
 		bool dedicated = reader.ReadBool();
 
-		StringBuffer<256> info;
-		reader.ReadString(info);
+		TString info = reader.ReadString();
 
 		u64 memAvail = reader.ReadU64();
 		u64 memTotal = reader.ReadU64();
