[UBA]
* Fixed so cache server still validates that cas entries exists when finding existing cache entry after entry sent from client. This also bumps access time for cas entry to prevent it from being deleted

diff --git a/Engine/Source/Programs/UnrealBuildAccelerator/Common/Private/UbaCacheServer.cpp b/Engine/Source/Programs/UnrealBuildAccelerator/Common/Private/UbaCacheServer.cpp
--- a/Engine/Source/Programs/UnrealBuildAccelerator/Common/Private/UbaCacheServer.cpp
+++ b/Engine/Source/Programs/UnrealBuildAccelerator/Common/Private/UbaCacheServer.cpp
@@ -1737,6 +1737,8 @@ namespace uba
 			break;
 		}
 
+		UnorderedSet<CasKey> requestedOutputs;
+
 		// Already exists
 		if (matchingEntry != cacheEntries.entries.end())
 		{
@@ -1759,6 +1761,8 @@ namespace uba
 			// Traverse outputs and log out the mismatching output, then remove the old entry
 			reader.SetPosition(outputStartOffset);
 			u64 left = outputCount;
+			struct Pair { CasKey key; u32 offset; };
+			Vector<Pair> outputKeys;
 			while (left--)
 			{
 				u32 outputOffset = u32(reader.Read7BitEncoded());
@@ -1783,9 +1787,22 @@ namespace uba
 					shouldOverwrite = true;
 					break;
 				}
+				outputKeys.push_back(Pair{casKey, outputOffset});
 			}
+
 			if (!shouldOverwrite)
+			{
+				// Verify so output cas entries exist. Otherwise we want client to resend them
+				for (Pair& pair : outputKeys)
+				{
+					if (!requestedOutputs.insert(pair.key).second)
+						continue;
+					if (m_storage.EnsureCasFile(pair.key, nullptr))
+						continue;
+					writer.Write7BitEncoded(pair.offset);
+				}
 				return true;
+			}
 		}
 
 		Set<u32> outputs;
@@ -1795,7 +1812,6 @@ namespace uba
 		bool hasAllContent = true;
 		reader.SetPosition(outputStartOffset);
 		u64 left = outputCount;
-		UnorderedSet<CasKey> requestedOutputs;
 		while (left--)
 		{
 			u32 outputOffset = u32(reader.Read7BitEncoded());
