[UBA]
* Improved error message to also contain client id

diff --git a/Engine/Source/Programs/UnrealBuildAccelerator/Common/Private/UbaStorageServer.cpp b/Engine/Source/Programs/UnrealBuildAccelerator/Common/Private/UbaStorageServer.cpp
--- a/Engine/Source/Programs/UnrealBuildAccelerator/Common/Private/UbaStorageServer.cpp
+++ b/Engine/Source/Programs/UnrealBuildAccelerator/Common/Private/UbaStorageServer.cpp
@@ -590,7 +590,7 @@ namespace uba
 				SCOPED_WRITE_LOCK(casEntry->lock, entryLock);
 				if (!casEntry->exists)
 				{
-					m_logger.Info(TC("FetchBegin failed for cas file %s (%s). Cas was deleted"), CasKeyString(casKey).str, hint.data);
+					m_logger.Info(TC("Client %u failed FetchBegin for cas file %s (%s). Cas was deleted"), connectionInfo.GetId(), CasKeyString(casKey).str, hint.data);
 					writer.WriteU16(0);
 					return true;
 				}
