[UBA]
* Added so GetTraceInformation is a system message instead of session message.

diff --git a/Engine/Source/Programs/UnrealBuildAccelerator/Common/Private/UbaSessionServer.cpp b/Engine/Source/Programs/UnrealBuildAccelerator/Common/Private/UbaSessionServer.cpp
--- a/Engine/Source/Programs/UnrealBuildAccelerator/Common/Private/UbaSessionServer.cpp
+++ b/Engine/Source/Programs/UnrealBuildAccelerator/Common/Private/UbaSessionServer.cpp
@@ -191,6 +191,7 @@ namespace uba
 			m_storage.Reset();
 
 		m_storage.SetTrace(&m_trace, m_detailedTrace);
+		m_server.SetTrace(&m_trace);
 
 		if (m_detailedTrace)
 			m_server.SetWorkTracker(&m_trace);
@@ -1683,6 +1684,7 @@ namespace uba
 
 	bool SessionServer::HandleGetTraceInformation(const ConnectionInfo& connectionInfo, const WorkContext& workContext, BinaryReader& reader, BinaryWriter& writer)
 	{
+		// This code is deprecated and moved to UbaNetworkServer. Trace is now a system message
 		u32 remotePos = reader.ReadU32();
 		u32 localPos;
 		{
diff --git a/Engine/Source/Programs/UnrealBuildAccelerator/Common/Public/UbaNetwork.h b/Engine/Source/Programs/UnrealBuildAccelerator/Common/Public/UbaNetwork.h
--- a/Engine/Source/Programs/UnrealBuildAccelerator/Common/Public/UbaNetwork.h
+++ b/Engine/Source/Programs/UnrealBuildAccelerator/Common/Public/UbaNetwork.h
@@ -28,6 +28,7 @@ namespace uba
 		UBA_SYSTEM_MESSAGE(SetConnectionCount) \
 		UBA_SYSTEM_MESSAGE(KeepAlive) \
 		UBA_SYSTEM_MESSAGE(FetchConfig) \
+		UBA_SYSTEM_MESSAGE(GetTraceInformation) \
 
 	#define UBA_STORAGE_MESSAGES \
 		UBA_STORAGE_MESSAGE(FetchBegin) \
diff --git a/Engine/Source/Programs/UnrealBuildAccelerator/Common/Public/UbaNetworkServer.h b/Engine/Source/Programs/UnrealBuildAccelerator/Common/Public/UbaNetworkServer.h
--- a/Engine/Source/Programs/UnrealBuildAccelerator/Common/Public/UbaNetworkServer.h
+++ b/Engine/Source/Programs/UnrealBuildAccelerator/Common/Public/UbaNetworkServer.h
@@ -14,6 +14,7 @@
 namespace uba
 {
 	class Config;
+	class Trace;
 	struct BinaryReader;
 	struct BinaryWriter;
 
@@ -64,6 +65,7 @@ namespace uba
 
 		bool RegisterCryptoKey(const u8* cryptoKey128, u64 expirationTime = ~u64(0), u64 cryptoUserData = 0);
 		void SetClientsConfig(const Config& config);
+		void SetTrace(Trace* trace);
 
 		bool AddClient(NetworkBackend& backend, const tchar* ip, u16 port = DefaultPort, const u8* cryptoKey128 = nullptr); // Adds a client that server will create one or more connections to (note this will return before we know if it was a success or not)
 		bool HasConnectInProgress();
@@ -206,6 +208,8 @@ namespace uba
 		Atomic<u64> m_recvBytes;
 		Atomic<u32> m_recvCount;
 
+		Trace* m_trace = nullptr;
+
 		NetworkServer(const NetworkServer&) = delete;
 		void operator=(const NetworkServer&) = delete;
 	};
diff --git a/Engine/Source/Programs/UnrealBuildAccelerator/Common/Public/UbaTrace.h b/Engine/Source/Programs/UnrealBuildAccelerator/Common/Public/UbaTrace.h
--- a/Engine/Source/Programs/UnrealBuildAccelerator/Common/Public/UbaTrace.h
+++ b/Engine/Source/Programs/UnrealBuildAccelerator/Common/Public/UbaTrace.h
@@ -195,6 +195,7 @@ namespace uba
 		u64 m_previousIdleCpuTime = 0;
 		float m_cpuLoad = 0;
 
+		friend class NetworkServer;
 		friend class SessionServer;
 	};
 
diff --git a/Engine/Source/Programs/UnrealBuildAccelerator/Host/Private/UbaExports.cpp b/Engine/Source/Programs/UnrealBuildAccelerator/Host/Private/UbaExports.cpp
--- a/Engine/Source/Programs/UnrealBuildAccelerator/Host/Private/UbaExports.cpp
+++ b/Engine/Source/Programs/UnrealBuildAccelerator/Host/Private/UbaExports.cpp
@@ -231,10 +231,14 @@ extern "C"
 
 		bool success = true;
 		auto server = new NetworkServerWithBackend(success, info, networkBackend);
-		if (success)
-			return server;
-		delete server;
-		return nullptr;
+		if (!success)
+		{
+			delete server;
+			return nullptr;
+		}
+
+		server->SetTrace(uba::g_globalTrace);
+		return server;
 	}
 
 	void NetworkServer_Destroy(uba::NetworkServer* server)
diff --git a/Engine/Source/Programs/UnrealBuildAccelerator/Visualizer/Private/UbaTraceReader.cpp b/Engine/Source/Programs/UnrealBuildAccelerator/Visualizer/Private/UbaTraceReader.cpp
--- a/Engine/Source/Programs/UnrealBuildAccelerator/Visualizer/Private/UbaTraceReader.cpp
+++ b/Engine/Source/Programs/UnrealBuildAccelerator/Visualizer/Private/UbaTraceReader.cpp
@@ -245,7 +245,7 @@ namespace uba
 			u32 pos = u32(m_memoryEnd - m_memoryBegin);
 
 			StackBinaryWriter<32> writer;
-			NetworkMessage msg(client, SessionServiceId, SessionMessageType_GetTraceInformation, writer);
+			NetworkMessage msg(client, SystemServiceId, SystemMessageType_GetTraceInformation, writer);
 			writer.WriteU32(pos);
 			StackBinaryReader<SendDefaultSize> reader;
 			if (!msg.Send(reader))
