[UBA]
* Missed one file from previous submit

diff --git a/Engine/Source/Programs/UnrealBuildAccelerator/Common/Private/UbaNetworkServer.cpp b/Engine/Source/Programs/UnrealBuildAccelerator/Common/Private/UbaNetworkServer.cpp
--- a/Engine/Source/Programs/UnrealBuildAccelerator/Common/Private/UbaNetworkServer.cpp
+++ b/Engine/Source/Programs/UnrealBuildAccelerator/Common/Private/UbaNetworkServer.cpp
@@ -6,6 +6,7 @@
 #include "UbaBinaryReaderWriter.h"
 #include "UbaEnvironment.h"
 #include "UbaPlatform.h"
+#include "UbaTrace.h"
 
 namespace uba
 {
@@ -17,6 +18,11 @@ namespace uba
 		const ConfigTable& table = *tablePtr;
 		table.GetValueAsU32(recvBufferSize, TC("RecvBufferSize"));
 		table.GetValueAsU32(sendBufferSize, TC("SendBufferSize"));
+		table.GetValueAsU32(workerCount, TC("WorkerCount"));
+		table.GetValueAsU32(sendSize, TC("SendSize"));
+		table.GetValueAsBool(logConnections, TC("LogConnections"));
+		table.GetValueAsBool(useKeepAlive, TC("UseKeepAlive"));
+		table.GetValueAsBool(useMessagePriority, TC("UseMessagePriority"));
 	}
 
 	struct NetworkServer::Client
@@ -893,6 +899,11 @@ namespace uba
 		config.SaveToText(m_logger, m_clientsConfig);
 	}
 
+	void NetworkServer::SetTrace(Trace* trace)
+	{
+		m_trace = trace;
+	}
+
 	bool NetworkServer::AddClient(NetworkBackend& backend, const tchar* ip, u16 port, const u8* cryptoKey128)
 	{
 		SCOPED_FUTEX(m_addConnectionsLock, lock);
@@ -1328,7 +1339,12 @@ namespace uba
 
 			lock.Leave();
 
-			req.ev.IsSet();
+			// TODO: We should log something here
+			//if (!req.ev.IsSet(20*1000))
+			//{
+				//m_logger.Info(TC("Waited 20 seconds for available worker. This should not happen!"));
+				req.ev.IsSet();
+			//}
 
 			if (req.worker)
 				return req.worker;
@@ -1537,6 +1553,27 @@ namespace uba
 				writer.WriteBytes(m_clientsConfig.data(), m_clientsConfig.size());
 				return true;
 			}
+			case SystemMessageType_GetTraceInformation:
+			{
+				if (!m_trace)
+				{
+					writer.WriteU32(0);
+					return true;
+				}
+				Trace& trace = *m_trace;
+
+				u32 remotePos = reader.ReadU32();
+				u32 localPos;
+				{
+					SCOPED_FUTEX_READ(trace.m_memoryLock, l);
+					localPos = u32(trace.m_memoryPos);
+				}
+
+				writer.WriteU32(localPos);
+				u32 toWrite = Min(localPos - remotePos, u32(writer.GetCapacityLeft()));
+				writer.WriteBytes(trace.m_memoryBegin + remotePos, toWrite);
+				return true;
+			}
 		}
 		return false;
 	}
