[UBA]
* Added LogWriter_GetNull as export from UbaHost
* Renamed CallbackLogWriter to use LogWriter_xx instead (old ones still in there but deprecated)

diff --git a/Engine/Source/Programs/UnrealBuildAccelerator/Host/Private/UbaExports.cpp b/Engine/Source/Programs/UnrealBuildAccelerator/Host/Private/UbaExports.cpp
--- a/Engine/Source/Programs/UnrealBuildAccelerator/Host/Private/UbaExports.cpp
+++ b/Engine/Source/Programs/UnrealBuildAccelerator/Host/Private/UbaExports.cpp
@@ -130,12 +130,17 @@ namespace uba
 
 extern "C"
 {
-	uba::LogWriter* GetDefaultLogWriter()
+	uba::LogWriter* LogWriter_GetDefault()
 	{
 		return &uba::g_consoleLogWriter;
 	}
 
-	uba::LogWriter* CreateCallbackLogWriter(uba::CallbackLogWriter::BeginScopeCallback begin, uba::CallbackLogWriter::EndScopeCallback end, uba::CallbackLogWriter::LogCallback log)
+	uba::LogWriter* LogWriter_GetNull()
+	{
+		return &uba::g_nullLogWriter;
+	}
+
+	uba::LogWriter* LogWriter_CreateCallback(uba::CallbackLogWriter::BeginScopeCallback begin, uba::CallbackLogWriter::EndScopeCallback end, uba::CallbackLogWriter::LogCallback log)
 	{
 		#if !PLATFORM_LINUX
 		uba::AddExceptionHandler();
@@ -145,12 +150,22 @@ extern "C"
 		return new uba::CallbackLogWriter(begin, end, log);
 	}
 
-	void DestroyCallbackLogWriter(uba::LogWriter* writer)
+	void LogWriter_DestroyCallback(uba::LogWriter* writer)
 	{
 		if (writer != &uba::g_consoleLogWriter)
 			delete writer;
 	}
 
+	uba::LogWriter* CreateCallbackLogWriter(uba::CallbackLogWriter::BeginScopeCallback begin, uba::CallbackLogWriter::EndScopeCallback end, uba::CallbackLogWriter::LogCallback log)
+	{
+		return LogWriter_CreateCallback(begin, end, log);
+	}
+
+	void DestroyCallbackLogWriter(uba::LogWriter* writer)
+	{
+		LogWriter_DestroyCallback(writer);
+	}
+
 	uba::Config* Config_Load(const uba::tchar* configFile)
 	{
 		return &uba::GetConfig(configFile);
diff --git a/Engine/Source/Programs/UnrealBuildAccelerator/Host/Public/UbaExports.h b/Engine/Source/Programs/UnrealBuildAccelerator/Host/Public/UbaExports.h
--- a/Engine/Source/Programs/UnrealBuildAccelerator/Host/Public/UbaExports.h
+++ b/Engine/Source/Programs/UnrealBuildAccelerator/Host/Public/UbaExports.h
@@ -45,9 +45,10 @@ namespace uba
 extern "C"
 {
 	// LogWriter
-	UBA_API uba::LogWriter* GetDefaultLogWriter();
-	UBA_API uba::LogWriter* CreateCallbackLogWriter(uba::CallbackLogWriter::BeginScopeCallback begin, uba::CallbackLogWriter::EndScopeCallback end, uba::CallbackLogWriter::LogCallback log);
-	UBA_API void DestroyCallbackLogWriter(uba::LogWriter* writer);
+	UBA_API uba::LogWriter* LogWriter_GetDefault();
+	UBA_API uba::LogWriter* LogWriter_GetNull();
+	UBA_API uba::LogWriter* LogWriter_CreateCallback(uba::CallbackLogWriter::BeginScopeCallback begin, uba::CallbackLogWriter::EndScopeCallback end, uba::CallbackLogWriter::LogCallback log);
+	UBA_API void LogWriter_DestroyCallback(uba::LogWriter* writer);
 
 	// Config
 	UBA_API uba::Config* Config_Load(const uba::tchar* configFile);
@@ -211,6 +212,8 @@ extern "C"
 
 
 	// DEPRECATED, don't use
+	UBA_API uba::LogWriter* CreateCallbackLogWriter(uba::CallbackLogWriter::BeginScopeCallback begin, uba::CallbackLogWriter::EndScopeCallback end, uba::CallbackLogWriter::LogCallback log);
+	UBA_API void DestroyCallbackLogWriter(uba::LogWriter* writer);
 	UBA_API void DestroyProcessHandle(const uba::ProcessHandle* handle);
 	UBA_API uba::RootPaths* RootPaths_Create(uba::LogWriter& writer);
 	UBA_API bool RootPaths_RegisterRoot(uba::RootPaths* rootPaths, const uba::tchar* path, bool includeInKey, uba::u8 id = 0);
