[UBA]
* Fixed so stored cas files will not be freed until they have been used for copyorlink. cas entries will be ref counted on store and released on copyonwrite (or disconnect)

#FailedValidations Preflight

diff --git a/Engine/Source/Programs/UnrealBuildAccelerator/Common/Private/UbaSessionClient.cpp b/Engine/Source/Programs/UnrealBuildAccelerator/Common/Private/UbaSessionClient.cpp
--- a/Engine/Source/Programs/UnrealBuildAccelerator/Common/Private/UbaSessionClient.cpp
+++ b/Engine/Source/Programs/UnrealBuildAccelerator/Common/Private/UbaSessionClient.cpp
@@ -540,7 +540,7 @@ namespace uba
 
 		constexpr bool writeCompressed = false;
 		constexpr bool isTemp = true;
-		return m_storage.CopyOrLink(casKey, out.data, fileAttributes, writeCompressed, {}, isTemp, allowHardlink);
+		return m_storage.CopyOrLink(casKey, out.data, fileAttributes, writeCompressed, 0, {}, isTemp, allowHardlink);
 	}
 
 	bool SessionClient::ProcessThreadStart(ProcessImpl& process)
diff --git a/Engine/Source/Programs/UnrealBuildAccelerator/Common/Private/UbaSessionServer.cpp b/Engine/Source/Programs/UnrealBuildAccelerator/Common/Private/UbaSessionServer.cpp
--- a/Engine/Source/Programs/UnrealBuildAccelerator/Common/Private/UbaSessionServer.cpp
+++ b/Engine/Source/Programs/UnrealBuildAccelerator/Common/Private/UbaSessionServer.cpp
@@ -1174,7 +1174,7 @@ namespace uba
 		{
 			StringBuffer<> logPath;
 			logPath.Append(m_sessionLogDir).Append(destination.data + 5);
-			if (!m_storage.CopyOrLink(casKey, logPath.data, attributes))
+			if (!m_storage.CopyOrLink(casKey, logPath.data, attributes, false, clientId))
 				m_logger.Error(TC("Failed to copy cas from %s to %s"), CasKeyString(casKey).str, logPath.data);
 			else if (!m_storage.DropCasFile(casKey, false, logPath.data))
 				m_logger.Error(TC("Failed to drop cas %s"), CasKeyString(casKey).str);
@@ -1196,7 +1196,7 @@ namespace uba
 				ubaPath.Append(TCV("Connection")).AppendValue(clientId);
 
 			ubaPath.Append(TCV(".uba"));
-			m_storage.CopyOrLink(casKey, ubaPath.data, attributes);
+			m_storage.CopyOrLink(casKey, ubaPath.data, attributes, false, clientId);
 			m_storage.DropCasFile(casKey, false, ubaPath.data);
 			return true;
 		}
@@ -1243,7 +1243,7 @@ namespace uba
 				writeCompressed = g_globalRules.FileCanBeCompressed(destination);
 			}
 
-			success = m_storage.CopyOrLink(casKey, destination.data, attributes, writeCompressed, formattingFunc);
+			success = m_storage.CopyOrLink(casKey, destination.data, attributes, writeCompressed, clientId, formattingFunc);
 			if (!success)
 				m_logger.Error(TC("Failed to copy cas from %s to %s (%s)"), CasKeyString(casKey).str, destination.data, GetProcessDescription(processId).c_str());
 			else
diff --git a/Engine/Source/Programs/UnrealBuildAccelerator/Common/Private/UbaStorage.cpp b/Engine/Source/Programs/UnrealBuildAccelerator/Common/Private/UbaStorage.cpp
--- a/Engine/Source/Programs/UnrealBuildAccelerator/Common/Private/UbaStorage.cpp
+++ b/Engine/Source/Programs/UnrealBuildAccelerator/Common/Private/UbaStorage.cpp
@@ -96,7 +96,6 @@ namespace uba
 		table.GetValueAsBool(storeCompressed, TC("StoreCompressed"));
 		table.GetValueAsU32(fileEntryTableWipeThreshold, TC("FileEntryTableWipeThreshold"));
 		table.GetValueAsU32(maxParallelCopyOrLink, TC("MaxParallelCopyOrLink"));
-		table.GetValueAsBool(dropIndependentMappingsOnCopy, TC("DropIndependentMappingsOnCopy"));
 
 		const tchar* compressor;
 		if (table.GetValueAsString(compressor, TC("Compressor")))
@@ -217,7 +216,7 @@ namespace uba
 				break;
 			}
 
-			if (ce.beingWritten || ce.readCount)
+			if (ce.beingWritten || ce.refCount)
 			{
 				if (ce.beingWritten)
 					m_logger.Info(TC("Want to delete file that is begin written. We should never end up here"));
@@ -548,6 +547,9 @@ namespace uba
 			}
 		}
 
+		// File is requested by server itself to provide to someone remote. We have to un-drop the cas in case it was marked as dropped
+		casEntry.dropped = false;
+
 		if (casEntry.verified)
 			if (casEntry.exists)
 				return true;
@@ -883,7 +885,6 @@ namespace uba
 		m_asyncUnmapViewOfFile = info.asyncUnmapViewOfFile;
 		m_allowDeleteVerified = info.allowDeleteVerified;
 		m_writeToDisk = info.writeToDisk;
-		m_dropIndependentMappingsOnCopy = info.dropIndependentMappingsOnCopy;
 		if (!m_writeToDisk)
 			m_casCapacityBytes = 0;
 
@@ -2061,6 +2062,8 @@ namespace uba
 	
 		SCOPED_WRITE_LOCK(casEntry.lock, entryLock);
 
+		UBA_ASSERT(casEntry.refCount == 0);
+
 		if (forceDelete)
 		{
 			u64 sizeDeleted = 0;
@@ -2131,7 +2134,7 @@ namespace uba
 		return out != CasKeyZero;
 	}
 
-	bool StorageImpl::CopyOrLink(const CasKey& casKey, const tchar* destination, u32 fileAttributes, bool writeCompressed, const FormattingFunc& formattingFunc, bool isTemp, bool allowHardLink)
+	bool StorageImpl::CopyOrLink(const CasKey& casKey, const tchar* destination, u32 fileAttributes, bool writeCompressed, u32 clientId, const FormattingFunc& formattingFunc, bool isTemp, bool allowHardLink)
 	{
 		UBA_ASSERT(casKey != CasKeyZero);
 		UBA_ASSERT(fileAttributes);
@@ -2172,6 +2175,22 @@ namespace uba
 			UBA_ASSERT(casEntry.verified);
 			UBA_ASSERT(casEntry.exists);
 
+			auto done = [&](FileAccessor& destinationFile, u64 sizeOnDisk)
+				{
+					u64 lastWriteTime = 0;
+					if (!destinationFile.Close(&lastWriteTime))
+						return false;
+					UBA_ASSERT(lastWriteTime);
+					if (lastWriteTime)
+					{
+						fileEntry.casKey = casKey;
+						fileEntry.lastWritten = lastWriteTime;
+						fileEntry.size = sizeOnDisk;
+						fileEntry.verified = true;
+					}
+					return true;
+				};
+
 			if (IsCompressed(actualKey))
 			{
 				FileHandle readHandle = InvalidFileHandle;
@@ -2307,26 +2326,7 @@ namespace uba
 						return false;
 				}
 
-				u64 lastWriteTime = 0;
-				if (!destinationFile.Close(&lastWriteTime))
-					return false;
-				UBA_ASSERT(lastWriteTime);
-				if (lastWriteTime)
-				{
-					fileEntry.casKey = casKey;
-					fileEntry.lastWritten = lastWriteTime;
-					fileEntry.size = sizeOnDisk;
-					fileEntry.verified = true;
-				}
-
-				if (m_dropIndependentMappingsOnCopy && casEntry.mappingIsIndependent && casEntry.mappingHandle.IsValid())
-				{
-					mapViewGuard.Execute();
-					casEntryLock.Leave();
-					DropCasFile(casKey, true, destination);
-				}
-
-				return true;
+				return done(destinationFile, sizeOnDisk);
 			}
 
 			if (casEntry.mappingHandle.IsValid())
@@ -2334,7 +2334,7 @@ namespace uba
 				SCOPED_FUTEX(fileEntry.lock, entryLock);
 				fileEntry.verified = false;
 
-				MappedView mappedView = m_casDataBuffer.MapView(casEntry.mappingHandle, casEntry.mappingOffset, casEntry.mappingSize, TC(""));
+				MappedView mappedView = m_casDataBuffer.MapView(casEntry.mappingHandle, casEntry.mappingOffset, casEntry.mappingSize, TC(""), casEntry.mappingIsIndependent);
 				auto mapViewGuard = MakeGuard([&](){ m_casDataBuffer.UnmapView(mappedView, destination); });
 
 				FileAccessor destinationFile(m_logger, destination);
@@ -2343,18 +2343,7 @@ namespace uba
 				u8* writePos = destinationFile.GetData();
 				TimerScope cts(stats.memoryCopy);
 				MapMemoryCopy(writePos, mappedView.memory, mappedView.size);
-				u64 lastWriteTime = 0;
-				if (!destinationFile.Close(&lastWriteTime))
-					return false;
-				UBA_ASSERT(lastWriteTime);
-				if (lastWriteTime)
-				{
-					fileEntry.casKey = casKey;
-					fileEntry.lastWritten = lastWriteTime;
-					fileEntry.size = mappedView.size;
-					fileEntry.verified = true;
-				}
-				return true;
+				return done(destinationFile, mappedView.size);
 			}
 
 			StringBuffer<> casFile;
diff --git a/Engine/Source/Programs/UnrealBuildAccelerator/Common/Private/UbaStorageClient.cpp b/Engine/Source/Programs/UnrealBuildAccelerator/Common/Private/UbaStorageClient.cpp
--- a/Engine/Source/Programs/UnrealBuildAccelerator/Common/Private/UbaStorageClient.cpp
+++ b/Engine/Source/Programs/UnrealBuildAccelerator/Common/Private/UbaStorageClient.cpp
@@ -102,6 +102,12 @@ namespace uba
 				m_storageServerUid = reader.ReadGuid();
 				m_casCompressor = reader.ReadByte();
 				m_casCompressionLevel = reader.ReadByte();
+				if (reader.GetLeft())
+				{
+					m_resendCas = reader.ReadBool();
+					if (m_resendCas)
+						m_checkExistsOnServer = false;
+				}
 
 			});
 
diff --git a/Engine/Source/Programs/UnrealBuildAccelerator/Common/Private/UbaStorageServer.cpp b/Engine/Source/Programs/UnrealBuildAccelerator/Common/Private/UbaStorageServer.cpp
--- a/Engine/Source/Programs/UnrealBuildAccelerator/Common/Private/UbaStorageServer.cpp
+++ b/Engine/Source/Programs/UnrealBuildAccelerator/Common/Private/UbaStorageServer.cpp
@@ -18,6 +18,7 @@ namespace uba
 		table.GetValueAsBool(writeReceivedCasFilesToDisk, TC("WriteReceivedCasFilesToDisk"));
 		table.GetValueAsBool(allowHintAsFallback, TC("AllowHintAsFallback"));
 		table.GetValueAsBool(createIndependentMappings, TC("CreateIndependentMappings"));
+		table.GetValueAsBool(dropIndependentMappingsOnCopy, TC("DropIndependentMappingsOnCopy"));
 	}
 
 	StorageServer::StorageServer(const StorageServerCreateInfo& info)
@@ -28,6 +29,7 @@ namespace uba
 		m_allowHintAsFallback = info.allowHintAsFallback;
 		m_writeReceivedCasFilesToDisk = info.writeReceivedCasFilesToDisk;
 		m_createIndependentMappings = info.createIndependentMappings;
+		m_dropIndependentMappingsOnCopy = info.dropIndependentMappingsOnCopy;
 
 		if (!CreateGuid(m_uid))
 			UBA_ASSERT(false);
@@ -141,6 +143,14 @@ namespace uba
 		return waitEntry.success;
 	}
 
+	bool StorageServer::CopyOrLink(const CasKey& casKey, const tchar* destination, u32 fileAttributes, bool writeCompressed, u32 clientId, const FormattingFunc& formattingFunc, bool isTemp, bool allowHardLink)
+	{
+		bool res = StorageImpl::CopyOrLink(casKey, destination, fileAttributes, writeCompressed, clientId, formattingFunc, isTemp, allowHardLink);
+		if (clientId != 0)
+			ReleaseActiveRefCount(casKey, clientId);
+		return res;
+	}
+
 	void StorageServer::ActiveFetch::Release(StorageServer& server, const tchar* reason)
 	{
 		if (mappedView.handle.IsValid())
@@ -171,8 +181,7 @@ namespace uba
 			memoryBegin = nullptr;
 		}
 
-		SCOPED_WRITE_LOCK(casEntry->lock, entryLock);
-		--casEntry->readCount;
+		server.ReleaseCasEntry(*casEntry);
 	}
 
 	void StorageServer::OnDisconnected(u32 clientId)
@@ -262,12 +271,71 @@ namespace uba
 				it = m_activeFetches.erase(it);
 			}
 		}
+
+		{
+			SCOPED_WRITE_LOCK(m_activeRefCountsLock, refCountLock);
+			auto findIt = m_activeRefCounts.find(clientId);
+			if (findIt != m_activeRefCounts.end())
+			{
+				auto refCounts = std::move(findIt->second);
+				m_activeRefCounts.erase(findIt);
+				refCountLock.Leave();
+				for (CasEntry* casEntry : refCounts)
+					ReleaseCasEntry(*casEntry);
+			}
+		}
+
 		{
 			SCOPED_WRITE_LOCK(m_connectionInfoLock, lock);
 			m_connectionInfo.erase(clientId);
 		}
 	}
 
+	void StorageServer::ReleaseCasEntry(CasEntry& casEntry)
+	{
+		SCOPED_WRITE_LOCK(casEntry.lock, casEntryLock);
+		ReleaseCasEntryNoLock(casEntry);
+	}
+
+	void StorageServer::ReleaseCasEntryNoLock(CasEntry& casEntry)
+	{
+		if (--casEntry.refCount)
+			return;
+		if (!casEntry.dropped || !m_dropIndependentMappingsOnCopy)
+			return;
+		if (!casEntry.mappingIsIndependent || !casEntry.mappingHandle.IsValid())
+			return;
+		FileMappingHandle mapping = casEntry.mappingHandle;
+		casEntry.mappingHandle = {};
+		casEntry.mappingIsIndependent = false;
+		casEntry.verified = true;
+		casEntry.exists = false;
+		casEntry.dropped = false;
+
+		m_server.AddWork([this, mapping](const WorkContext& context)
+			{
+				CloseFileMapping(m_logger, mapping, TC("ReleaseCasEntryNoLock"));
+			}, 1, TC("ReleaseCasEntryNoLock"));
+	}
+
+	void StorageServer::ReleaseActiveRefCount(const CasKey& casKey, u32 clientId)
+	{
+		SCOPED_WRITE_LOCK(m_activeRefCountsLock, refCountLock);
+		auto findIt = m_activeRefCounts.find(clientId);
+		if (findIt == m_activeRefCounts.end())
+			return;
+		for (auto it=findIt->second.begin();it!=findIt->second.end(); ++it)
+		{
+			CasEntry* casEntry = *it;
+			if (casEntry->key != casKey)
+				continue;
+			findIt->second.erase(it);
+			refCountLock.Leave();
+			ReleaseCasEntry(*casEntry);
+			break;
+		}
+	}
+
 	void StorageServer::UpdateWaitEntries(const CasKey& casKey, bool success)
 	{
 		SCOPED_FUTEX(m_waitEntriesLock, waitLock);
@@ -469,6 +537,8 @@ namespace uba
 				writer.WriteGuid(m_uid);
 				writer.WriteByte(m_casCompressor);
 				writer.WriteByte(m_casCompressionLevel);
+				bool resendCas = m_dropIndependentMappingsOnCopy;
+				writer.WriteBool(resendCas);
 				return true;
 			}
 
@@ -527,8 +597,8 @@ namespace uba
 				auto traceFetchGuard = MakeGuard([&](){ if (detailedTrace) m_trace->FileFetchEnd(connectionInfo.GetId(), casKey); });
 
 
-				CasEntry* casEntry = nullptr;
-				bool has = HasCasFile(casKey, &casEntry); // HasCasFile also writes deferred cas entries if in queue
+				CasEntry* casEntryPtr = nullptr;
+				bool has = HasCasFile(casKey, &casEntryPtr); // HasCasFile also writes deferred cas entries if in queue
 				if (!has)
 				{
 					todoRemoveMe(3);
@@ -556,7 +626,7 @@ namespace uba
 							m_logger.Info(TC("Server did not find cas %s for %s in file table lookup. Recalculating cas key"), CasKeyString(casKey).str, hint.data);
 							if (!CalculateCasKey(checkedCasKey, hint.data))
 							{
-								m_logger.Error(TC("FetchBegin failed for cas file %s (%s) requested by %s. Can't calculate cas key for file"), CasKeyString(casKey).str, hint.data, GuidToString(connectionInfo.GetUid()).str);
+								m_logger.Error(TC("FetchBegin failed for cas file %s (%s) requested by client &u. Can't calculate cas key for file"), CasKeyString(casKey).str, hint.data, connectionInfo.GetId());
 								writer.WriteU16(0);
 								return false;
 							}
@@ -584,23 +654,26 @@ namespace uba
 						writer.WriteU16(0);
 						return true;
 					}
-					casEntry = &findIt->second;
+					casEntryPtr = &findIt->second;
 				}
 
-				SCOPED_WRITE_LOCK(casEntry->lock, entryLock);
-				if (!casEntry->exists)
+				CasEntry& casEntry = *casEntryPtr;
+				SCOPED_WRITE_LOCK(casEntry.lock, entryLock);
+				if (!casEntry.exists)
 				{
 					m_logger.Info(TC("Client %u failed FetchBegin for cas file %s (%s). Cas was deleted"), connectionInfo.GetId(), CasKeyString(casKey).str, hint.data);
 					writer.WriteU16(0);
 					return true;
 				}
-				++casEntry->readCount;
+				++casEntry.refCount;
+				if (!casEntry.refCount)
+					m_logger.Error(TC("CasEntry refcount wrap-around during fetch. This is bad, report to uba devs"));
 				entryLock.Leave();
-				auto rlg = MakeGuard([&]() { entryLock.Enter(); --casEntry->readCount; });
+				auto rlg = MakeGuard([&]() { entryLock.Enter(); ReleaseCasEntryNoLock(casEntry); });
 
 				todoRemoveMe(4);
 
-				if (casEntry->disallowed)
+				if (casEntry.disallowed)
 				{
 					writer.WriteU16(FetchCasIdDisallowed);
 					m_logger.Warning(TC("Client %u is asking for cas content of file that is not allowed to be transferred (%s)"), connectionInfo.GetId(), hint.data);
@@ -642,12 +715,12 @@ namespace uba
 
 				todoRemoveMe(5);
 
-				bool useFileMapping = casEntry->mappingHandle.IsValid();
+				bool useFileMapping = casEntry.mappingHandle.IsValid();
 				if (useFileMapping)
 				{
-					mappedView = m_casDataBuffer.MapView(casEntry->mappingHandle, casEntry->mappingOffset, casEntry->mappingSize, CasKeyString(casKey).str, casEntry->mappingIsIndependent);
+					mappedView = m_casDataBuffer.MapView(casEntry.mappingHandle, casEntry.mappingOffset, casEntry.mappingSize, CasKeyString(casKey).str, casEntry.mappingIsIndependent);
 					memoryBegin = mappedView.memory;
-					fileSize = casEntry->mappingSize;
+					fileSize = casEntry.mappingSize;
 					if (!memoryBegin)
 						return m_logger.Error(TC("Failed to map memory map for %s. Will use file handle instead (%s)"), CasKeyString(casKey).str, LastErrorToText().data);
 					memoryPos = memoryBegin;
@@ -782,7 +855,7 @@ namespace uba
 				fetch.memoryBegin = memoryBegin;
 				fetch.memoryPos = memoryPos;
 				fetch.left = left;
-				fetch.casEntry = casEntry;
+				fetch.casEntry = &casEntry;
 				fetch.sendCasTime = GetTime() - start;
 
 				return true;
@@ -859,16 +932,19 @@ namespace uba
 				bool exists = casEntry.verified && casEntry.exists;
 
 				if (!exists && casEntry.exists)
-				{
 					if (!VerifyExists(exists, casEntry, entryLock, casKey))
 						return false;
-				}
+
+				UBA_ASSERTF(!m_dropIndependentMappingsOnCopy, TC("Not supported... ref count is increased in next"));
+
 				writer.WriteBool(exists);
 				return true;
 			}
 
 			case StorageMessageType_StoreBegin:
 			{
+				u32 clientId = connectionInfo.GetId();
+
 				u64 start = GetTime();
 				CasKey casKey = reader.ReadCasKey();
 				u64 fileSize = reader.ReadU64();
@@ -882,8 +958,21 @@ namespace uba
 				CasEntry& casEntry = casInsres.first->second;
 				lookupLock.Leave();
 
+				auto failGuard = MakeGuard([&]() { ReleaseActiveRefCount(casEntry.key, clientId); });
+				auto success = [&]() { failGuard.Cancel(); return true; };
+
 				SCOPED_WRITE_LOCK(casEntry.lock, entryLock);
 
+				// We need to take ref count on cas entry early since there are a bunch of sections
+				// that temporarily release cas entry lock where file could actually be dropped
+				{
+					++casEntry.refCount;
+					if (!casEntry.refCount)
+						m_logger.Error(TC("CasEntry refcount wrap-around during store. This is bad, report to uba devs"));
+					SCOPED_WRITE_LOCK(m_activeRefCountsLock, refCountLock);
+					m_activeRefCounts[clientId].push_back(&casEntry);
+				}
+
 				{
 					// Check if there is a deferred creation of this file, if there is we can just as well just use the local file
 					SCOPED_FUTEX(m_deferredCasCreationLookupLock, deferredLock);
@@ -896,14 +985,13 @@ namespace uba
 						{
 							writer.WriteU16(u16(~0));
 							writer.WriteBool(m_traceStore);
-							m_logger.Debug(TC("Client %u Store request of %s which already exists in casdb (was in deferred list) (%s)"), connectionInfo.GetId(), CasKeyString(casKey).str, hint.data);
-							return true;
+							m_logger.Debug(TC("Client %u Store request of %s which already exists in casdb (was in deferred list) (%s)"), clientId, CasKeyString(casKey).str, hint.data);
+							return success();
 						}
 						entryLock.Enter();
 					}
 				}
 
-
 				if (!casEntry.verified)
 				{
 					if (bool exists = casEntry.exists)
@@ -915,8 +1003,8 @@ namespace uba
 							// Unless writeReceivedCasFilesToDisk is set this is extremely unlikely to happen. "PerModuleInline.gen.cpp" seems to be the only one (they are supposed to all be the same)
 							writer.WriteU16(u16(~0));
 							writer.WriteBool(m_traceStore);
-							m_logger.Debug(TC("Client %u Store request of %s which already exists in casdb (%s)"), connectionInfo.GetId(), CasKeyString(casKey).str, hint.data);
-							return true;
+							m_logger.Debug(TC("Client %u Store request of %s which already exists in casdb (%s)"), clientId, CasKeyString(casKey).str, hint.data);
+							return success();
 						}
 					}
 					else
@@ -933,13 +1021,13 @@ namespace uba
 						CasEntryAccessed(casEntry);
 						writer.WriteU16(u16(~0));
 						writer.WriteBool(m_traceStore);
-						//m_logger.Debug(TC("Client %u Store request of %s which was just added to casdb (%s)"), connectionInfo.GetId(), CasKeyString(casKey).str, hint.data);
-						return true;
+						//m_logger.Debug(TC("Client %u Store request of %s which was just added to casdb (%s)"), clientId, CasKeyString(casKey).str, hint.data);
+						return success();
 					}
 				}
 				if (!fileSize)
 				{
-					m_logger.Error(TC("Client %u Store is of 0 size (%s)"), connectionInfo.GetId(), hint.data);
+					m_logger.Error(TC("Client %u Store is of 0 size (%s)"), clientId, hint.data);
 					casEntry.verified = false;
 					return false;
 				}
@@ -998,7 +1086,9 @@ namespace uba
 				firstStore->recvCasTime = GetTime() - start;
 
 				if (m_trace)
-					m_trace->FileStoreBegin(connectionInfo.GetId(), casKey, fileSize, hint, m_traceStore);
+					m_trace->FileStoreBegin(clientId, casKey, fileSize, hint, m_traceStore);
+
+				success();
 
 				[[fallthrough]];
 			}
diff --git a/Engine/Source/Programs/UnrealBuildAccelerator/Common/Public/UbaStorage.h b/Engine/Source/Programs/UnrealBuildAccelerator/Common/Public/UbaStorage.h
--- a/Engine/Source/Programs/UnrealBuildAccelerator/Common/Public/UbaStorage.h
+++ b/Engine/Source/Programs/UnrealBuildAccelerator/Common/Public/UbaStorage.h
@@ -55,7 +55,7 @@ namespace uba
 		virtual bool CalculateCasKey(CasKey& out, const tchar* fileName) = 0;
 
 		using FormattingFunc = Function<bool(MemoryBlock& destData, const void* sourceData, u64 sourceSize, const tchar* hint)>;
-		virtual bool CopyOrLink(const CasKey& casKey, const tchar* destination, u32 fileAttributes, bool writeCompressed = false, const FormattingFunc& formattingFunc = {}, bool isTemp = false, bool allowHardLink = true) = 0;
+		virtual bool CopyOrLink(const CasKey& casKey, const tchar* destination, u32 fileAttributes, bool writeCompressed = false, u32 clientId = 0, const FormattingFunc& formattingFunc = {}, bool isTemp = false, bool allowHardLink = true) = 0;
 		virtual bool FakeCopy(const CasKey& casKey, const tchar* destination, u64 size, u64 lastWritten, u32 fileAttributes, bool deleteExisting = true) = 0;
 		virtual bool GetCasFileName(StringBufferBase& out, const CasKey& casKey) = 0;
 
@@ -113,7 +113,6 @@ namespace uba
 		bool writeToDisk = true;
 		bool keepTransientDataMapped = true;	// Will keep the transient data mapped in instead of doing map/unmap everytime we access it
 		bool allowDeleteVerified = false;
-		bool dropIndependentMappingsOnCopy = true; // Will drop independent mappings once they have been copyorlink to destination
 		MutexHandle exclusiveMutex = InvalidMutexHandle; // Set this to hand over pre-exclusive access to storage. Storage will release mutex on shutdown
 		u8 casCompressor = 0;
 		u8 casCompressionLevel = 0;
@@ -168,7 +167,7 @@ namespace uba
 		virtual bool DropCasFile(const CasKey& casKey, bool forceDelete, const tchar* hint) override;
 		virtual bool ReportBadCasFile(const CasKey& casKey) override;
 		virtual bool CalculateCasKey(CasKey& out, const tchar* fileName) override;
-		virtual bool CopyOrLink(const CasKey& casKey, const tchar* destination, u32 fileAttributes, bool writeCompressed = false, const FormattingFunc& formattingFunc = {}, bool isTemp = false, bool allowHardLink = true) override;
+		virtual bool CopyOrLink(const CasKey& casKey, const tchar* destination, u32 fileAttributes, bool writeCompressed = false, u32 clientId = 0, const FormattingFunc& formattingFunc = {}, bool isTemp = false, bool allowHardLink = true) override;
 		virtual bool FakeCopy(const CasKey& casKey, const tchar* destination, u64 size, u64 lastWritten, u32 fileAttributes, bool deleteExisting = true) override;
 		virtual bool GetCasFileName(StringBufferBase& out, const CasKey& casKey) override;
 		virtual MappedView MapView(const CasKey& casKey, const tchar* hint) override;
@@ -252,10 +251,10 @@ namespace uba
 			CasEntry* prevAccessed = nullptr;
 			CasEntry* nextAccessed = nullptr;
 			u64 size = 0; // Stored size of cas entry. Can be both compressed or uncompressed
-			u16 readCount = 0; // This is set while file is being read over network
+			u16 refCount = 0; // This is set while file is being read or queued to be written over network
 			bool verified = false; // This flag needs to be set for below flags to be reliable. if this is false below flags are assumptions
 			bool exists = false; // File exists on disk
-			bool dropped = false; // This file is not seen anymore. will be deleted during shutdown
+			bool dropped = false; // This file is not seen anymore. will be deleted when refCount reaches zero or during shutdown
 			bool beingWritten = false; // This is set while file is being written (when coming from network)..
 			bool disallowed = false; // This is set if cas is created from disallowed file
 
@@ -284,7 +283,6 @@ namespace uba
 		bool m_asyncUnmapViewOfFile = true;
 		bool m_allowDeleteVerified = false;
 		bool m_writeToDisk;
-		bool m_dropIndependentMappingsOnCopy = true;
 
 		MutexHandle m_exclusiveMutex = InvalidMutexHandle;
 
diff --git a/Engine/Source/Programs/UnrealBuildAccelerator/Common/Public/UbaStorageClient.h b/Engine/Source/Programs/UnrealBuildAccelerator/Common/Public/UbaStorageClient.h
--- a/Engine/Source/Programs/UnrealBuildAccelerator/Common/Public/UbaStorageClient.h
+++ b/Engine/Source/Programs/UnrealBuildAccelerator/Common/Public/UbaStorageClient.h
@@ -25,7 +25,7 @@ namespace uba
 		bool sendCompressed = true;
 		bool allowProxy = true;
 		bool sendOneBigFileAtTheTime = true;
-		bool checkExistsOnServer = false; // There is no point having this enabled if server is not received cas files to disk (which is default)
+		bool checkExistsOnServer = false; // There is no point having this enabled if server is not writing received cas files to disk (which is default)
 		bool resendCas = false; // Will try to send created file everytime even if the cas has been seen before
 		
 		u32 proxyConnectionCount = 4;
diff --git a/Engine/Source/Programs/UnrealBuildAccelerator/Common/Public/UbaStorageServer.h b/Engine/Source/Programs/UnrealBuildAccelerator/Common/Public/UbaStorageServer.h
--- a/Engine/Source/Programs/UnrealBuildAccelerator/Common/Public/UbaStorageServer.h
+++ b/Engine/Source/Programs/UnrealBuildAccelerator/Common/Public/UbaStorageServer.h
@@ -17,6 +17,7 @@ namespace uba
 		bool allowHintAsFallback = true; // Will fallback to file system to recalculate cas if cas content does not exist anymore
 		bool writeReceivedCasFilesToDisk = false; // Write received files to disk. This is only useful if we expect same file to be written multiple times (requires deterministic outputs ofc)
 		bool createIndependentMappings = false; // Set to true if we want every store to end up in its own file mapping. This will use more memory but makes it possible to delete files
+		bool dropIndependentMappingsOnCopy = true; // Will drop independent mappings once they have been copyorlink to destination
 		const tchar* zone = TC("");
 	};
 
@@ -35,6 +36,8 @@ namespace uba
 
 		void WaitForActiveWork();
 
+		virtual bool CopyOrLink(const CasKey& casKey, const tchar* destination, u32 fileAttributes, bool writeCompressed = false, u32 clientId = 0, const FormattingFunc& formattingFunc = {}, bool isTemp = false, bool allowHardLink = true) override;
+
 	protected:
 
 		virtual bool GetZone(StringBufferBase& out) override;
@@ -44,6 +47,9 @@ namespace uba
 		virtual bool HasProxy(u32 clientId) override;
 
 		void OnDisconnected(u32 clientId);
+		void ReleaseCasEntry(CasEntry& casEntry);
+		void ReleaseCasEntryNoLock(CasEntry& casEntry);
+		void ReleaseActiveRefCount(const CasKey& casKey, u32 clientId);
 		void UpdateWaitEntries(const CasKey& casKey, bool success);
 		bool HandleMessage(const ConnectionInfo& connectionInfo, const WorkContext& workContext, u8 messageType, BinaryReader& reader, BinaryWriter& writer);
 		bool WaitForWritten(CasEntry& casEntry, ScopedWriteLock& entryLock, const ConnectionInfo& connectionInfo, const tchar* hint);
@@ -108,6 +114,8 @@ namespace uba
 		UnorderedMap<u16, ActiveFetch> m_activeFetches;
 		Atomic<u32> m_activeUnmap;
 
+		ReaderWriterLock m_activeRefCountsLock;
+		UnorderedMap<u32, List<CasEntry*>> m_activeRefCounts;
 
 		Futex m_availableIdsLock;
 		Vector<u16> m_availableIds;
@@ -142,5 +150,6 @@ namespace uba
 		bool m_allowHintAsFallback;
 		bool m_writeReceivedCasFilesToDisk;
 		bool m_createIndependentMappings;
+		bool m_dropIndependentMappingsOnCopy;
 	};
 }
diff --git a/Engine/Source/Programs/UnrealBuildAccelerator/Test/Private/UbaTestAll.h b/Engine/Source/Programs/UnrealBuildAccelerator/Test/Private/UbaTestAll.h
--- a/Engine/Source/Programs/UnrealBuildAccelerator/Test/Private/UbaTestAll.h
+++ b/Engine/Source/Programs/UnrealBuildAccelerator/Test/Private/UbaTestAll.h
@@ -37,6 +37,7 @@ namespace uba
 		UBA_TEST(TestStorage) \
 		UBA_TEST(TestRemoteStorageStore) \
 		UBA_TEST(TestRemoteStorageFetch) \
+		UBA_TEST(TestRemoteStorageStore2) \
 		UBA_TEST(TestDetouredTestApp) \
 		UBA_TEST(TestRemoteDetouredTestApp) \
 		UBA_TEST(TestCompactPathTable) \
diff --git a/Engine/Source/Programs/UnrealBuildAccelerator/Test/Private/UbaTestStorage.cpp b/Engine/Source/Programs/UnrealBuildAccelerator/Test/Private/UbaTestStorage.cpp
--- a/Engine/Source/Programs/UnrealBuildAccelerator/Test/Private/UbaTestStorage.cpp
+++ b/Engine/Source/Programs/UnrealBuildAccelerator/Test/Private/UbaTestStorage.cpp
@@ -227,4 +227,100 @@ namespace uba
 
 		return true;
 	}
+
+	bool TestRemoteStorageStore2(LoggerWithWriter& logger, const StringBufferBase& testRootDir)
+	{
+		auto& logWriter = logger.m_writer;
+		NetworkBackendTcp serverTcp(logWriter, TC("ServerTcp"));
+		NetworkBackendTcp clientTcp(logWriter, TC("ClientTcp"));
+
+		bool ctorSuccess = true;
+		NetworkServer server(ctorSuccess, logWriter);
+		NetworkClient client(ctorSuccess, logWriter);
+		auto dg = MakeGuard([&]() { client.Disconnect(); });
+
+		StringBuffer<> rootDir;
+		rootDir.Append(testRootDir).Append(TCV("Uba"));
+		if (!DeleteAllFiles(logger, rootDir.data))
+			return false;
+
+		StorageServerCreateInfo storageServerInfo(server, rootDir.data, logger.m_writer);
+		storageServerInfo.casCapacityBytes = 1024ull * 1024 * 1024;
+		storageServerInfo.createIndependentMappings = true;
+		StorageServer storageServer(storageServerInfo);
+
+		auto g = MakeGuard([&]() { server.DisconnectClients(); });
+
+		rootDir.Append(TCV("Client"));
+		if (!DeleteAllFiles(logger, rootDir.data))
+			return false;
+
+		StorageClientCreateInfo storageClientInfo(client, rootDir.data);
+		storageClientInfo.resendCas = true;
+		StorageClient storageClient(storageClientInfo);
+		storageClient.Start();
+
+		if (!storageClient.LoadCasTable(true))
+			return false;
+
+		rootDir.EnsureEndsWithSlash();
+
+		if (!server.StartListen(serverTcp, 1234))
+			return logger.Error(TC("Failed to listen"));
+		Sleep(100);
+		if (!client.Connect(clientTcp, TC("127.0.0.1"), 1234))
+			return logger.Error(TC("Failed to connect"));
+		auto cg = MakeGuard([&]() { client.Disconnect(); });
+
+		StringBuffer<> fileName;
+		{
+			fileName.Append(rootDir).Append(TCV("UbaTestFile"));
+			FileAccessor fileHandle(logger, fileName.data);
+			if (!fileHandle.CreateWrite())
+				return logger.Error(TC("Failed to create file for write"));
+			u8 byte = 'H';
+			if (!fileHandle.Write(&byte, 1))
+				return false;
+			if (!fileHandle.Close())
+				return false;
+		}
+
+		CasKey key;
+		auto store = [&]()
+			{
+				bool storeCompressed = false;
+				if (!storageClient.StoreCasFileClient(key, ToStringKeyLower(fileName), fileName.data, FileMappingHandle{}, 0, 0, TC("UbaTestFile"), false, storeCompressed))
+					return logger.Error(TC("Failed to store file %s"), fileName.data);
+				return true;
+			};
+
+		auto copy = [&]()
+			{
+				fileName.Clear().Append(testRootDir).Append(TCV("Uba")).Append(PathSeparator).Append(TCV("UbaTestFile"));
+				if (!storageServer.CopyOrLink(key, fileName.data, DefaultAttributes(), false, 1, {}, false, true))
+					return logger.Error(TC("Failed to copy cas to file %s"), fileName.data);
+				return true;
+			};
+
+		if (!store())
+			return false;
+		if (!copy())
+			return false;
+		if (!store())
+			return false;
+		if (!copy())
+			return false;
+
+		if (!store())
+			return false;
+		if (!store())
+			return false;
+		if (!copy())
+			return false;
+		if (!copy())
+			return false;
+
+		return true;
+	}
+
 }
