[UBA]
* Added so wait entries are updated if StoreBegin hits a path where cas was already uploaded

diff --git a/Engine/Source/Programs/UnrealBuildAccelerator/Common/Private/UbaStorageServer.cpp b/Engine/Source/Programs/UnrealBuildAccelerator/Common/Private/UbaStorageServer.cpp
--- a/Engine/Source/Programs/UnrealBuildAccelerator/Common/Private/UbaStorageServer.cpp
+++ b/Engine/Source/Programs/UnrealBuildAccelerator/Common/Private/UbaStorageServer.cpp
@@ -986,6 +986,7 @@ namespace uba
 							writer.WriteU16(u16(~0));
 							writer.WriteBool(m_traceStore);
 							m_logger.Debug(TC("Client %u Store request of %s which already exists in casdb (was in deferred list) (%s)"), clientId, CasKeyString(casKey).str, hint.data);
+							UpdateWaitEntries(casEntry.key, true);
 							return success();
 						}
 						entryLock.Enter();
@@ -1004,6 +1005,7 @@ namespace uba
 							writer.WriteU16(u16(~0));
 							writer.WriteBool(m_traceStore);
 							m_logger.Debug(TC("Client %u Store request of %s which already exists in casdb (%s)"), clientId, CasKeyString(casKey).str, hint.data);
+							UpdateWaitEntries(casEntry.key, true);
 							return success();
 						}
 					}
@@ -1022,6 +1024,7 @@ namespace uba
 						writer.WriteU16(u16(~0));
 						writer.WriteBool(m_traceStore);
 						//m_logger.Debug(TC("Client %u Store request of %s which was just added to casdb (%s)"), clientId, CasKeyString(casKey).str, hint.data);
+						UpdateWaitEntries(casEntry.key, true);
 						return success();
 					}
 				}
