[UBA]
* Changed so HasCasFile increase ref count of CasEntry if found and out pointer is set. This removes a small window where file could get dropped even though it was about to be used

diff --git a/Engine/Source/Programs/UnrealBuildAccelerator/Common/Private/UbaCacheClient.cpp b/Engine/Source/Programs/UnrealBuildAccelerator/Common/Private/UbaCacheClient.cpp
--- a/Engine/Source/Programs/UnrealBuildAccelerator/Common/Private/UbaCacheClient.cpp
+++ b/Engine/Source/Programs/UnrealBuildAccelerator/Common/Private/UbaCacheClient.cpp
@@ -1314,9 +1314,8 @@ namespace uba
 
 			TimerScope createCasTs(stats.createCas);
 			StorageImpl::CasEntry* casEntry;
-			bool hasCasFile = m_storage.HasCasFile(casKey, &casEntry);
 
-			if (hasCasFile)
+			if (m_storage.HasCasFile(casKey, &casEntry))
 			{
 				createCasTs.Leave();
 
@@ -1338,6 +1337,7 @@ namespace uba
 
 				{
 					SCOPED_READ_LOCK(casEntry->lock, entryLock);
+					--casEntry->refCount; // HasCasFile increase refcount
 					auto mh = casEntry->mappingHandle;
 					if (mh.IsValid()) // If file was created by helper it will be in the transient mapped memory
 					{
diff --git a/Engine/Source/Programs/UnrealBuildAccelerator/Common/Private/UbaStorage.cpp b/Engine/Source/Programs/UnrealBuildAccelerator/Common/Private/UbaStorage.cpp
--- a/Engine/Source/Programs/UnrealBuildAccelerator/Common/Private/UbaStorage.cpp
+++ b/Engine/Source/Programs/UnrealBuildAccelerator/Common/Private/UbaStorage.cpp
@@ -1891,13 +1891,17 @@ namespace uba
 		lookupLock.Leave();
 		CasEntryAccessed(casEntry);
 
-		if (out)
-			*out = &casEntry;
-
 		SCOPED_WRITE_LOCK(casEntry.lock, entryLock);
 
 		if (casEntry.verified && casEntry.exists)
+		{
+			if (out)
+			{
+				*out = &casEntry;
+				++casEntry.refCount;
+			}
 			return true;
+		}
 
 		SCOPED_FUTEX(m_deferredCasCreationLookupLock, deferredLock);
 		auto findIt = m_deferredCasCreationLookup.find(casKey);
@@ -1938,6 +1942,13 @@ namespace uba
 		}
 
 		casEntry.exists = true;
+
+		if (out)
+		{
+			*out = &casEntry;
+			++casEntry.refCount;
+		}
+
 		entryLock.Leave();
 
 		if (!hasMapping)
@@ -2174,6 +2185,7 @@ namespace uba
 			SCOPED_READ_LOCK(casEntry.lock, casEntryLock);
 			UBA_ASSERT(casEntry.verified);
 			UBA_ASSERT(casEntry.exists);
+			--casEntry.refCount; // Decrease after HasCasFile increase it
 
 			auto done = [&](FileAccessor& destinationFile, u64 sizeOnDisk)
 				{
diff --git a/Engine/Source/Programs/UnrealBuildAccelerator/Common/Private/UbaStorageServer.cpp b/Engine/Source/Programs/UnrealBuildAccelerator/Common/Private/UbaStorageServer.cpp
--- a/Engine/Source/Programs/UnrealBuildAccelerator/Common/Private/UbaStorageServer.cpp
+++ b/Engine/Source/Programs/UnrealBuildAccelerator/Common/Private/UbaStorageServer.cpp
@@ -598,7 +598,7 @@ namespace uba
 
 
 				CasEntry* casEntryPtr = nullptr;
-				bool has = HasCasFile(casKey, &casEntryPtr); // HasCasFile also writes deferred cas entries if in queue
+				bool has = HasCasFile(casKey, &casEntryPtr); // HasCasFile writes deferred cas entries if in queue. Also increase CasEntry refcount
 				if (!has)
 				{
 					todoRemoveMe(3);
@@ -659,17 +659,19 @@ namespace uba
 
 				CasEntry& casEntry = *casEntryPtr;
 				SCOPED_WRITE_LOCK(casEntry.lock, entryLock);
+				if (!has) // HasCasFile increase refcount, otherwise we have to do it manually
+					++casEntry.refCount;
+				if (!casEntry.refCount)
+					m_logger.Error(TC("CasEntry refcount wrap-around during fetch. This is bad, report to uba devs"));
+				entryLock.Leave();
+				auto rlg = MakeGuard([&]() { entryLock.Enter(); ReleaseCasEntryNoLock(casEntry); });
+
 				if (!casEntry.exists)
 				{
 					m_logger.Info(TC("Client %u failed FetchBegin for cas file %s (%s). Cas was deleted"), connectionInfo.GetId(), CasKeyString(casKey).str, hint.data);
 					writer.WriteU16(0);
 					return true;
 				}
-				++casEntry.refCount;
-				if (!casEntry.refCount)
-					m_logger.Error(TC("CasEntry refcount wrap-around during fetch. This is bad, report to uba devs"));
-				entryLock.Leave();
-				auto rlg = MakeGuard([&]() { entryLock.Enter(); ReleaseCasEntryNoLock(casEntry); });
 
 				todoRemoveMe(4);
 
diff --git a/Engine/Source/Programs/UnrealBuildAccelerator/Test/Private/UbaTestStorage.cpp b/Engine/Source/Programs/UnrealBuildAccelerator/Test/Private/UbaTestStorage.cpp
--- a/Engine/Source/Programs/UnrealBuildAccelerator/Test/Private/UbaTestStorage.cpp
+++ b/Engine/Source/Programs/UnrealBuildAccelerator/Test/Private/UbaTestStorage.cpp
@@ -256,7 +256,6 @@ namespace uba
 			return false;
 
 		StorageClientCreateInfo storageClientInfo(client, rootDir.data);
-		storageClientInfo.resendCas = true;
 		StorageClient storageClient(storageClientInfo);
 		storageClient.Start();
 
@@ -288,7 +287,7 @@ namespace uba
 		CasKey key;
 		auto store = [&]()
 			{
-				bool storeCompressed = false;
+				bool storeCompressed = true;
 				if (!storageClient.StoreCasFileClient(key, ToStringKeyLower(fileName), fileName.data, FileMappingHandle{}, 0, 0, TC("UbaTestFile"), false, storeCompressed))
 					return logger.Error(TC("Failed to store file %s"), fileName.data);
 				return true;
@@ -320,6 +319,16 @@ namespace uba
 		if (!copy())
 			return false;
 
+		fileName.Clear().Append(testRootDir).Append(TCV("Uba")).Append(PathSeparator).Append(TCV("UbaTestFile"));
+		if (!storageServer.StoreCasFile(key, fileName.data, {}, true))
+			return false;
+
+		//client.Fetch
+		if (!store())
+			return false;
+		//if (!copy())
+		//	return false;
+
 		return true;
 	}
 
