[UBA]
* Changed so createIndependentMappings is defaulting to true. This seems to always be beneficial when not linking remotely since it saves memory. It has been properly tested by the epic build farm over multiple days (thousands of runs)

diff --git a/Engine/Source/Programs/UnrealBuildAccelerator/Common/Public/UbaStorageServer.h b/Engine/Source/Programs/UnrealBuildAccelerator/Common/Public/UbaStorageServer.h
--- a/Engine/Source/Programs/UnrealBuildAccelerator/Common/Public/UbaStorageServer.h
+++ b/Engine/Source/Programs/UnrealBuildAccelerator/Common/Public/UbaStorageServer.h
@@ -16,7 +16,7 @@ namespace uba
 		NetworkServer& server;
 		bool allowHintAsFallback = true; // Will fallback to file system to recalculate cas if cas content does not exist anymore
 		bool writeReceivedCasFilesToDisk = false; // Write received files to disk. This is only useful if we expect same file to be written multiple times (requires deterministic outputs ofc)
-		bool createIndependentMappings = false; // Set to true if we want every store to end up in its own file mapping. This will use more memory but makes it possible to delete files
+		bool createIndependentMappings = true; // Set to true if we want every store to end up in its own file mapping. This will use more memory but makes it possible to delete files
 		bool dropIndependentMappingsOnCopy = true; // Will drop independent mappings once they have been copyorlink to destination
 		const tchar* zone = TC("");
 	};
