[UBA]
* Fixed UbaVisualizer network path (-host=xyz) when running multiple traces in a row on the host.
* Added so port can be provided with -host

diff --git a/Engine/Source/Programs/UnrealBuildAccelerator/Visualizer/Private/UbaVisualizer.cpp b/Engine/Source/Programs/UnrealBuildAccelerator/Visualizer/Private/UbaVisualizer.cpp
--- a/Engine/Source/Programs/UnrealBuildAccelerator/Visualizer/Private/UbaVisualizer.cpp
+++ b/Engine/Source/Programs/UnrealBuildAccelerator/Visualizer/Private/UbaVisualizer.cpp
@@ -320,7 +320,10 @@ namespace uba
 
 	bool Visualizer::ShowUsingSocket(NetworkBackend& backend, const wchar_t* host, u16 port)
 	{
-		auto destroyClient = MakeGuard([&]() { delete m_client; m_client = nullptr; });
+		NetworkClient* activeClient = nullptr;
+		NetworkClient* oldClient = nullptr;
+
+		auto destroyClient = MakeGuard([&]() { m_client = nullptr; delete activeClient; });
 		m_looping = true;
 		m_autoScroll = false;
 		if (!StartHwndThread())
@@ -334,12 +337,12 @@ namespace uba
 		StringBuffer<256> traceName;
 		while (m_hwnd)
 		{
-			if (!m_client)
+			if (!activeClient)
 			{
 				bool ctorSuccess = true;
 				NetworkClientCreateInfo ncci;
 				ncci.workerCount = 0;
-				m_client = new NetworkClient(ctorSuccess, ncci);
+				activeClient = new NetworkClient(ctorSuccess, ncci);
 				if (!ctorSuccess)
 					return false;
 			}
@@ -348,24 +351,35 @@ namespace uba
 			PostNewTitle(GetTitlePrefix(title).Appendf(L"Trying to connect to %s:%u%s", host, port, dots + ((dotsCounter--) % 4)));
 
 			bool timedOut;
-			if (!m_client->Connect(backend, host, port, &timedOut))
+			if (!activeClient->Connect(backend, host, port, &timedOut))
 				continue;
+			m_client = activeClient;
 
 			PostNewTitle(GetTitlePrefix(title).Appendf(L"Connected to %s:%u", host, port));
 			PostNewTrace(0, false);
 
-			while (m_hwnd && m_client->IsConnected() && !m_clientDisconnect.IsSet(100))
+			if (oldClient)
+			{
+				delete oldClient;
+				oldClient = nullptr;
+			}
+
+			while (m_hwnd && !m_clientDisconnect.IsSet(100) && activeClient->IsConnected())
 			{
-				m_trace.UpdateReceiveClient(*m_client);
+				m_trace.UpdateReceiveClient(*activeClient);
 			}
 
+			while (activeClient->IsConnected())
+			{
+				Sleep(500);
+				activeClient->SendKeepAlive();
+			}
 			PostNewTitle(GetTitlePrefix(title).Appendf(L"Disconnected..."));
 
-			m_client->Disconnect();
-			delete m_client;
-			m_client = nullptr;
+			oldClient = activeClient;
+			activeClient = nullptr;
+
 			m_clientDisconnect.Reset();
-			Sleep(4000); // To prevent it from reconnecting to the same thing again and get thrown out (since it will post a WM_NEWTRACE and clean everything
 		}
 		return true;
 	}
diff --git a/Engine/Source/Programs/UnrealBuildAccelerator/Visualizer/Private/UbaVisualizerApp.cpp b/Engine/Source/Programs/UnrealBuildAccelerator/Visualizer/Private/UbaVisualizerApp.cpp
--- a/Engine/Source/Programs/UnrealBuildAccelerator/Visualizer/Private/UbaVisualizerApp.cpp
+++ b/Engine/Source/Programs/UnrealBuildAccelerator/Visualizer/Private/UbaVisualizerApp.cpp
@@ -23,8 +23,7 @@ static int PrintHelp(const tchar* message = nullptr)
 	s.Appendf(TC("  When started UbaVisualizer will keep trying to connect to provided host address or named memory buffer.\r\n"));
 	s.Appendf(TC("  Once connected it will start visualizing. Nothing else is needed :)\r\n"));
 	s.Appendf(TC("\r\n"));
-	s.Appendf(TC("  -host=<host>         The ip/name of the machine we want to connect to\r\n"));
-	s.Appendf(TC("  -port=<port>         The port to connect to. Defaults to \"%u\"\r\n"), DefaultPort);
+	s.Appendf(TC("  -host=<host>[:<port>] The ip/name and port (default: %u) of the machine we want to connect to\r\n"), DefaultPort);
 	s.Appendf(TC("  -named=<name>        Name of named memory to connect to\r\n"));
 	s.Appendf(TC("  -file=<name>         Name of file to parse\r\n"));
 	s.Appendf(TC("  -listen[=<channel>]  Listen for announcements of new sessions. Defaults to channel '%s'\r\n"), TC("Default"));
@@ -122,8 +121,15 @@ int WINAPI wWinMain(_In_ HINSTANCE hInstance, _In_opt_ HINSTANCE hPrevInstance,
 		}
 		else if (name.Equals(TCV("-host")))
 		{
+			if (const tchar* portIndex = value.First(':'))
+			{
+				StringBuffer<> portStr(portIndex + 1);
+				if (!portStr.Parse(port))
+					return PrintHelp(TC("Invalid value for port in -host"));
+				value.Resize(portIndex - value.data);
+			}
 			if (value.IsEmpty())
-				return PrintHelp(TC("-host needs a value"));
+				return PrintHelp(TC("-host needs a name/ip"));
 			host.Append(value);
 		}
 		else if (name.Equals(TCV("-named")))
