[UBA]
* Fixed so UbaObjTool can read its own produced coff files. IsCoffFile checked if first section was .text or .drectve.. but with last changes our own coff file starts with .rdata

diff --git a/Engine/Source/Programs/UnrealBuildAccelerator/Common/Private/UbaObjectFileCoff.cpp b/Engine/Source/Programs/UnrealBuildAccelerator/Common/Private/UbaObjectFileCoff.cpp
--- a/Engine/Source/Programs/UnrealBuildAccelerator/Common/Private/UbaObjectFileCoff.cpp
+++ b/Engine/Source/Programs/UnrealBuildAccelerator/Common/Private/UbaObjectFileCoff.cpp
@@ -159,9 +159,9 @@ namespace uba
 			return false;
 		if (header.Characteristics != 0) // Should always be 0
 			return false;
-		// We expect .text or .drectve sections first.. this is also hacky.. but don't know how to verify that a file is a coff file
+		// We expect .text, .rdata or .drectve sections first.. this is also hacky.. but don't know how to verify that a file is a coff file
 		const u8* firstSection = data + sizeof(ImageFileHeader);
-		return memcmp(firstSection, ".text", 5) == 0 || memcmp(firstSection, ".drectve", 8) == 0;
+		return memcmp(firstSection, ".text", 5) == 0 || memcmp(firstSection, ".drectve", 8) == 0 || memcmp(firstSection, ".rdata", 6) == 0;
 	}
 
 	template<typename SymbolType>
