[UBA]
* Added session option AllowDiscardVirtualMemory. If set to true, uba will use DiscardVirtualMemory just after mapping in anonymous memory for write and also before fully releasing memory. This seems to have a large impact on performance. Main reason is because when Unmap/close happens the kernel can be half way through writing to page file.. and need to wait for it to finish before releasing. By discarding the kernel will only finish current page and then exit.
* Added LinkCrawlerMaxParallel which defaults to maxcores - 2. This variable can be used to bottleneck the decompression of obj files for the linker

diff --git a/Engine/Source/Programs/UnrealBuildAccelerator/Common/Private/UbaFileMapping.cpp b/Engine/Source/Programs/UnrealBuildAccelerator/Common/Private/UbaFileMapping.cpp
--- a/Engine/Source/Programs/UnrealBuildAccelerator/Common/Private/UbaFileMapping.cpp
+++ b/Engine/Source/Programs/UnrealBuildAccelerator/Common/Private/UbaFileMapping.cpp
@@ -236,13 +236,32 @@ namespace uba
 #endif
 	}
 
-	u8* MapViewOfFile(Logger& logger, FileMappingHandle fileMappingObject, u32 desiredAccess, u64 offset, u64 bytesToMap)
+	static void InternalDiscard(Logger& logger, u8* mem, u64 size, const tchar* from, const tchar* hint)
+	{
+#if PLATFORM_WINDOWS
+		u64 stepSize = 16ull*1024*1024*1024; // Extremely large.. here to test if it is faster to do multiple calls
+		u64 left = size;
+		while (left)
+		{
+			u64 toDiscard = Min(left, stepSize);
+			DWORD res = DiscardVirtualMemory(mem, AlignUp(toDiscard, 4*1024));
+			if (res != ERROR_SUCCESS)
+				logger.Warning(TC("%s - Failed to discard memory (%s) (%u)"), from, hint, res);
+			mem += toDiscard;
+			left -= toDiscard;
+		}
+#endif
+	}
+
+	u8* MapViewOfFile(Logger& logger, FileMappingHandle fileMappingObject, u32 desiredAccess, u64 offset, u64 bytesToMap, bool allowDiscard)
 	{
 		ExtendedTimerScope ts(KernelStats::GetCurrent().mapViewOfFile);
 #if PLATFORM_WINDOWS
 		u8* res = (u8*)::MapViewOfFile(fileMappingObject.mh, desiredAccess, (DWORD)ToHigh(offset), ToLow(offset), bytesToMap);
 		if (!res)
 			return nullptr;
+		if (allowDiscard)
+			InternalDiscard(logger, res, bytesToMap, TC("MapViewOfFile"), TC(""));
 #else
 		int prot = 0;
 		if (desiredAccess & FILE_MAP_READ)
@@ -275,19 +294,23 @@ namespace uba
 		return res;
 	}
 
-	bool MapViewCommit(void* address, u64 size)
+	bool MapViewCommit(Logger& logger, void* address, u64 size, bool allowDiscard)
 	{
 #if PLATFORM_WINDOWS
 		ExtendedTimerScope ts(KernelStats::GetCurrent().virtualAlloc);
-		return ::VirtualAlloc(address, size, MEM_COMMIT, PAGE_READWRITE);
+		if (!::VirtualAlloc(address, size, MEM_COMMIT, PAGE_READWRITE))
+			return false;
+		if (allowDiscard)
+			InternalDiscard(logger, (u8*)address, size, TC("MapViewCommit"), TC(""));
+		return true;
 #else
 		return true;
 #endif
 	}
 
-	bool UnmapViewOfFile(Logger& logger, const void* lpBaseAddress, u64 bytesToUnmap, const tchar* hint)
+	bool UnmapViewOfFile(Logger& logger, const void* address, u64 bytesToUnmap, const tchar* hint, bool allowDiscard)
 	{
-		if (!lpBaseAddress)
+		if (!address)
 			return true;
 
 		ExtendedTimerScope ts(KernelStats::GetCurrent().unmapViewOfFile);
@@ -295,7 +318,7 @@ namespace uba
 		#if UBA_DEBUG_FILE_MAPPING
 		{
 			SCOPED_FUTEX(g_fileMappingsLock, lock);
-			auto findIt = g_viewMappings.find(lpBaseAddress);
+			auto findIt = g_viewMappings.find(address);
 			UBA_ASSERT(findIt != g_viewMappings.end());
 			auto findIt2 = g_fileMappings.find(findIt->second);
 			g_viewMappings.erase(findIt);
@@ -305,17 +328,26 @@ namespace uba
 		#endif
 
 #if PLATFORM_WINDOWS
-		(void)bytesToUnmap;
-		return ::UnmapViewOfFile(lpBaseAddress);
+		if (allowDiscard)
+			InternalDiscard(logger, (u8*)address, bytesToUnmap, TC("UnmapViewOfFile"), hint);
+		return ::UnmapViewOfFile(address);
 #else
-		UBA_ASSERTF(bytesToUnmap, TC("bytesToUnmap is zero unmapping %p (%s)"), lpBaseAddress, hint);
-		if (munmap((void*)lpBaseAddress, bytesToUnmap) == 0)
+		UBA_ASSERTF(bytesToUnmap, TC("bytesToUnmap is zero unmapping %p (%s)"), address, hint);
+		if (munmap((void*)address, bytesToUnmap) == 0)
 			return true;
 		UBA_ASSERT(false);
 		return false;
 #endif
 	}
 
+	void DiscardFileMapping(Logger& logger, FileMappingHandle h, u64 size, const tchar* hint)
+	{
+		#if PLATFORM_WINDOWS
+		void* mem = MapViewOfFile(logger, h, FILE_MAP_WRITE|FILE_MAP_READ, 0, size, true);
+		UnmapViewOfFile(logger, mem, size, TC(""));
+		#endif
+	}
+
 	bool CloseFileMapping(Logger& logger, FileMappingHandle h, const tchar* hint)
 	{
 		if (!h.IsValid())
@@ -461,7 +493,7 @@ namespace uba
 			FileMappingHandle mapping = CreateMemoryMappingW(m_logger, PAGE_READWRITE, size, hint, TC("Independent"));
 			if (!mapping.IsValid())
 				return {};
-			u8* memory = MapViewOfFile(m_logger, mapping, FILE_MAP_WRITE, 0, size);
+			u8* memory = MapViewOfFile(m_logger, mapping, FILE_MAP_WRITE, 0, size, true);
 			if (!memory)
 			{
 				CloseFileMapping(m_logger, mapping, hint);
@@ -528,7 +560,7 @@ namespace uba
 		{
 			u64 commitStart = committedBefore - alignedOffsetStart;
 			u64 commitSize = committedAfter - committedBefore;
-			if (!MapViewCommit(data + commitStart, commitSize))
+			if (!MapViewCommit(m_logger, data + commitStart, commitSize, true))
 			{
 				UnmapViewOfFile(m_logger, data, mapSize, hint);
 				m_logger.Error(TC("%s - Failed to allocate/map %s for %s (%s)"), f.name, BytesToText(commitSize).str, hint, LastErrorToText().data);
@@ -866,7 +898,7 @@ namespace uba
 
 		if (needCommit)
 		{
-			if (!MapViewCommit(data, m_blockSize))
+			if (!MapViewCommit(m_logger, data, m_blockSize, true))
 			{
 				m_logger.Error(TC("%s - Failed to allocate memory for %s (%s)"), m_name, hint, LastErrorToText().data);
 				return { {}, 0, 0 };
diff --git a/Engine/Source/Programs/UnrealBuildAccelerator/Common/Private/UbaProcess.cpp b/Engine/Source/Programs/UnrealBuildAccelerator/Common/Private/UbaProcess.cpp
--- a/Engine/Source/Programs/UnrealBuildAccelerator/Common/Private/UbaProcess.cpp
+++ b/Engine/Source/Programs/UnrealBuildAccelerator/Common/Private/UbaProcess.cpp
@@ -1851,6 +1851,7 @@ namespace uba
 			payload.runningRemote = m_runningRemote;
 			payload.allowKeepFilesInMemory = m_session.m_allowKeepFilesInMemory;
 			payload.allowOutputFiles = m_session.m_allowOutputFiles;
+			payload.allowDiscardVirtualMemory = m_session.m_allowDiscardVirtualMemory;
 			payload.suppressLogging = m_session.m_suppressLogging;
 			payload.isChild = m_parentProcess != nullptr;
 			payload.trackInputs = m_startInfo.trackInputs;
diff --git a/Engine/Source/Programs/UnrealBuildAccelerator/Common/Private/UbaSession.cpp b/Engine/Source/Programs/UnrealBuildAccelerator/Common/Private/UbaSession.cpp
--- a/Engine/Source/Programs/UnrealBuildAccelerator/Common/Private/UbaSession.cpp
+++ b/Engine/Source/Programs/UnrealBuildAccelerator/Common/Private/UbaSession.cpp
@@ -208,6 +208,7 @@ namespace uba
 		table.GetValueAsBool(allowMemoryMaps, TC("AllowMemoryMaps"));
 		table.GetValueAsBool(allowKeepFilesInMemory, TC("AllowKeepFilesInMemory"));
 		table.GetValueAsBool(allowOutputFiles, TC("AllowOutputFiles"));
+		table.GetValueAsBool(allowDiscardVirtualMemory, TC("AllowDiscardVirtualMemory"));
 		table.GetValueAsBool(allowSpecialApplications, TC("AllowSpecialApplications"));
 		table.GetValueAsBool(suppressLogging, TC("SuppressLogging"));
 		table.GetValueAsBool(shouldWriteToDisk, TC("ShouldWriteToDisk"));
@@ -215,6 +216,7 @@ namespace uba
 		table.GetValueAsBool(detailedTrace, TC("DetailedTrace"));
 		table.GetValueAsBool(traceChildProcesses, TC("TraceChildProcesses"));
 		table.GetValueAsBool(traceWrittenFiles, TC("TraceWrittenFiles"));
+		table.GetValueAsInt(linkCrawlerMaxParallel, TC("LinkCrawlerMaxParallel"));
 		table.GetValueAsBool(storeIntermediateFilesCompressed, TC("StoreIntermediateFilesCompressed"));
 		table.GetValueAsBool(allowLocalDetour, TC("AllowLocalDetour"));
 		table.GetValueAsBool(extractObjFilesSymbols, TC("ExtractObjFilesSymbols"));
@@ -249,6 +251,14 @@ namespace uba
 		u64 writtenSize = 0;
 		bool shouldWriteToDisk = ShouldWriteToDisk(file.name);
 
+		u64 fileSize = file.mappingWritten;
+		u8* mem = nullptr;
+		auto memClose = MakeGuard([&]()
+			{
+				if (mem)
+					UnmapViewOfFile(m_logger, mem, fileSize, file.name.c_str(), shouldEvictFromMemory && m_allowDiscardVirtualMemory);
+			});
+
 		if (shouldWriteToDisk)
 		{
 			RootsHandle rootsHandle = process.GetStartInfo().rootsHandle;
@@ -318,15 +328,16 @@ namespace uba
 			m_debugLogger->Info(TC("Writing written file with mapping 0x%llx for %s"), u64(file.mappingHandle.mh), file.name.c_str());
 			#endif
 
-			u64 fileSize = file.mappingWritten;
-			u8* mem = MapViewOfFile(m_logger, file.mappingHandle, FILE_MAP_READ, 0, fileSize);
+			u32 desiredAccess = FILE_MAP_READ;
+			if (m_allowDiscardVirtualMemory)
+				desiredAccess |= FILE_MAP_WRITE;
+
+			mem = MapViewOfFile(m_logger, file.mappingHandle, desiredAccess, 0, fileSize);
 			if (!mem)
 				return m_logger.Error(TC("Failed to map view of filehandle for read %s (%s)"), file.name.c_str(), LastErrorToText().data);
 
 			//PrefetchVirtualMemory(mem, fileSize);
 
-			auto memClose = MakeGuard([&](){ UnmapViewOfFile(m_logger, mem, fileSize, file.name.c_str()); });
-
 			if (storeCompressed)
 			{
 				Storage::WriteResult res;
@@ -443,6 +454,8 @@ namespace uba
 			}
 		}
 
+		memClose.Execute();
+
 		if (shouldEvictFromMemory)
 		{
 			#if UBA_DEBUG_TRACK_MAPPING
@@ -519,7 +532,7 @@ namespace uba
 		if (!mem)
 			return false;
 		bool res = func(mem, size);
-		UnmapViewOfFile(m_logger, mem, size, filePath.data);
+		UnmapViewOfFile(m_logger, mem, size, filePath.data, deleteInternalMapping && m_allowDiscardVirtualMemory);
 		if (deleteInternalMapping)
 			CloseFileMapping(m_logger, mapping, TC(""));
 		return res;
@@ -963,7 +976,7 @@ namespace uba
 				mapping = CreateMemoryMappingW(m_logger, PAGE_READWRITE, contentSize, nullptr, fileName);
 				if (!mapping.IsValid())
 					return false;
-				mappingMemory = MapViewOfFile(m_logger, mapping, FILE_MAP_WRITE, 0, contentSize);
+				mappingMemory = MapViewOfFile(m_logger, mapping, FILE_MAP_WRITE, 0, contentSize, m_allowDiscardVirtualMemory);
 				if (!mappingMemory)
 					return false;
 			}
@@ -1520,6 +1533,14 @@ namespace uba
 		return result;
 	}
 
+	u32 GetCrawlerBottleneck(int linkCrawlerMaxParallel)
+	{
+		if (linkCrawlerMaxParallel <= 0)
+			return u32(Max(0, int(GetLogicalProcessorCount()) + linkCrawlerMaxParallel));
+		else
+			return u32(Min(linkCrawlerMaxParallel, int(GetLogicalProcessorCount())));
+	}
+
 	Session::Session(const SessionCreateInfo& info, const tchar* logPrefix, bool runningRemote, WorkManager& workManager)
 	:	m_storage(info.storage)
 	,	m_logger(info.logWriter, logPrefix)
@@ -1532,7 +1553,7 @@ namespace uba
 	,	m_writeFilesBottleneck(info.writeFilesBottleneck)
 	,	m_writeFilesFileMapMax(u64(info.writeFilesFileMapMaxMb)*1024*1024)
 	,	m_writeFilesNoBufferingMin(u64(info.writeFilesNoBufferingMinMb)*1024*1024)
-	,	m_dependencyCrawler(m_logger, workManager)
+	,	m_dependencyCrawler(m_logger, workManager, GetCrawlerBottleneck(info.linkCrawlerMaxParallel))
 	{
 		UBA_ASSERTF(info.rootDir && *info.rootDir, TC("No root dir set when creating session"));
 		m_rootDir.count = GetFullPathNameW(info.rootDir, m_rootDir.capacity, m_rootDir.data, NULL);
@@ -1586,6 +1607,7 @@ namespace uba
 		m_allowLocalDetour = info.allowLocalDetour;
 		m_extractObjFilesSymbols = info.extractObjFilesSymbols;
 		m_allowLinkDependencyCrawler = info.allowLinkDependencyCrawler;
+		m_allowDiscardVirtualMemory = info.allowDiscardVirtualMemory;
 
 		m_detailedTrace = info.detailedTrace;
 		m_traceChildProcesses = info.traceChildProcesses;
@@ -1675,7 +1697,7 @@ namespace uba
 
 		m_fileMappingTableHandle = uba::CreateMemoryMappingW(m_logger, PAGE_READWRITE, FileMappingTableMemSize, nullptr, TC("FileMappings"));
 		UBA_ASSERT(m_fileMappingTableHandle.IsValid());
-		m_fileMappingTableMem = MapViewOfFile(m_logger, m_fileMappingTableHandle, FILE_MAP_WRITE, 0, FileMappingTableMemSize);
+		m_fileMappingTableMem = MapViewOfFile(m_logger, m_fileMappingTableHandle, FILE_MAP_WRITE, 0, FileMappingTableMemSize, m_allowDiscardVirtualMemory);
 		UBA_ASSERT(m_fileMappingTableMem);
 
 		m_directoryTableHandle = uba::CreateMemoryMappingW(m_logger, PAGE_READWRITE|SEC_RESERVE, DirTableMemSize, nullptr, TC("DirMappings"));
@@ -1813,10 +1835,10 @@ namespace uba
 			if (i.second.canBeFreed)
 				CloseFileMapping(m_logger, i.second.mapping, TC("FileMappingKeptFromOutput"));
 
-		UnmapViewOfFile(m_logger, m_fileMappingTableMem, FileMappingTableMemSize, TC("FileMappingTable"));
+		UnmapViewOfFile(m_logger, m_fileMappingTableMem, FileMappingTableMemSize, TC("FileMappingTable"), m_allowDiscardVirtualMemory);
 		CloseFileMapping(m_logger, m_fileMappingTableHandle, TC("FileMappingTable"));
 
-		UnmapViewOfFile(m_logger, m_directoryTableMem, DirTableMemSize, TC("DirectoryTable"));
+		UnmapViewOfFile(m_logger, m_directoryTableMem, DirTableMemSize, TC("DirectoryTable"), m_allowDiscardVirtualMemory);
 		CloseFileMapping(m_logger, m_directoryTableHandle, TC("DirectoryTable"));
 
 		//#if !UBA_DEBUG
@@ -2470,7 +2492,12 @@ namespace uba
 			m_debugLogger->Info(TC("Mapping freed 0x%llx (%s)"), u64(entry.mapping.mh), entry.name.c_str());
 			#endif
 
-			m_workManager.AddWork([mh = entry.mapping, this](const WorkContext&) { CloseFileMapping(m_logger, mh, TC("UsedFileMapping")); }, 1, TC("CloseFileMapping"));
+			m_workManager.AddWork([mh = entry.mapping, fileSize = entry.contentSize, this](const WorkContext&)
+				{
+					if (m_allowDiscardVirtualMemory)
+						DiscardFileMapping(m_logger, mh, fileSize, TC("UsedFileMapping"));
+					CloseFileMapping(m_logger, mh, TC("UsedFileMapping"));
+				}, 1, TC("CloseFileMapping"));
 			entry.handled = false;
 			entry.mapping = {};
 		}
@@ -3688,7 +3715,7 @@ namespace uba
 
 		u64 toCommit = newSize - m_directoryTableMemCommitted;
 		u8* address = m_directoryTableMem + m_directoryTableMemCommitted;
-		if (!MapViewCommit(address, toCommit))
+		if (!MapViewCommit(m_logger, address, toCommit, m_allowDiscardVirtualMemory))
 			m_logger.Error(TC("Failed to commit memory for directory table (Committed: %llu, ToCommit: %llu) (%s)"), m_directoryTableMemCommitted, toCommit, LastErrorToText().data);
 		m_directoryTableMemCommitted += toCommit;
 	}
diff --git a/Engine/Source/Programs/UnrealBuildAccelerator/Common/Private/UbaTrace.cpp b/Engine/Source/Programs/UnrealBuildAccelerator/Common/Private/UbaTrace.cpp
--- a/Engine/Source/Programs/UnrealBuildAccelerator/Common/Private/UbaTrace.cpp
+++ b/Engine/Source/Programs/UnrealBuildAccelerator/Common/Private/UbaTrace.cpp
@@ -318,7 +318,7 @@ namespace uba
 		if (m_memoryCommitted >= committedMemoryNeeded)
 			return true;
 
-		if (MapViewCommit(m_memoryBegin + m_memoryCommitted, committedMemoryNeeded - m_memoryCommitted))
+		if (MapViewCommit(m_logger, m_memoryBegin + m_memoryCommitted, committedMemoryNeeded - m_memoryCommitted))
 		{
 			m_memoryCommitted = committedMemoryNeeded;
 			return true;
diff --git a/Engine/Source/Programs/UnrealBuildAccelerator/Common/Public/UbaDependencyCrawler.h b/Engine/Source/Programs/UnrealBuildAccelerator/Common/Public/UbaDependencyCrawler.h
--- a/Engine/Source/Programs/UnrealBuildAccelerator/Common/Public/UbaDependencyCrawler.h
+++ b/Engine/Source/Programs/UnrealBuildAccelerator/Common/Public/UbaDependencyCrawler.h
@@ -3,6 +3,7 @@
 #pragma once
 
 #include "UbaApplicationRules.h"
+#include "UbaBottleneck.h"
 #include "UbaFileAccessor.h"
 #include "UbaPathUtils.h"
 #include "UbaProcessUtils.h"
@@ -26,7 +27,8 @@ namespace uba
 	class DependencyCrawler
 	{
 	public:
-		inline DependencyCrawler(Logger& logger, WorkManager& workManager);
+		inline DependencyCrawler(Logger& logger, WorkManager& workManager, u32 bottleneckMax = ~0u);
+		inline ~DependencyCrawler();
 
 		using FileExistsFunc = Function<bool(const StringView& fileName, u32& outAttr)>;
 		using FileFunc = Function<void(const StringView& file, bool isDirectory)>;
@@ -95,6 +97,8 @@ namespace uba
 		FileExistsFunc m_fileExistsFunc;
 		TraverseFilesFunc m_traverseFilesFunc;
 
+		Bottleneck* m_linkerDecompressBottleneck = nullptr;
+
 		struct CodeFile { TString path; bool hasPch = false; };
 		using CodeFiles = List<CodeFile>;
 
@@ -108,10 +112,17 @@ namespace uba
 	};
 
 
-	DependencyCrawler::DependencyCrawler(Logger& logger, WorkManager& workManager)
+	DependencyCrawler::DependencyCrawler(Logger& logger, WorkManager& workManager, u32 bottleneckMax)
 	:	m_logger(logger)
 	,	m_workManager(workManager)
 	{
+		if (bottleneckMax != ~0u)
+			m_linkerDecompressBottleneck = new Bottleneck(bottleneckMax);
+	}
+
+	DependencyCrawler::~DependencyCrawler()
+	{
+		delete m_linkerDecompressBottleneck;
 	}
 
 	void DependencyCrawler::Init(FileExistsFunc&& fileExistsFunc, TraverseFilesFunc&& traverseFilesFunc, bool useBloomFilter)
@@ -174,10 +185,23 @@ namespace uba
 					}
 					else if (instance.type == DependencyCrawlerType_MsvcLinker)
 					{
+						BottleneckTicket* ticket = nullptr;
+						if (m_linkerDecompressBottleneck)
+						{
+							ticket = new BottleneckTicket();
+							Timer timer;
+							m_linkerDecompressBottleneck->Enter(*ticket, timer);
+						}
 						// If we end up here it means that we have compressed obj files enabled and want to decompress them in parallel
-						m_workManager.AddWork([this, ref, codeFile = cf.path](const WorkContext& context)
+						m_workManager.AddWork([this, ref, codeFile = cf.path, ticket](const WorkContext& context)
 							{
 								ref.instance.createFileFunc(context.tracker, codeFile, {});
+								if (ticket)
+								{
+									m_linkerDecompressBottleneck->Leave(*ticket);
+									delete ticket;
+								}
+
 							}, 1, TC("CrawlForDecomp"), ColorWork);
 					}
 					else
diff --git a/Engine/Source/Programs/UnrealBuildAccelerator/Common/Public/UbaFileMapping.h b/Engine/Source/Programs/UnrealBuildAccelerator/Common/Public/UbaFileMapping.h
--- a/Engine/Source/Programs/UnrealBuildAccelerator/Common/Public/UbaFileMapping.h
+++ b/Engine/Source/Programs/UnrealBuildAccelerator/Common/Public/UbaFileMapping.h
@@ -30,9 +30,10 @@ namespace uba
 
 	FileMappingHandle CreateMemoryMappingW(Logger& logger, u32 flProtect, u64 maxSize, const tchar* name, const tchar* hint);
 	FileMappingHandle CreateFileMappingW(Logger& logger, FileHandle file, u32 flProtect, u64 maxSize, const tchar* hint);
-	u8* MapViewOfFile(Logger& logger, FileMappingHandle fileMappingObject, u32 dwDesiredAccess, u64 offset, u64 dwNumberOfBytesToMap);
-	bool MapViewCommit(void* address, u64 size);
-	bool UnmapViewOfFile(Logger& logger, const void* lpBaseAddress, u64 bytesToUnmap, const tchar* hint);
+	u8* MapViewOfFile(Logger& logger, FileMappingHandle fileMappingObject, u32 dwDesiredAccess, u64 offset, u64 dwNumberOfBytesToMap, bool allowDiscard = false);
+	bool MapViewCommit(Logger& logger, void* address, u64 size, bool allowDiscard = false);
+	bool UnmapViewOfFile(Logger& logger, const void* lpBaseAddress, u64 bytesToUnmap, const tchar* hint, bool allowDiscard = false);
+	void DiscardFileMapping(Logger& logger, FileMappingHandle h, u64 size, const tchar* hint);
 	bool CloseFileMapping(Logger& logger, FileMappingHandle h, const tchar* hint);
 	bool DuplicateFileMapping(Logger& logger, ProcHandle sourceProcessHandle, FileMappingHandle sourceHandle, ProcHandle targetProcessHandle, FileMappingHandle& targetHandle, u32 dwDesiredAccess, bool bInheritHandle, u32 dwOptions, const tchar* hint);
 	void MapMemoryCopy(void* dest, const void* source, u64 size);
diff --git a/Engine/Source/Programs/UnrealBuildAccelerator/Common/Public/UbaSession.h b/Engine/Source/Programs/UnrealBuildAccelerator/Common/Public/UbaSession.h
--- a/Engine/Source/Programs/UnrealBuildAccelerator/Common/Public/UbaSession.h
+++ b/Engine/Source/Programs/UnrealBuildAccelerator/Common/Public/UbaSession.h
@@ -245,6 +245,7 @@ namespace uba
 		bool m_allowLocalDetour;
 		bool m_extractObjFilesSymbols;
 		bool m_allowLinkDependencyCrawler;
+		bool m_allowDiscardVirtualMemory;
 		bool m_ownsTrace;
 
 		u64 m_keepOutputFileMemoryMapsThreshold;
diff --git a/Engine/Source/Programs/UnrealBuildAccelerator/Common/Public/UbaSessionCreateInfo.h b/Engine/Source/Programs/UnrealBuildAccelerator/Common/Public/UbaSessionCreateInfo.h
--- a/Engine/Source/Programs/UnrealBuildAccelerator/Common/Public/UbaSessionCreateInfo.h
+++ b/Engine/Source/Programs/UnrealBuildAccelerator/Common/Public/UbaSessionCreateInfo.h
@@ -43,7 +43,9 @@ namespace uba
 		bool treatTempDirAsEmpty = true;		// If this is true, directory table will always see temp as empty. Use sub directories of temp instead
 		bool useFakeVolumeSerial = true;		// Fake volume serials reduce the directory table size since serials end up being compressed down to one byte
 		bool keepTransientDataMapped = true;	// Will keep the transient data mapped in instead of doing map/unmap everytime we access it
+		bool allowDiscardVirtualMemory = false;	// Set to true to use DiscardVirtualMemory around map/unmap views of file mappings. This seems to speed up perf quite a lot
 		bool allowLinkDependencyCrawler = true;	// Enable crawler that reads ahead to find needed obj files for linker
+		int linkCrawlerMaxParallel = -2;		// If allowLinkDependencyCrawler is set, this controls number of decompress in parallel. If negative it is max cores - value
 		u64 deleteSessionsOlderThanSeconds = 12 * 60 * 60; // Delete session folders older than 12 hours by default . Set to 0 to not delete or 1 to delete all
 		u64 keepOutputFileMemoryMapsThreshold = 256 * 1024; // If allowMemoryMaps is true, output files will be kept in memory if smaller than this size
 		u32 traceReserveSizeMb = 128;			// Memory reserved for trace file (in mb)
diff --git a/Engine/Source/Programs/UnrealBuildAccelerator/Core/Public/UbaDetoursPayload.h b/Engine/Source/Programs/UnrealBuildAccelerator/Core/Public/UbaDetoursPayload.h
--- a/Engine/Source/Programs/UnrealBuildAccelerator/Core/Public/UbaDetoursPayload.h
+++ b/Engine/Source/Programs/UnrealBuildAccelerator/Core/Public/UbaDetoursPayload.h
@@ -26,6 +26,7 @@ namespace uba
 		bool isChild = false;
 		bool allowKeepFilesInMemory = IsWindows;
 		bool allowOutputFiles = IsWindows;
+		bool allowDiscardVirtualMemory = false;
 		bool suppressLogging = false;
 		bool readIntermediateFilesCompressed = false;
 		bool reportAllExceptions = false;
diff --git a/Engine/Source/Programs/UnrealBuildAccelerator/Detours/Private/UbaDetoursShared.cpp b/Engine/Source/Programs/UnrealBuildAccelerator/Detours/Private/UbaDetoursShared.cpp
--- a/Engine/Source/Programs/UnrealBuildAccelerator/Detours/Private/UbaDetoursShared.cpp
+++ b/Engine/Source/Programs/UnrealBuildAccelerator/Detours/Private/UbaDetoursShared.cpp
@@ -31,6 +31,7 @@ namespace uba
 	bool g_isChild;
 	bool g_allowKeepFilesInMemory = IsWindows;
 	bool g_allowOutputFiles = IsWindows;
+	bool g_allowDiscardVirtualMemory;
 	bool g_suppressLogging = false;
 	u32 g_vfsEntryCount;
 	u32 g_vfsMatchingLength;
diff --git a/Engine/Source/Programs/UnrealBuildAccelerator/Detours/Private/UbaDetoursShared.h b/Engine/Source/Programs/UnrealBuildAccelerator/Detours/Private/UbaDetoursShared.h
--- a/Engine/Source/Programs/UnrealBuildAccelerator/Detours/Private/UbaDetoursShared.h
+++ b/Engine/Source/Programs/UnrealBuildAccelerator/Detours/Private/UbaDetoursShared.h
@@ -101,6 +101,7 @@ namespace uba
 	extern bool g_isChild;
 	extern bool g_allowKeepFilesInMemory;
 	extern bool g_allowOutputFiles;
+	extern bool g_allowDiscardVirtualMemory;
 	extern bool g_suppressLogging;
 
 	#if PLATFORM_WINDOWS
diff --git a/Engine/Source/Programs/UnrealBuildAccelerator/Detours/Private/Windows/UbaDetoursFunctionsWin.cpp b/Engine/Source/Programs/UnrealBuildAccelerator/Detours/Private/Windows/UbaDetoursFunctionsWin.cpp
--- a/Engine/Source/Programs/UnrealBuildAccelerator/Detours/Private/Windows/UbaDetoursFunctionsWin.cpp
+++ b/Engine/Source/Programs/UnrealBuildAccelerator/Detours/Private/Windows/UbaDetoursFunctionsWin.cpp
@@ -460,6 +460,8 @@ void MemoryFile::EnsureCommitted(const DetouredHandle& handle, u64 size)
 		u64 toCommit = Min(reserveSize, AlignUp(size - committedSize, g_pageSize));
 		if (!VirtualAlloc(baseAddress + committedSize, toCommit, MEM_COMMIT, PAGE_READWRITE))
 			FatalError(1347, L"Failed to ensure virtual memory for %ls trying to commit %llu at %llx. MappedSize: %llu, CommittedSize: %llu RequestedSize: %llu. (%u)", HandleToName(handle), toCommit, baseAddress + committedSize, mappedSize, committedSize, size, GetLastError());
+		if (g_allowDiscardVirtualMemory)
+			DiscardVirtualMemory(baseAddress + committedSize, toCommit);
 		committedSize += toCommit;
 	}
 }
@@ -674,11 +676,13 @@ void SendExitMessage(DWORD exitCode, u64 startTime)
 	if (g_consoleStringIndex)
 		Shared_WriteConsole(L"\n", 1, 0);
 
+	#if UBA_USE_MIMALLOC
 	#if TRACK_UCRT_ALLOC_ENABLED
 	StringBuffer<256> str;
 	str.Appendf(L"Alloc: %llu Realloc: %llu Free: %llu\n", g_allocCount.load(), g_reallocCount.load(), g_freeCount.load());
 	Shared_WriteConsole(str.data, str.count, 0);
 	#endif
+	#endif
 
 	if (g_trackInputsMem)
 		SendInput();
@@ -909,6 +913,7 @@ void PreInit(const DetoursPayload& payload)
 	g_isChild = payload.isChild;
 	g_allowKeepFilesInMemory = payload.allowKeepFilesInMemory;
 	g_allowOutputFiles = g_allowKeepFilesInMemory && payload.allowOutputFiles;
+	g_allowDiscardVirtualMemory = payload.allowDiscardVirtualMemory;
 	g_suppressLogging = payload.suppressLogging;
 	g_isDetachedProcess = g_rules->AllowDetach();
 	g_isRunningWine = payload.isRunningWine;
