[UBA]
* Added Bottleneck.TryEnter which returns true if it manage to fit or false if not
* Added support for bottleneck in work manager. Uses TryEnter to try to start a job. If it fails it picks the next job
* Changed so dependency crawler use work manager bottleneck since the previous approach should in theory be able to cause deadlocks

(Note to self, it is annoying if crawler ends after process it crawls for.. we would end up with mapped files that will not be released and just take up memory.. maybe add a cancellation token of some sort)

diff --git a/Engine/Source/Programs/UnrealBuildAccelerator/Common/Private/UbaNetworkServer.cpp b/Engine/Source/Programs/UnrealBuildAccelerator/Common/Private/UbaNetworkServer.cpp
--- a/Engine/Source/Programs/UnrealBuildAccelerator/Common/Private/UbaNetworkServer.cpp
+++ b/Engine/Source/Programs/UnrealBuildAccelerator/Common/Private/UbaNetworkServer.cpp
@@ -1,6 +1,7 @@
 // Copyright Epic Games, Inc. All Rights Reserved.
 
 #include "UbaNetworkServer.h"
+#include "UbaBottleneck.h"
 #include "UbaConfig.h"
 #include "UbaCrypto.h"
 #include "UbaBinaryReaderWriter.h"
@@ -725,24 +726,39 @@ namespace uba
 
 	void NetworkServer::Worker::DoAdditionalWorkAndSignalAvailable(NetworkServer& server)
 	{
+		auto& additionalWork = server.m_additionalWork;
+		auto& additionalWorkLock = server.m_additionalWorkLock;
+		BottleneckTicket ticket;
+
 		while (true)
 		{
 			while (true)
 			{
 				AdditionalWork work;
-				SCOPED_FUTEX(server.m_additionalWorkLock, lock);
-				if (server.m_additionalWork.empty())
+				SCOPED_FUTEX(additionalWorkLock, lock);
+				for (auto i=additionalWork.begin(), e=additionalWork.end(); i!=e; ++i)
+				{
+					if (i->bottleneck && !i->bottleneck->TryEnter())
+						continue;
+					work = std::move(*i);
+					additionalWork.erase(i);
 					break;
-				work = std::move(server.m_additionalWork.front());
-				server.m_additionalWork.pop_front();
+				}
 				lock.Leave();
 
+				if (!work.func)
+					break;
+
 				#if UBA_TRACK_WORK
 				TrackWorkScope tws(server, work.desc, ColorWork);
 				#else
 				TrackWorkScope tws;
 				#endif
+
 				work.func({tws});
+
+				if (work.bottleneck)
+					work.bottleneck->Leave(ticket);
 			}
 
 			// Both locks needs to be taken to verify if additional work
@@ -750,10 +766,10 @@ namespace uba
 			// a race where AddWork would not see this thread in the
 			// available list after adding some work.
 			SCOPED_FUTEX(server.m_availableWorkersLock, lock1);
-			SCOPED_FUTEX_READ(server.m_additionalWorkLock, lock2);
+			SCOPED_FUTEX_READ(additionalWorkLock, lock2);
 			// Verify there is not additional work while we hold both lock
 			// and only add ourself as available if no additional work is present.
-			if (!server.m_additionalWork.empty())
+			if (!additionalWork.empty())
 				continue;
 			server.PushWorkerNoLock(this);
 			break;
@@ -1035,7 +1051,7 @@ namespace uba
 		}
 	}
 
-	void NetworkServer::AddWork(const WorkFunction& work, u32 count, const tchar* desc, const Color& color, bool highPriority)
+	void NetworkServer::AddWork(const WorkFunction& work, u32 count, const tchar* desc, const Color& color, bool highPriority, Bottleneck* bottleneck)
 	{
 		UBA_ASSERT(*desc);
 		SCOPED_FUTEX(m_additionalWorkLock, lock);
@@ -1043,13 +1059,13 @@ namespace uba
 		{
 			if (highPriority)
 			{
-				m_additionalWork.push_front({ work });
+				m_additionalWork.push_front({ work, bottleneck });
 				if (m_workTracker)
 					m_additionalWork.front().desc = desc;
 			}
 			else
 			{
-				m_additionalWork.push_back({ work });
+				m_additionalWork.push_back({ work, bottleneck });
 				if (m_workTracker)
 					m_additionalWork.back().desc = desc;
 			}
diff --git a/Engine/Source/Programs/UnrealBuildAccelerator/Common/Private/UbaWorkManager.cpp b/Engine/Source/Programs/UnrealBuildAccelerator/Common/Private/UbaWorkManager.cpp
--- a/Engine/Source/Programs/UnrealBuildAccelerator/Common/Private/UbaWorkManager.cpp
+++ b/Engine/Source/Programs/UnrealBuildAccelerator/Common/Private/UbaWorkManager.cpp
@@ -1,6 +1,7 @@
 // Copyright Epic Games, Inc. All Rights Reserved.
 
 #include "UbaWorkManager.h"
+#include "UbaBottleneck.h"
 #include "UbaEvent.h"
 #include "UbaPlatform.h"
 #include "UbaThread.h"
@@ -52,6 +53,8 @@ namespace uba
 
 		void ThreadWorker(WorkManagerImpl& manager)
 		{
+			BottleneckTicket ticket;
+
 			while (true)
 			{
 				if (!m_workAvailable.IsSet())
@@ -63,12 +66,19 @@ namespace uba
 				{
 					while (true)
 					{
+						Work work;
 						SCOPED_FUTEX(manager.m_workLock, lock);
-						if (manager.m_work.empty())
+						for (auto i=manager.m_work.begin(), e=manager.m_work.end(); i!=e; ++i)
+						{
+							if (i->bottleneck && !i->bottleneck->TryEnter())
+								continue;
+							work = std::move(*i);
+							manager.m_work.erase(i);
 							break;
-						Work work = manager.m_work.front();
-						manager.m_work.pop_front();
+						}
 						lock.Leave();
+						if (!work.func)
+							break;
 
 						#if UBA_TRACK_WORK
 						TrackWorkScope tws(manager, work.desc);
@@ -76,6 +86,9 @@ namespace uba
 						TrackWorkScope tws;
 						#endif
 						work.func({tws});
+
+						if (work.bottleneck)
+							work.bottleneck->Leave(ticket);
 					}
 
 					SCOPED_FUTEX(manager.m_availableWorkersLock, lock1);
@@ -113,7 +126,7 @@ namespace uba
 	}
 
 
-	void WorkManagerImpl::AddWork(const WorkFunction& work, u32 count, const tchar* desc, const Color& color, bool highPriority)
+	void WorkManagerImpl::AddWork(const WorkFunction& work, u32 count, const tchar* desc, const Color& color, bool highPriority, Bottleneck* bottleneck)
 	{
 		UBA_ASSERT(*desc);
 		UBA_ASSERT(!m_workers.empty());
@@ -123,13 +136,13 @@ namespace uba
 		{
 			if (highPriority)
 			{
-				m_work.push_front({ work });
+				m_work.push_front({ work, bottleneck });
 				if (trackWork)
 					m_work.front().desc = desc;
 			}
 			else
 			{
-				m_work.push_back({ work });
+				m_work.push_back({ work, bottleneck });
 				if (trackWork)
 					m_work.back().desc = desc;
 			}
diff --git a/Engine/Source/Programs/UnrealBuildAccelerator/Common/Public/UbaBottleneck.h b/Engine/Source/Programs/UnrealBuildAccelerator/Common/Public/UbaBottleneck.h
--- a/Engine/Source/Programs/UnrealBuildAccelerator/Common/Public/UbaBottleneck.h
+++ b/Engine/Source/Programs/UnrealBuildAccelerator/Common/Public/UbaBottleneck.h
@@ -21,13 +21,12 @@ namespace uba
 
 		void Enter(BottleneckTicket& ticket, Timer& timer)
 		{
+			if (TryEnter())
+				return;
+
 			SCOPED_FUTEX(m_lock, lock);
-			if (m_activeCount < m_maxCount)
-			{
-				UBA_ASSERT(!m_first);
-				++m_activeCount;
+			if (TryEnter())
 				return;
-			}
 
 			if (!m_first)
 				m_first = &ticket;
@@ -44,14 +43,23 @@ namespace uba
 			ticket.ev.IsSet();
 		}
 
+		bool TryEnter()
+		{
+			u32 ac = m_activeCount.load(std::memory_order_relaxed);
+			while (ac < m_maxCount)
+				if (m_activeCount.compare_exchange_weak(ac, ac + 1, std::memory_order_acquire, std::memory_order_relaxed))
+					return true;
+			return false;
+		}
+
 		void Leave(BottleneckTicket& ticket)
 		{
 			SCOPED_FUTEX(m_lock, lock);
 			BottleneckTicket* first = m_first;
 			if (!first)
 			{
-				UBA_ASSERT(m_activeCount > 0);
-				--m_activeCount;
+				UBA_ASSERT(m_activeCount.load(std::memory_order_relaxed) > 0);
+				m_activeCount.fetch_sub(1, std::memory_order_release);
 				return;
 			}
 			m_first = first->next;
@@ -64,7 +72,7 @@ namespace uba
 		Futex m_lock;
 		BottleneckTicket* m_first = nullptr;
 		BottleneckTicket* m_last = nullptr;
-		u32 m_activeCount = 0;
+		Atomic<u32> m_activeCount;
 		u32 m_maxCount;
 	};
 
diff --git a/Engine/Source/Programs/UnrealBuildAccelerator/Common/Public/UbaDependencyCrawler.h b/Engine/Source/Programs/UnrealBuildAccelerator/Common/Public/UbaDependencyCrawler.h
--- a/Engine/Source/Programs/UnrealBuildAccelerator/Common/Public/UbaDependencyCrawler.h
+++ b/Engine/Source/Programs/UnrealBuildAccelerator/Common/Public/UbaDependencyCrawler.h
@@ -185,24 +185,11 @@ namespace uba
 					}
 					else if (instance.type == DependencyCrawlerType_MsvcLinker)
 					{
-						BottleneckTicket* ticket = nullptr;
-						if (m_linkerDecompressBottleneck)
-						{
-							ticket = new BottleneckTicket();
-							Timer timer;
-							m_linkerDecompressBottleneck->Enter(*ticket, timer);
-						}
 						// If we end up here it means that we have compressed obj files enabled and want to decompress them in parallel
-						m_workManager.AddWork([this, ref, codeFile = cf.path, ticket](const WorkContext& context)
+						m_workManager.AddWork([this, ref, codeFile = cf.path](const WorkContext& context)
 							{
 								ref.instance.createFileFunc(context.tracker, codeFile, {});
-								if (ticket)
-								{
-									m_linkerDecompressBottleneck->Leave(*ticket);
-									delete ticket;
-								}
-
-							}, 1, TC("CrawlForDecomp"), ColorWork);
+							}, 1, TC("CrawlForDecomp"), ColorWork, false, m_linkerDecompressBottleneck);
 					}
 					else
 					{
diff --git a/Engine/Source/Programs/UnrealBuildAccelerator/Common/Public/UbaNetworkServer.h b/Engine/Source/Programs/UnrealBuildAccelerator/Common/Public/UbaNetworkServer.h
--- a/Engine/Source/Programs/UnrealBuildAccelerator/Common/Public/UbaNetworkServer.h
+++ b/Engine/Source/Programs/UnrealBuildAccelerator/Common/Public/UbaNetworkServer.h
@@ -17,6 +17,7 @@ namespace uba
 	class Trace;
 	struct BinaryReader;
 	struct BinaryWriter;
+	struct Bottleneck;
 
 	struct ConnectionInfo
 	{
@@ -89,7 +90,7 @@ namespace uba
 		//void RegisterOnConnection(u8 id, const OnConnectionFunction& func);
 		//void UnregisterOnConnection(u8 id);
 
-		virtual void AddWork(const WorkFunction& work, u32 count, const tchar* desc, const Color& color = ColorWork, bool highPriority = false) override final;
+		virtual void AddWork(const WorkFunction& work, u32 count, const tchar* desc, const Color& color = ColorWork, bool highPriority = false, Bottleneck* bottleneck = nullptr) override final;
 		virtual void DoWork(u32 count = 1) override final;
 		virtual u32 GetWorkerCount() override final;
 		virtual bool FlushWork(u32 timeoutMs = 0) override final;
@@ -158,7 +159,7 @@ namespace uba
 		u32 m_maxWorkerCount = 0;
 
 		Futex m_additionalWorkLock;
-		struct AdditionalWork { WorkFunction func; TString desc; };
+		struct AdditionalWork { WorkFunction func; Bottleneck* bottleneck; TString desc; };
 		List<AdditionalWork> m_additionalWork;
 
 		Futex m_availableWorkersLock;
diff --git a/Engine/Source/Programs/UnrealBuildAccelerator/Common/Public/UbaWorkManager.h b/Engine/Source/Programs/UnrealBuildAccelerator/Common/Public/UbaWorkManager.h
--- a/Engine/Source/Programs/UnrealBuildAccelerator/Common/Public/UbaWorkManager.h
+++ b/Engine/Source/Programs/UnrealBuildAccelerator/Common/Public/UbaWorkManager.h
@@ -10,6 +10,7 @@
 
 namespace uba
 {
+	struct Bottleneck;
 	struct StringView;
 	struct TrackWorkScope;
 
@@ -35,7 +36,7 @@ namespace uba
 	{
 	public:
 		using WorkFunction = Function<void(const WorkContext& context)>;
-		virtual void AddWork(const WorkFunction& work, u32 count, const tchar* desc, const Color& color = ColorWork, bool highPriority = false) = 0;
+		virtual void AddWork(const WorkFunction& work, u32 count, const tchar* desc, const Color& color = ColorWork, bool highPriority = false, Bottleneck* bottleneck = nullptr) = 0;
 		virtual u32 GetWorkerCount() = 0;
 		virtual void DoWork(u32 count = 1) = 0;
 		virtual bool FlushWork(u32 timeoutMs = 0) = 0;
@@ -60,7 +61,7 @@ namespace uba
 	public:
 		WorkManagerImpl(u32 workerCount, const tchar* workerDesc = TC("UbaWrk"));
 		virtual ~WorkManagerImpl();
-		virtual void AddWork(const WorkFunction& work, u32 count, const tchar* desc, const Color& color = ColorWork, bool highPriority = false) override final;
+		virtual void AddWork(const WorkFunction& work, u32 count, const tchar* desc, const Color& color = ColorWork, bool highPriority = false, Bottleneck* bottleneck = nullptr) override final;
 		virtual u32 GetWorkerCount() override final;
 		virtual void DoWork(u32 count = 1) override final;
 		virtual bool FlushWork(u32 timeoutMs = 0) override final;
@@ -72,7 +73,7 @@ namespace uba
 		Worker* PopWorkerNoLock();
 
 		Vector<Worker*> m_workers;
-		struct Work { WorkFunction func; TString desc; };
+		struct Work { WorkFunction func; Bottleneck* bottleneck; TString desc; };
 		Futex m_workLock;
 		List<Work> m_work;
 		Atomic<u32> m_activeWorkerCount;
