[UBA]
* Removed discarding of virtual memory for DirectoryTable since entire memory is not mapped which reports an error

diff --git a/Engine/Source/Programs/UnrealBuildAccelerator/Common/Private/UbaSession.cpp b/Engine/Source/Programs/UnrealBuildAccelerator/Common/Private/UbaSession.cpp
--- a/Engine/Source/Programs/UnrealBuildAccelerator/Common/Private/UbaSession.cpp
+++ b/Engine/Source/Programs/UnrealBuildAccelerator/Common/Private/UbaSession.cpp
@@ -1838,7 +1838,7 @@ namespace uba
 		UnmapViewOfFile(m_logger, m_fileMappingTableMem, FileMappingTableMemSize, TC("FileMappingTable"), m_allowDiscardVirtualMemory);
 		CloseFileMapping(m_logger, m_fileMappingTableHandle, TC("FileMappingTable"));
 
-		UnmapViewOfFile(m_logger, m_directoryTableMem, DirTableMemSize, TC("DirectoryTable"), m_allowDiscardVirtualMemory);
+		UnmapViewOfFile(m_logger, m_directoryTableMem, DirTableMemSize, TC("DirectoryTable"), false); // Can't discard since only parts might be mapped
 		CloseFileMapping(m_logger, m_directoryTableHandle, TC("DirectoryTable"));
 
 		//#if !UBA_DEBUG
