[UBA]
* Improved link dependency crawler to not load obj files after they have been used
* Changed so processes are using workmanager for message threads in order to reusing them

diff --git a/Engine/Source/Programs/UnrealBuildAccelerator/Common/Private/UbaProcess.cpp b/Engine/Source/Programs/UnrealBuildAccelerator/Common/Private/UbaProcess.cpp
--- a/Engine/Source/Programs/UnrealBuildAccelerator/Common/Private/UbaProcess.cpp
+++ b/Engine/Source/Programs/UnrealBuildAccelerator/Common/Private/UbaProcess.cpp
@@ -167,7 +167,8 @@ namespace uba
 				m_cancelEvent.Set();
 		}
 
-		m_messageThread.Wait();
+		if (m_messageThread)
+			m_messageThreadDone.IsSet();
 
 		if (m_comMemory.memory)
 		{
@@ -247,12 +248,16 @@ namespace uba
 		m_session.ProcessAdded(*this, 0);
 
 		if (async)
-			m_messageThread.Start([this, environment]()
+		{
+			m_messageThreadDone.Create(true);
+			m_session.m_processMessageThreads.AddWork([this, environment](const WorkContext& context)
 				{
+					m_messageThread = context.thread;
 					ThreadRun(environment);
 					ThreadExit();
-					return 0;
-				}, startInfo.description);
+					m_messageThreadDone.Set();
+				}, 1, TC("UbaProcMsg"));
+		}
 		else
 		{
 			// This is needed to handle cltr-c. Otherwise this thread might exit before the detoured process which in turn might do things after ctrl-c which can cause deadlocks in detoured process
@@ -443,7 +448,7 @@ namespace uba
 
 	bool ProcessImpl::WaitForExit(u32 millisecondsTimeout)
 	{
-		return m_messageThread.Wait(millisecondsTimeout);
+		return m_messageThreadDone.IsSet(millisecondsTimeout);
 	}
 
 	u64 ProcessImpl::GetTotalProcessorTime() const
@@ -1876,10 +1881,10 @@ namespace uba
 		}
 
 		bool affinitySet = false;
-		if (m_messageThread.IsInside())
+		if (m_messageThread->IsInside())
 		{
 			GroupAffinity aff;
-			if (m_messageThread.GetGroupAffinity(aff))
+			if (m_messageThread->GetGroupAffinity(aff))
 				affinitySet = SetThreadGroupAffinity(m_nativeThreadHandle, aff);
 		}
 		
@@ -2377,6 +2382,6 @@ namespace uba
 
 	bool ProcessImpl::IsCalledFromThis() const
 	{
-		return m_messageThread.IsInside();
+		return m_messageThread->IsInside();
 	}
 }
diff --git a/Engine/Source/Programs/UnrealBuildAccelerator/Common/Private/UbaSession.cpp b/Engine/Source/Programs/UnrealBuildAccelerator/Common/Private/UbaSession.cpp
--- a/Engine/Source/Programs/UnrealBuildAccelerator/Common/Private/UbaSession.cpp
+++ b/Engine/Source/Programs/UnrealBuildAccelerator/Common/Private/UbaSession.cpp
@@ -858,6 +858,9 @@ namespace uba
 
 		SCOPED_FUTEX(entry.lock, entryLock);
 
+		if (requestingProcess && (requestingProcess->m_gotExitMessage || requestingProcess->m_cancelled)) // This is a crawler call that we can ignore
+			return false;
+
 		if (entry.handled)
 		{
 			if (!entry.success)
@@ -1549,6 +1552,7 @@ namespace uba
 	,	m_directoryTable(m_directoryTableMemory)
 	,	m_fileMappingBuffer(m_logger, &workManager)
 	,	m_processCommunicationAllocator(m_logger, TC("CommunicationAllocator"))
+	,	m_processMessageThreads(*new WorkManagerImpl(~0u, TC("UbaProcMsg")))
 	,	m_trace(info.trace ? *info.trace : *new Trace(info.logWriter))
 	,	m_writeFilesBottleneck(info.writeFilesBottleneck)
 	,	m_writeFilesFileMapMax(u64(info.writeFilesFileMapMaxMb)*1024*1024)
@@ -1828,6 +1832,8 @@ namespace uba
 		CancelAllProcessesAndWait();
 		FlushDeadProcesses();
 
+		delete &m_processMessageThreads;
+
 		for (auto& kv : m_virtualSourceFiles)
 			CloseFileMapping(m_logger, kv.second.mappingHandle, TC("VirtualFile"));
 
@@ -3937,17 +3943,28 @@ namespace uba
 		if (!at)
 			return;
 
-		auto CreateFileFunc = [this, ph = ProcessHandle(&process), rules = startInfo.rules](TrackWorkScope& tracker, const StringView& fileName, const DependencyCrawler::AccessFileFunc& func)
+		bool isLinkCrawler = crawlerType == DependencyCrawlerType_ClangLinker || crawlerType == DependencyCrawlerType_MsvcLinker;
+
+		auto CreateFileFunc = [this, ph = ProcessHandle(&process), rules = startInfo.rules, isLinkCrawler](TrackWorkScope& tracker, const StringView& fileName, const DependencyCrawler::AccessFileFunc& func)
 			{
 				auto& process = *static_cast<ProcessImpl*>(ph.m_process);
-				if (process.IsCancelled())
+				if (process.m_gotExitMessage || process.IsCancelled() || process.HasExited())
 					return false;
 
+				StringKey fileNameKey = ToStringKey(fileName);
+
+				if (isLinkCrawler)
+				{
+					SCOPED_FUTEX_READ(process.m_usedFileMappingsLock, lock)
+					if (process.m_usedFileMappings.find(fileNameKey) != process.m_usedFileMappings.end())
+						return false;
+				}
+
 				CreateFileResponse out;
 				{
 					tracker.AddHint(fileName);
 
-					if (!CreateFileForRead(out, tracker, fileName, ToStringKey(fileName), process, *rules))
+					if (!CreateFileForRead(out, tracker, fileName, fileNameKey, process, *rules))
 						return false;
 				}
 
diff --git a/Engine/Source/Programs/UnrealBuildAccelerator/Common/Private/UbaWorkManager.cpp b/Engine/Source/Programs/UnrealBuildAccelerator/Common/Private/UbaWorkManager.cpp
--- a/Engine/Source/Programs/UnrealBuildAccelerator/Common/Private/UbaWorkManager.cpp
+++ b/Engine/Source/Programs/UnrealBuildAccelerator/Common/Private/UbaWorkManager.cpp
@@ -37,7 +37,6 @@ namespace uba
 		Worker(WorkManagerImpl& manager, const tchar* workerDesc) : m_workAvailable(false)
 		{
 			m_loop = true;
-			manager.PushWorker(this);
 			m_thread.Start([&]() { ThreadWorker(manager); return 0; }, workerDesc);
 		}
 		~Worker()
@@ -85,7 +84,7 @@ namespace uba
 						#else
 						TrackWorkScope tws;
 						#endif
-						work.func({tws});
+						work.func({tws, &m_thread});
 
 						if (work.bottleneck)
 							work.bottleneck->Leave(ticket);
@@ -109,27 +108,29 @@ namespace uba
 		Thread m_thread;
 	};
 
-	WorkManagerImpl::WorkManagerImpl(u32 workerCount, const tchar* workerDesc)
+	WorkManagerImpl::WorkManagerImpl(u32 maxWorkerCount, const tchar* workerDesc, u32 precreateWorkerCount)
+	:	m_maxWorkerCount(maxWorkerCount)
+	,	m_workerDesc(workerDesc)
 	{
-		m_workers.resize(workerCount);
-		m_activeWorkerCount = workerCount;
-		for (u32 i = 0; i != workerCount; ++i)
-			m_workers[i] = new Worker(*this, workerDesc);
+		for (u32 i = 0; i != precreateWorkerCount; ++i)
+		{
+			auto& worker = m_workers.emplace_back(*this, m_workerDesc.c_str());
+			PushWorkerNoLock(&worker);
+		}
+		m_activeWorkerCount = 0;
 	}
 
 	WorkManagerImpl::~WorkManagerImpl()
 	{
-		for (auto worker : m_workers)
-			worker->Stop();
-		for (auto worker : m_workers)
-			delete worker;
+		for (auto& worker : m_workers)
+			worker.Stop();
+		m_workers.clear();
 	}
 
 
 	void WorkManagerImpl::AddWork(const WorkFunction& work, u32 count, const tchar* desc, const Color& color, bool highPriority, Bottleneck* bottleneck)
 	{
 		UBA_ASSERT(*desc);
-		UBA_ASSERT(!m_workers.empty());
 		SCOPED_FUTEX(m_workLock, lock);
 		bool trackWork = m_workTracker.load();
 		for (u32 i = 0; i != count; ++i)
@@ -161,13 +162,7 @@ namespace uba
 
 	u32 WorkManagerImpl::GetWorkerCount()
 	{
-		return u32(m_workers.size());
-	}
-
-	void WorkManagerImpl::PushWorker(Worker* worker)
-	{
-		SCOPED_FUTEX(m_availableWorkersLock, lock);
-		PushWorkerNoLock(worker);
+		return m_maxWorkerCount;
 	}
 
 	void WorkManagerImpl::PushWorkerNoLock(Worker* worker)
@@ -230,7 +225,13 @@ namespace uba
 	{
 		Worker* worker = m_firstAvailableWorker;
 		if (!worker)
-			return nullptr;
+		{
+			if (m_workers.size() == m_maxWorkerCount)
+				return nullptr;
+			Worker& newWorker = m_workers.emplace_back(*this, m_workerDesc.c_str());
+			++m_activeWorkerCount;
+			return &newWorker;
+		}
 		m_firstAvailableWorker = worker->m_nextWorker;
 		if (m_firstAvailableWorker)
 			m_firstAvailableWorker->m_prevWorker = nullptr;
diff --git a/Engine/Source/Programs/UnrealBuildAccelerator/Common/Public/UbaProcess.h b/Engine/Source/Programs/UnrealBuildAccelerator/Common/Public/UbaProcess.h
--- a/Engine/Source/Programs/UnrealBuildAccelerator/Common/Public/UbaProcess.h
+++ b/Engine/Source/Programs/UnrealBuildAccelerator/Common/Public/UbaProcess.h
@@ -183,7 +183,8 @@ namespace uba
 		Futex m_usedFileMappingsLock;
 		UnorderedMap<StringKey, UsedFileMapping> m_usedFileMappings;
 
-		Thread m_messageThread;
+		Event m_messageThreadDone;
+		Thread* m_messageThread = nullptr;
 
 		bool m_extractExports = false;
 
diff --git a/Engine/Source/Programs/UnrealBuildAccelerator/Common/Public/UbaSession.h b/Engine/Source/Programs/UnrealBuildAccelerator/Common/Public/UbaSession.h
--- a/Engine/Source/Programs/UnrealBuildAccelerator/Common/Public/UbaSession.h
+++ b/Engine/Source/Programs/UnrealBuildAccelerator/Common/Public/UbaSession.h
@@ -306,6 +306,7 @@ namespace uba
 		Futex m_processesLock;
 		UnorderedMap<u32, ProcessHandle> m_processes;
 		Vector<ProcessHandle> m_deadProcesses;
+		WorkManagerImpl& m_processMessageThreads;
 		UnorderedMap<TString, Timer> m_applicationStats;
 
 		Futex m_outputFilesLock;
diff --git a/Engine/Source/Programs/UnrealBuildAccelerator/Common/Public/UbaWorkManager.h b/Engine/Source/Programs/UnrealBuildAccelerator/Common/Public/UbaWorkManager.h
--- a/Engine/Source/Programs/UnrealBuildAccelerator/Common/Public/UbaWorkManager.h
+++ b/Engine/Source/Programs/UnrealBuildAccelerator/Common/Public/UbaWorkManager.h
@@ -10,6 +10,7 @@
 
 namespace uba
 {
+	class Thread;
 	struct Bottleneck;
 	struct StringView;
 	struct TrackWorkScope;
@@ -30,6 +31,7 @@ namespace uba
 	struct WorkContext
 	{
 		TrackWorkScope& tracker;
+		Thread* thread = nullptr;
 	};
 
 	class WorkManager : public WorkTracker
@@ -59,7 +61,7 @@ namespace uba
 	class WorkManagerImpl : public WorkManager
 	{
 	public:
-		WorkManagerImpl(u32 workerCount, const tchar* workerDesc = TC("UbaWrk"));
+		WorkManagerImpl(u32 maxWorkerCount, const tchar* workerDesc = TC("UbaWrk"), u32 precreateWorkerCount = 1);
 		virtual ~WorkManagerImpl();
 		virtual void AddWork(const WorkFunction& work, u32 count, const tchar* desc, const Color& color = ColorWork, bool highPriority = false, Bottleneck* bottleneck = nullptr) override final;
 		virtual u32 GetWorkerCount() override final;
@@ -68,16 +70,16 @@ namespace uba
 
 	private:
 		struct Worker;
-		void PushWorker(Worker* worker);
 		void PushWorkerNoLock(Worker* worker);
 		Worker* PopWorkerNoLock();
 
-		Vector<Worker*> m_workers;
+		u32 m_maxWorkerCount;
+		TString m_workerDesc;
+		List<Worker> m_workers;
 		struct Work { WorkFunction func; Bottleneck* bottleneck; TString desc; };
 		Futex m_workLock;
 		List<Work> m_work;
 		Atomic<u32> m_activeWorkerCount;
-		Atomic<u32> m_workCounter;
 
 		Futex m_availableWorkersLock;
 		Worker* m_firstAvailableWorker = nullptr;
diff --git a/Engine/Source/Programs/UnrealBuildAccelerator/Core/Private/UbaApplicationRules.cpp b/Engine/Source/Programs/UnrealBuildAccelerator/Core/Private/UbaApplicationRules.cpp
--- a/Engine/Source/Programs/UnrealBuildAccelerator/Core/Private/UbaApplicationRules.cpp
+++ b/Engine/Source/Programs/UnrealBuildAccelerator/Core/Private/UbaApplicationRules.cpp
@@ -228,6 +228,11 @@ namespace uba
 			};
 			return preloads;
 		}
+
+		virtual bool CloseFileMappingAfterUse(StringView fileName) const
+		{
+			return true;
+		}
 	};
 
 	class ApplicationRulesLldLinkExe : public ApplicationRulesVcLink
@@ -278,6 +283,11 @@ namespace uba
 				|| fileName.Contains(TCV(".dll.tmp"))
 				|| fileName.Contains(TCV(".pdb.tmp"));
 		}
+
+		virtual bool CloseFileMappingAfterUse(StringView fileName) const
+		{
+			return true;
+		}
 	};
 
 	// ==== Clang tool chain ====
