[UBA]
* Bumped uba extra link file memoryreserve size from 16mb to 64mb

diff --git a/Engine/Source/Programs/UnrealBuildAccelerator/Common/Private/UbaObjectFile.cpp b/Engine/Source/Programs/UnrealBuildAccelerator/Common/Private/UbaObjectFile.cpp
--- a/Engine/Source/Programs/UnrealBuildAccelerator/Common/Private/UbaObjectFile.cpp
+++ b/Engine/Source/Programs/UnrealBuildAccelerator/Common/Private/UbaObjectFile.cpp
@@ -251,7 +251,7 @@ namespace uba
 		ObjectFileCoff objectFileCoff;
 		ObjectFileElf objectFileElf;
 		
-		MemoryBlock memoryBlock(16*1024*1024);
+		MemoryBlock memoryBlock(64*1024*1024);
 
 		bool res;
 		if (extraObjFilename.EndsWith(TCV(".obj")))
