[UBA]
* Fixed so log warnings are not appearing if crawling is stopped because process has exited

diff --git a/Engine/Source/Programs/UnrealBuildAccelerator/Common/Private/UbaSession.cpp b/Engine/Source/Programs/UnrealBuildAccelerator/Common/Private/UbaSession.cpp
--- a/Engine/Source/Programs/UnrealBuildAccelerator/Common/Private/UbaSession.cpp
+++ b/Engine/Source/Programs/UnrealBuildAccelerator/Common/Private/UbaSession.cpp
@@ -3945,11 +3945,14 @@ namespace uba
 
 		bool isLinkCrawler = crawlerType == DependencyCrawlerType_ClangLinker || crawlerType == DependencyCrawlerType_MsvcLinker;
 
-		auto CreateFileFunc = [this, ph = ProcessHandle(&process), rules = startInfo.rules, isLinkCrawler](TrackWorkScope& tracker, const StringView& fileName, const DependencyCrawler::AccessFileFunc& func)
+		auto CreateFileFunc = [this, ph = ProcessHandle(&process), rules = startInfo.rules, isLinkCrawler](TrackWorkScope& tracker, const StringView& fileName, const DependencyCrawler::AccessFileFunc& func, bool& stopCrawling)
 			{
 				auto& process = *static_cast<ProcessImpl*>(ph.m_process);
 				if (process.m_gotExitMessage || process.IsCancelled() || process.HasExited())
+				{
+					stopCrawling = true;
 					return false;
+				}
 
 				StringKey fileNameKey = ToStringKey(fileName);
 
@@ -3965,7 +3968,10 @@ namespace uba
 					tracker.AddHint(fileName);
 
 					if (!CreateFileForRead(out, tracker, fileName, fileNameKey, process, *rules))
+					{
+						stopCrawling = process.m_gotExitMessage || process.IsCancelled() || process.HasExited();
 						return false;
+					}
 				}
 
 				if (!func)
diff --git a/Engine/Source/Programs/UnrealBuildAccelerator/Common/Public/UbaDependencyCrawler.h b/Engine/Source/Programs/UnrealBuildAccelerator/Common/Public/UbaDependencyCrawler.h
--- a/Engine/Source/Programs/UnrealBuildAccelerator/Common/Public/UbaDependencyCrawler.h
+++ b/Engine/Source/Programs/UnrealBuildAccelerator/Common/Public/UbaDependencyCrawler.h
@@ -36,7 +36,7 @@ namespace uba
 
 		using AccessFileFunc = Function<bool(const void* data, u64 dataSize)>;
 		using DevirtualizePathFunc = Function<bool(StringBufferBase& inOut)>;
-		using CreateFileFunc = Function<bool(TrackWorkScope& tracker, const StringView& fileName, const AccessFileFunc& func)>;
+		using CreateFileFunc = Function<bool(TrackWorkScope& tracker, const StringView& fileName, const AccessFileFunc& func, bool& stopCrawling)>;
 
 		inline void Init(FileExistsFunc&& fileExistsFunc, TraverseFilesFunc&& traverseFilesFunc, bool useBloomFilter = true);
 		inline bool Add(const tchar* rsp, const tchar* workDir, CreateFileFunc&& func, DevirtualizePathFunc&& devirtualizePathFunc, const tchar* app, DependencyCrawlerType type, u32 ruleIndex);
@@ -175,10 +175,13 @@ namespace uba
 						if (!pch.handled)
 						{
 							fixedFile.Append(TCV(".dep.json"));
+							bool stopCrawling = false;
 							instance.createFileFunc(context.tracker, fixedFile, [&](const void* data, u64 dataSize)
 							{
 								return ParsePch(pch, instance, data, dataSize);
-							});
+							}, stopCrawling);
+							if (stopCrawling)
+								return;
 							pch.handled = true;
 						}
 						instance.pch = &pch;
@@ -188,7 +191,8 @@ namespace uba
 						// If we end up here it means that we have compressed obj files enabled and want to decompress them in parallel
 						m_workManager.AddWork([this, ref, codeFile = cf.path](const WorkContext& context)
 							{
-								ref.instance.createFileFunc(context.tracker, codeFile, {});
+								bool stopCrawling = false;
+								ref.instance.createFileFunc(context.tracker, codeFile, {}, stopCrawling);
 							}, 1, TC("CrawlForDecomp"), ColorWork, false, m_linkerDecompressBottleneck);
 					}
 					else
@@ -237,8 +241,9 @@ namespace uba
 			UBA_LOG_DEVIRTUALIZATION_ERROR("Failed to devirtualize path %s in %s", fixedPath.data, rsp.data);
 		if (CaseInsensitiveFs)
 			fixedPath.MakeLower();
-		if (!instance.createFileFunc(tracker, fixedPath, [&](const void* data, u64 dataSize) { return ParseRsp2(tracker, instance, data, dataSize, rsp.data, outCodeFiles); }))
-			return m_logger.Warning(TC("Failed to parse rsp %s"), rsp.data);
+		bool stopCrawling = false;
+		if (!instance.createFileFunc(tracker, fixedPath, [&](const void* data, u64 dataSize) { return ParseRsp2(tracker, instance, data, dataSize, rsp.data, outCodeFiles); }, stopCrawling))
+			return stopCrawling ? false : m_logger.Warning(TC("Failed to parse rsp %s"), rsp.data);
 		return true;
 	}
 
@@ -616,8 +621,9 @@ namespace uba
 
 	bool DependencyCrawler::ParseCodeFile(TrackWorkScope& tracker, Instance& instance, HandledFiles& handledFiles, const StringView& codeFile, bool parseDefines, const StringView& caller)
 	{
-		if (!instance.createFileFunc(tracker, codeFile, [&](const void* data, u64 dataSize) { return ParseCodeFile2(tracker, instance, handledFiles, data, dataSize, codeFile, parseDefines); }))
-			return m_logger.Warning(TC("Failed to parse code file %s found in %s"), codeFile.data, caller.data);
+		bool stopCrawling = false;
+		if (!instance.createFileFunc(tracker, codeFile, [&](const void* data, u64 dataSize) { return ParseCodeFile2(tracker, instance, handledFiles, data, dataSize, codeFile, parseDefines); }, stopCrawling))
+			return stopCrawling ? false : m_logger.Warning(TC("Failed to parse code file %s found in %s"), codeFile.data, caller.data);
 		return true;
 	}
 
@@ -909,8 +915,10 @@ namespace uba
 
 					m_workManager.AddWork([this, ref = InstanceRef(instance), codeFile = codeFilePath.ToString()](const WorkContext& context)
 						{
-							if (!ref.instance.createFileFunc(context.tracker, codeFile, {}))
-								m_logger.Warning(TC("Failed to open file %s from builtindir"), codeFile.c_str());
+							bool stopCrawling = false;
+							if (!ref.instance.createFileFunc(context.tracker, codeFile, {}, stopCrawling))
+								if (!stopCrawling)
+									m_logger.Warning(TC("Failed to open file %s from builtindir"), codeFile.c_str());
 						}, 1, TC("CrawlBiFile"), ColorWork);
 				}
 			});
