[UBA]
* Fixed macos compile error

diff --git a/Engine/Source/Programs/UnrealBuildAccelerator/Common/Public/UbaDependencyCrawler.h b/Engine/Source/Programs/UnrealBuildAccelerator/Common/Public/UbaDependencyCrawler.h
--- a/Engine/Source/Programs/UnrealBuildAccelerator/Common/Public/UbaDependencyCrawler.h
+++ b/Engine/Source/Programs/UnrealBuildAccelerator/Common/Public/UbaDependencyCrawler.h
@@ -798,7 +798,8 @@ namespace uba
 											if (!handledFiles.lookup.try_emplace(moduleKey).second)
 												return;
 											lock.Leave();
-											instance.createFileFunc(context.tracker, tmp, {});
+											bool stopCrawling = false;
+											instance.createFileFunc(context.tracker, tmp, {}, stopCrawling);
 											return;
 										}
 									}
