[UBA]
* Fixed crash in uba when not using message thread in process

diff --git a/Engine/Source/Programs/UnrealBuildAccelerator/Common/Private/UbaProcess.cpp b/Engine/Source/Programs/UnrealBuildAccelerator/Common/Private/UbaProcess.cpp
--- a/Engine/Source/Programs/UnrealBuildAccelerator/Common/Private/UbaProcess.cpp
+++ b/Engine/Source/Programs/UnrealBuildAccelerator/Common/Private/UbaProcess.cpp
@@ -1881,7 +1881,7 @@ namespace uba
 		}
 
 		bool affinitySet = false;
-		if (m_messageThread->IsInside())
+		if (m_messageThread && m_messageThread->IsInside())
 		{
 			GroupAffinity aff;
 			if (m_messageThread->GetGroupAffinity(aff))
