[UBA]
* Fixed shutdown crash caused by scheduler being destroyed when session was getting message from helpers asking for work. This change resets callback functions during scheduler destroy

diff --git a/Engine/Source/Programs/UnrealBuildAccelerator/Common/Private/UbaScheduler.cpp b/Engine/Source/Programs/UnrealBuildAccelerator/Common/Private/UbaScheduler.cpp
--- a/Engine/Source/Programs/UnrealBuildAccelerator/Common/Private/UbaScheduler.cpp
+++ b/Engine/Source/Programs/UnrealBuildAccelerator/Common/Private/UbaScheduler.cpp
@@ -256,14 +256,29 @@ namespace uba
 			m_session.WaitOnAllClients();
 		}
 
+		m_session.SetRemoteProcessReturnedEvent({});
+		m_session.SetRemoteProcessSlotAvailableEvent({});
+
 		if (m_memWatchdog)
 		{
+			m_session.SetNativeProcessCreatedFunc({});
 			m_memThreadEvent.Set();
 			m_memThread.Wait();
 			SaveMemTrackTable();
 		}
 
-		Cleanup();
+		{
+			SCOPED_FUTEX(m_processEntriesLock, lock);
+			for (auto& entry : m_processEntries)
+			{
+				UBA_ASSERTF(entry.runStatus > RunStatus_Running, TC("Found processes in queue/running state when stopping scheduler."));
+				delete[] entry.dependencies;
+				delete entry.info;
+			}
+			m_processEntries.clear();
+			m_processEntriesRunStart = 0;
+			m_session.SetOuterScheduler(nullptr);
+		}
 	}
 
 	void Scheduler::Cancel()
@@ -292,20 +307,6 @@ namespace uba
 			SkipProcess(*pi);
 	}
 
-	void Scheduler::Cleanup()
-	{
-		SCOPED_FUTEX(m_processEntriesLock, lock);
-		for (auto& entry : m_processEntries)
-		{
-			UBA_ASSERTF(entry.runStatus > RunStatus_Running, TC("Found processes in queue/running state when stopping scheduler."));
-			delete[] entry.dependencies;
-			delete entry.info;
-		}
-		m_processEntries.clear();
-		m_processEntriesRunStart = 0;
-		m_session.SetOuterScheduler(nullptr);
-	}
-
 	void Scheduler::SetMaxLocalProcessors(u32 maxLocalProcessors)
 	{
 		m_maxLocalProcessors = maxLocalProcessors;
diff --git a/Engine/Source/Programs/UnrealBuildAccelerator/Common/Public/UbaScheduler.h b/Engine/Source/Programs/UnrealBuildAccelerator/Common/Public/UbaScheduler.h
--- a/Engine/Source/Programs/UnrealBuildAccelerator/Common/Public/UbaScheduler.h
+++ b/Engine/Source/Programs/UnrealBuildAccelerator/Common/Public/UbaScheduler.h
@@ -127,7 +127,6 @@ namespace uba
 		UBA_NOINLINE void ThreadLoop();
 
 		void SkipAllQueued();
-		void Cleanup();
 		bool ProcessReturned(Process& process, bool isLocal);
 		void HandleCacheMissed(u32 processIndex);
 		void RemoteSlotAvailable(bool isCrossArchitecture);
