[UBA]
* Added so casdb deletes defaults to first move to a temporary folder and then deleted in a separate job
* Added CasEntryWriteLock and CasEntryReadLock. Intention is to implement code where threads waiting on a cas entry lock can help out doing the work needed for them to get the lock (compress file etc)

diff --git a/Engine/Source/Programs/UnrealBuildAccelerator/Common/Private/UbaCacheClient.cpp b/Engine/Source/Programs/UnrealBuildAccelerator/Common/Private/UbaCacheClient.cpp
--- a/Engine/Source/Programs/UnrealBuildAccelerator/Common/Private/UbaCacheClient.cpp
+++ b/Engine/Source/Programs/UnrealBuildAccelerator/Common/Private/UbaCacheClient.cpp
@@ -1335,21 +1335,20 @@ namespace uba
 
 				FileAccessor file(m_logger, casKeyFileName.data);
 
+				StorageImpl::CasEntryWriteLock casEntryLock(*casEntry);
+				--casEntry->refCount; // HasCasFile increase refcount
+				auto mh = casEntry->mappingHandle;
+				if (mh.IsValid()) // If file was created by helper it will be in the transient mapped memory
 				{
-					SCOPED_READ_LOCK(casEntry->lock, entryLock);
-					--casEntry->refCount; // HasCasFile increase refcount
-					auto mh = casEntry->mappingHandle;
-					if (mh.IsValid()) // If file was created by helper it will be in the transient mapped memory
-					{
-						u64 mappingOffset = casEntry->mappingOffset;
-						u64 mappingSize = casEntry->mappingSize;
-						entryLock.Leave();
+					u64 mappingOffset = casEntry->mappingOffset;
+					u64 mappingSize = casEntry->mappingSize;
+					casEntryLock.Leave();
 
-						mappedView = m_storage.m_casDataBuffer.MapView(mh, mappingOffset, mappingSize, path.data);
-						fileData = mappedView.memory;
-						fileSize = mappedView.size;
-					}
+					mappedView = m_storage.m_casDataBuffer.MapView(mh, mappingOffset, mappingSize, path.data);
+					fileData = mappedView.memory;
+					fileSize = mappedView.size;
 				}
+				casEntryLock.Leave();
 
 				if (!fileData)
 				{
diff --git a/Engine/Source/Programs/UnrealBuildAccelerator/Common/Private/UbaStorage.cpp b/Engine/Source/Programs/UnrealBuildAccelerator/Common/Private/UbaStorage.cpp
--- a/Engine/Source/Programs/UnrealBuildAccelerator/Common/Private/UbaStorage.cpp
+++ b/Engine/Source/Programs/UnrealBuildAccelerator/Common/Private/UbaStorage.cpp
@@ -91,6 +91,7 @@ namespace uba
 		const ConfigTable& table = *tablePtr;
 		table.GetValueAsString(rootDir, TC("RootDir"));
 		table.GetValueAsBool(asyncUnmapViewOfFile, TC("AsyncUnmapViewOfFile"));
+		table.GetValueAsBool(useDeferredDeletes, TC("UseDeferredDeletes"));
 		table.GetValueAsU64(casCapacityBytes, TC("CasCapacityBytes"));
 		table.GetValueAsBool(keepTransientDataMapped, TC("KeepTransientDataMapped"));
 		table.GetValueAsBool(storeCompressed, TC("StoreCompressed"));
@@ -139,10 +140,50 @@ namespace uba
 		u64 mapSize = 0;
 	};
 
+	StorageImpl::CasEntryReadLock::CasEntryReadLock(CasEntry& e) : entry(e)
+	{
+		active = true;
+		entry.internalLock.EnterRead();
+	}
+
+	StorageImpl::CasEntryReadLock::~CasEntryReadLock()
+	{
+		if (active)
+			entry.internalLock.LeaveRead();
+	}
+
+	StorageImpl::CasEntryWriteLock::CasEntryWriteLock(CasEntry& e) : entry(e)
+	{
+		active = true;
+		entry.internalLock.Enter();
+	}
+
+	StorageImpl::CasEntryWriteLock::~CasEntryWriteLock()
+	{
+		if (active)
+			entry.internalLock.Leave();
+	}
+
+	void StorageImpl::CasEntryWriteLock::Enter()
+	{
+		if (active)
+			return;
+		active = true;
+		entry.internalLock.Enter();
+	}
+
+	void StorageImpl::CasEntryWriteLock::Leave()
+	{
+		if (!active)
+			return;
+		active = false;
+		entry.internalLock.Leave();
+	}
+
 	void StorageImpl::CasEntryAccessed(CasEntry& casEntry)
 	{
 		{
-			SCOPED_READ_LOCK(casEntry.lock, l); // Note, this lock is taken again outside CasEntryAccessed.. so if this takes a long time it won't help to remove this lock
+			CasEntryReadLock casEntryLock(casEntry); // Note, this lock is taken again outside CasEntryAccessed.. so if this takes a long time it won't help to remove this lock
 			if (casEntry.dropped)
 				return;
 			//if (casEntry.mappingHandle.IsValid())
@@ -208,11 +249,11 @@ namespace uba
 		for (CasEntry* it = m_oldestAccessed; it;)
 		{
 			CasEntry& ce = *it;
-			ce.lock.Enter();
+			ce.internalLock.Enter();
 
 			if (ce.verified && !m_allowDeleteVerified) // Can't remove these since they might be in mapping tables and actively being used. It is tempting to delete dropped cas files but we can't
 			{
-				ce.lock.Leave();
+				ce.internalLock.Leave();
 				break;
 			}
 
@@ -220,7 +261,7 @@ namespace uba
 			{
 				if (ce.beingWritten)
 					m_logger.Info(TC("Want to delete file that is begin written. We should never end up here"));
-				ce.lock.Leave();
+				ce.internalLock.Leave();
 				it = ce.prevAccessed;
 				continue;
 			}
@@ -253,39 +294,104 @@ namespace uba
 
 		lock.Leave();
 
-		for (Rec& rec : toDelete)
-		{
-			StringBuffer<> casFile;
-			StorageImpl::GetCasFileName(casFile, rec.casEntry.key);
+		u64 startTime = GetTime();
 
-			if (!DeleteFileW(casFile.data))
+		const tchar* operationType = TC("delete");
+#if PLATFORM_WINDOWS
+		if (m_useDeferredDeletes)
+		{
+			// On windows, deletefile is stupid slow.. so we move the file somewhere first so we can delete it there instead
+			operationType = TC("move");
+			Vector<u32> toDelete2;
+			toDelete2.reserve(toDelete.size());
+			for (Rec& rec : toDelete)
 			{
-				u32 error = GetLastError();
-				if (error != ERROR_FILE_NOT_FOUND && error != ERROR_PATH_NOT_FOUND)
+				StringBuffer<> casFile;
+				StorageImpl::GetCasFileName(casFile, rec.casEntry.key);
+
+				static Atomic<u32> deleteIndexCounter;
+				u32 deleteIndex = deleteIndexCounter++;
+
+				StringBuffer<> deleteName;
+				deleteName.Append(m_tempPath).AppendValue(deleteIndex).Append(".del");
+
+				if (!MoveFileW(casFile.data, deleteName.data))
 				{
-					m_logger.Error(TC("Failed to delete %s while handling overflow (%s)"), casFile.data, LastErrorToText(error).data);
-					rec.casEntry.exists = true;
-					rec.casEntry.size = rec.size;
-					rec.casEntry.lock.Leave();
-
-					// TODO: Should this instead set some overflow
-					/*
-					SCOPED_FUTEX(m_accessLock, lock);
-					m_casEvictedBytes -= rec.size;
-					--m_casEvictedCount;
-					m_casTotalBytes += rec.size;
-					AttachEntry(&rec.casEntry); // TODO: This should be re-added in the end
-					*/
-					continue;
+					u32 error = GetLastError();
+					if (error != ERROR_FILE_NOT_FOUND && error != ERROR_PATH_NOT_FOUND)
+					{
+						m_logger.Error(TC("Failed to delete %s while handling overflow (%s)"), casFile.data, LastErrorToText(error).data);
+						rec.casEntry.exists = true;
+						rec.casEntry.size = rec.size;
+						rec.casEntry.internalLock.Leave();
+						continue;
+					}
 				}
+
+				toDelete2.push_back(deleteIndex);
+
+				//m_logger.Warning(TC("Evicted %s due to overflow"), casFile.data);
+
+				casEntry.verified = true; // Verified to be deleted
+
+				rec.casEntry.internalLock.Leave();
 			}
 
-			//m_logger.Warning(TC("Evicted %s due to overflow"), casFile.data);
+			m_workManager.AddWork([toDelete2 = std::move(toDelete2), this](const WorkContext&)
+				{
+					u64 startTime = GetTime();
+					for (u32 index : toDelete2)
+					{
+						StringBuffer<> deleteName;
+						deleteName.Append(m_tempPath).AppendValue(index).Append(".del");
+						DeleteFileW(deleteName.data);
+					}
+					u64 time = GetTime() - startTime;
+					if (TimeToMs(time) > 30*1000)
+						m_logger.Info(TC("Took delete job %s to delete %llu files"), TimeToText(time).str, toDelete2.size());
+				}, 1, TC("DeleteCasFiles"));
+		}
+		else
+#endif
+		{
+			for (Rec& rec : toDelete)
+			{
+				StringBuffer<> casFile;
+				StorageImpl::GetCasFileName(casFile, rec.casEntry.key);
 
-			casEntry.verified = true; // Verified to be deleted
+				if (!DeleteFileW(casFile.data))
+				{
+					u32 error = GetLastError();
+					if (error != ERROR_FILE_NOT_FOUND && error != ERROR_PATH_NOT_FOUND)
+					{
+						m_logger.Error(TC("Failed to delete %s while handling overflow (%s)"), casFile.data, LastErrorToText(error).data);
+						rec.casEntry.exists = true;
+						rec.casEntry.size = rec.size;
+						rec.casEntry.internalLock.Leave();
+
+						// TODO: Should this instead set some overflow
+						/*
+						SCOPED_FUTEX(m_accessLock, lock);
+						m_casEvictedBytes -= rec.size;
+						--m_casEvictedCount;
+						m_casTotalBytes += rec.size;
+						AttachEntry(&rec.casEntry); // TODO: This should be re-added in the end
+						*/
+						continue;
+					}
+				}
+
+				//m_logger.Warning(TC("Evicted %s due to overflow"), casFile.data);
 
-			rec.casEntry.lock.Leave();
+				casEntry.verified = true; // Verified to be deleted
+
+				rec.casEntry.internalLock.Leave();
+			}
 		}
+		u64 time = GetTime() - startTime;
+		if (TimeToMs(time) > 30*1000)
+			m_logger.Info(TC("Took %s to %s %llu files"), TimeToText(time).str, operationType, toDelete.size());
+
 	}
 
 	void StorageImpl::CasEntryDeleted(CasEntry& casEntry, u64 size)
@@ -482,7 +588,7 @@ namespace uba
 		return w.DecompressMemoryToMemory(compressedData, compressedSize, writeData, decompressedSize, readHint, writeHint);
 	}
 
-	bool StorageImpl::VerifyExisting(bool& outReturnValue, ScopedWriteLock& entryLock, const CasKey& casKey, CasEntry& casEntry, StringBufferBase& casFile)
+	bool StorageImpl::VerifyExisting(bool& outReturnValue, CasEntryWriteLock& casEntryLock, const CasKey& casKey, CasEntry& casEntry, StringBufferBase& casFile)
 	{
 		u64 outFileSize = 0;
 		u32 outAttributes = 0;
@@ -512,7 +618,7 @@ namespace uba
 			UBA_ASSERT(!casEntry.verified || casEntry.size);
 			casEntry.verified = true;
 			casEntry.exists = true;
-			entryLock.Leave();
+			casEntryLock.Leave();
 			CasEntryWritten(casEntry, outFileSize);
 			outReturnValue = true;
 			return true;
@@ -528,7 +634,7 @@ namespace uba
 		CasEntry& casEntry = GetOrCreateCasEntry(casKey);
 		CasEntryAccessed(casEntry);
 
-		SCOPED_WRITE_LOCK(casEntry.lock, entryLock);
+		CasEntryWriteLock casEntryLock(casEntry);
 
 		if (casEntry.beingWritten)
 		{
@@ -537,9 +643,9 @@ namespace uba
 			while (casEntry.beingWritten)
 			{
 				CasKey key = casEntry.key;
-				entryLock.Leave();
+				casEntryLock.Leave();
 				Sleep(100);
-				entryLock.Enter();
+				casEntryLock.Enter();
 
 				if (++waitCount < 12*60*10)
 					continue;
@@ -561,7 +667,7 @@ namespace uba
 			return false;
 	
 		bool verifyReturnValue;
-		if (!casEntry.verified && VerifyExisting(verifyReturnValue, entryLock, casKey, casEntry, casFile))
+		if (!casEntry.verified && VerifyExisting(verifyReturnValue, casEntryLock, casKey, casEntry, casFile))
 			return verifyReturnValue;
 
 		if (deferCreation)
@@ -616,7 +722,7 @@ namespace uba
 			return false;
 
 		casEntry.exists = true;
-		entryLock.Leave();
+		casEntryLock.Leave();
 		CasEntryWritten(casEntry, res.storedSize);
 		return true;
 	}
@@ -884,6 +990,7 @@ namespace uba
 		m_manuallyHandleOverflow = info.manuallyHandleOverflow;
 		m_asyncUnmapViewOfFile = info.asyncUnmapViewOfFile;
 		m_allowDeleteVerified = info.allowDeleteVerified;
+		m_useDeferredDeletes = info.useDeferredDeletes;
 		m_writeToDisk = info.writeToDisk;
 		if (!m_writeToDisk)
 			m_casCapacityBytes = 0;
@@ -1891,7 +1998,7 @@ namespace uba
 		lookupLock.Leave();
 		CasEntryAccessed(casEntry);
 
-		SCOPED_WRITE_LOCK(casEntry.lock, entryLock);
+		CasEntryWriteLock casEntryLock(casEntry);
 
 		if (casEntry.verified && casEntry.exists)
 		{
@@ -1949,7 +2056,7 @@ namespace uba
 			++casEntry.refCount;
 		}
 
-		entryLock.Leave();
+		casEntryLock.Leave();
 
 		if (!hasMapping)
 			CasEntryWritten(casEntry, res.storedSize);
@@ -1962,7 +2069,7 @@ namespace uba
 		CasEntryAccessed(casEntry);
 
 		{
-			SCOPED_READ_LOCK(casEntry.lock, entryLock);
+			CasEntryReadLock casEntryLock(casEntry);
 			if (casEntry.verified)
 			{
 				if (casEntry.exists)
@@ -1972,7 +2079,7 @@ namespace uba
 			}
 		}
 
-		SCOPED_WRITE_LOCK(casEntry.lock, entryLock);
+		CasEntryWriteLock casEntryLock(casEntry);
 
 		if (casEntry.verified)
 		{
@@ -1987,7 +2094,7 @@ namespace uba
 			return false;
 
 		bool verifyReturnValue;
-		if (VerifyExisting(verifyReturnValue, entryLock, casKey, casEntry, casFile))
+		if (VerifyExisting(verifyReturnValue, casEntryLock, casKey, casEntry, casFile))
 			return verifyReturnValue;
 
 		casEntry.exists = false;
@@ -1998,7 +2105,7 @@ namespace uba
 		if (!WriteCasFile(res, fileName, casKey))
 			return false;
 		casEntry.exists = true;
-		entryLock.Leave();
+		casEntryLock.Leave();
 		CasEntryWritten(casEntry, res.storedSize);
 		return true;
 	}
@@ -2025,7 +2132,7 @@ namespace uba
 		}
 		CasEntry& casEntry = findIt->second;
 		lookupLock.Leave();
-		SCOPED_WRITE_LOCK(casEntry.lock, entryLock);
+		CasEntryWriteLock casEntryLock(casEntry);
 		//if (!casEntry.verified)
 		//{
 		//	m_logger.Error(TC("Trying to use unverified mapping of %s (%s)"), CasKeyString(casKey).str, hint);
@@ -2034,7 +2141,7 @@ namespace uba
 		FileMappingHandle handle = casEntry.mappingHandle;
 		u64 offset = casEntry.mappingOffset;
 		u64 size = casEntry.mappingSize;
-		entryLock.Leave();
+		casEntryLock.Leave();
 
 		auto res = m_casDataBuffer.MapView(handle, offset, size, hint);
 		if (!res.memory)
@@ -2071,7 +2178,7 @@ namespace uba
 		CasEntry& casEntry = findIt->second;
 		lookupLock.Leave();
 	
-		SCOPED_WRITE_LOCK(casEntry.lock, entryLock);
+		CasEntryWriteLock casEntryLock(casEntry);
 
 		UBA_ASSERT(casEntry.refCount == 0);
 
@@ -2111,7 +2218,7 @@ namespace uba
 			casEntry.verified = true;
 			casEntry.exists = false;
 			casEntry.dropped = false;
-			entryLock.Leave();
+			casEntryLock.Leave();
 	
 			CasEntryDeleted(casEntry, sizeDeleted);
 		}
@@ -2182,7 +2289,7 @@ namespace uba
 			}
 
 			CasEntry& casEntry = *casEntryPtr;
-			SCOPED_READ_LOCK(casEntry.lock, casEntryLock);
+			CasEntryReadLock casEntryLock(casEntry);
 			UBA_ASSERT(casEntry.verified);
 			UBA_ASSERT(casEntry.exists);
 			--casEntry.refCount; // Decrease after HasCasFile increase it
@@ -2601,7 +2708,7 @@ namespace uba
 		if  (findIt == m_casLookup.end())
 			return m_logger.Error(TC("GetCasEntry - Couldn't find cas entry for %s"), hint);
 		CasEntry& casEntry = findIt->second;
-		SCOPED_READ_LOCK(casEntry.lock, casEntryLock);
+		CasEntryReadLock casEntryLock(casEntry);
 		if (!casEntry.mappingHandle.IsValid())
 			return m_logger.Error(TC("Code path only supports mapped cas entries"));
 		if (!IsCompressed(casEntry.key))
@@ -2624,7 +2731,7 @@ namespace uba
 		if  (findIt == m_casLookup.end())
 			return m_logger.Error(TC("GetCasEntry - Couldn't find cas entry for %s"), hint);
 		CasEntry& casEntry = findIt->second;
-		SCOPED_READ_LOCK(casEntry.lock, casEntryLock);
+		CasEntryReadLock casEntryLock(casEntry);
 		FileMappingScope scope(m_logger);
 		u8* fileMem = scope.MapView(TC("ReadFileSize"), casEntry.mappingHandle, casEntry.mappingOffset, casEntry.mappingSize, hint);
 		if (!fileMem)
diff --git a/Engine/Source/Programs/UnrealBuildAccelerator/Common/Private/UbaStorageClient.cpp b/Engine/Source/Programs/UnrealBuildAccelerator/Common/Private/UbaStorageClient.cpp
--- a/Engine/Source/Programs/UnrealBuildAccelerator/Common/Private/UbaStorageClient.cpp
+++ b/Engine/Source/Programs/UnrealBuildAccelerator/Common/Private/UbaStorageClient.cpp
@@ -251,7 +251,7 @@ namespace uba
 
 		StorageStats& stats = Stats();
 		CasEntry* casEntry = nullptr;
-		auto casEntryLock = MakeGuard([&]() { if (casEntry) casEntry->lock.Leave(); });
+		auto casEntryLock = MakeGuard([&]() { if (casEntry) casEntry->internalLock.Leave(); });
 		if (shouldStore)
 		{
 			TimerScope ts(stats.ensureCas);
@@ -263,7 +263,7 @@ namespace uba
 			casEntry = &m_casLookup.find(casKey)->second;
 			lock.Leave();
 
-			casEntry->lock.Enter();
+			casEntry->internalLock.Enter();
 			if (casEntry->verified && casEntry->exists)
 				return true;
 
diff --git a/Engine/Source/Programs/UnrealBuildAccelerator/Common/Private/UbaStorageServer.cpp b/Engine/Source/Programs/UnrealBuildAccelerator/Common/Private/UbaStorageServer.cpp
--- a/Engine/Source/Programs/UnrealBuildAccelerator/Common/Private/UbaStorageServer.cpp
+++ b/Engine/Source/Programs/UnrealBuildAccelerator/Common/Private/UbaStorageServer.cpp
@@ -238,7 +238,7 @@ namespace uba
 				if (store.mappedView.isIndependent)
 					CloseFileMapping(m_logger, store.mappedView.handle, TC("StoreDone"));
 
-				SCOPED_WRITE_LOCK(store.casEntry->lock, entryLock);
+				CasEntryWriteLock casEntryLock(*store.casEntry);
 				if (!store.casEntry->beingWritten)
 					m_logger.Warning(TC("CasEntry does not have being written flag?")); // Searching for a bug
 
@@ -293,7 +293,7 @@ namespace uba
 
 	void StorageServer::ReleaseCasEntry(CasEntry& casEntry)
 	{
-		SCOPED_WRITE_LOCK(casEntry.lock, casEntryLock);
+		CasEntryWriteLock casEntryLock(casEntry);
 		ReleaseCasEntryNoLock(casEntry);
 	}
 
@@ -371,7 +371,7 @@ namespace uba
 		return false;
 	}
 
-	bool StorageServer::WaitForWritten(CasEntry& casEntry, ScopedWriteLock& entryLock, const ConnectionInfo& connectionInfo, const tchar* hint)
+	bool StorageServer::WaitForWritten(CasEntry& casEntry, CasEntryWriteLock& casEntryLock, const ConnectionInfo& connectionInfo, const tchar* hint)
 	{
 		int waitCount = 0;
 		while (true)
@@ -379,9 +379,9 @@ namespace uba
 			if (!casEntry.beingWritten)
 				return true;
 			CasKey key = casEntry.key;
-			entryLock.Leave();
+			casEntryLock.Leave();
 			Sleep(100);
-			auto g = MakeGuard([&]() { entryLock.Enter(); });
+			auto g = MakeGuard([&]() { casEntryLock.Enter(); });
 
 			if (connectionInfo.ShouldDisconnect())
 				return false;
@@ -403,7 +403,7 @@ namespace uba
 		}
 	}
 
-	bool StorageServer::VerifyExists(bool& outExists, CasEntry& casEntry, ScopedWriteLock& entryLock, const CasKey& casKey)
+	bool StorageServer::VerifyExists(bool& outExists, CasEntry& casEntry, CasEntryWriteLock& casEntryLock, const CasKey& casKey)
 	{
 		outExists = false;
 		StringBuffer<> casFile;
@@ -424,7 +424,7 @@ namespace uba
 			{
 				casEntry.verified = true;
 				outExists = true;
-				entryLock.Leave();
+				casEntryLock.Leave();
 				CasEntryWritten(casEntry, outFileSize);
 			}
 		}
@@ -599,9 +599,12 @@ namespace uba
 
 				CasEntry* casEntryPtr = nullptr;
 				bool has = HasCasFile(casKey, &casEntryPtr); // HasCasFile writes deferred cas entries if in queue. Also increase CasEntry refcount
+
+				todoRemoveMe(3);
+
 				if (!has)
 				{
-					todoRemoveMe(3);
+					todoRemoveMe(4);
 					if (!EnsureCasFile(casKey, nullptr) && m_allowHintAsFallback)
 					{
 						// Last resort.. use hint to load file into cas (hint should be renamed since it is now a critical parameter)
@@ -658,13 +661,13 @@ namespace uba
 				}
 
 				CasEntry& casEntry = *casEntryPtr;
-				SCOPED_WRITE_LOCK(casEntry.lock, entryLock);
+				CasEntryWriteLock casEntryLock(casEntry);
 				if (!has) // HasCasFile increase refcount, otherwise we have to do it manually
 					++casEntry.refCount;
 				if (!casEntry.refCount)
 					m_logger.Error(TC("CasEntry refcount wrap-around during fetch. This is bad, report to uba devs"));
-				entryLock.Leave();
-				auto rlg = MakeGuard([&]() { entryLock.Enter(); ReleaseCasEntryNoLock(casEntry); });
+				casEntryLock.Leave();
+				auto rlg = MakeGuard([&]() { casEntryLock.Enter(); ReleaseCasEntryNoLock(casEntry); });
 
 				if (!casEntry.exists)
 				{
@@ -673,7 +676,7 @@ namespace uba
 					return true;
 				}
 
-				todoRemoveMe(4);
+				todoRemoveMe(5);
 
 				if (casEntry.disallowed)
 				{
@@ -686,7 +689,7 @@ namespace uba
 				FileHandle readFileHandle = InvalidFileHandle;
 				auto rfg = MakeGuard([&]()
 					{
-						todoRemoveMe(13);
+						todoRemoveMe(14);
 						if (!CloseFile(nullptr, readFileHandle))
 							m_logger.Error(TC("Failed to close file %s (%s)"), casFile.data, LastErrorToText().data);
 					});
@@ -700,7 +703,7 @@ namespace uba
 					{
 						if (ownsMapping)
 						{
-							todoRemoveMe(11);
+							todoRemoveMe(12);
 							if (mappedView.memory)
 								if (!UnmapViewOfFile(m_logger, mappedView.memory, fileSize, TC("FetchBegin")))
 									m_logger.Error(TC("Failed to unmap memory %s at 0x%llx with size %llu (%s)"), hint.data, mappedView.memory, fileSize, LastErrorToText().data);
@@ -710,12 +713,12 @@ namespace uba
 						}
 						else
 						{
-							todoRemoveMe(12);
+							todoRemoveMe(13);
 							m_casDataBuffer.UnmapView(mappedView, TC("FetchBegin"));
 						}
 					});
 
-				todoRemoveMe(5);
+				todoRemoveMe(6);
 
 				bool useFileMapping = casEntry.mappingHandle.IsValid();
 				if (useFileMapping)
@@ -755,7 +758,7 @@ namespace uba
 					}
 				}
 
-				todoRemoveMe(6);
+				todoRemoveMe(7);
 
 				if (detailedTrace)
 					m_trace->FileFetchSize(connectionInfo.GetId(), casKey, fileSize);
@@ -776,7 +779,7 @@ namespace uba
 				u32 toWrite = u32(Min(left, capacityLeft));
 				void* writeBuffer = writer.AllocWrite(toWrite);
 
-				todoRemoveMe(7);
+				todoRemoveMe(8);
 
 				if (useFileMapping)
 				{
@@ -810,7 +813,7 @@ namespace uba
 					readFileHandle = InvalidFileHandle;
 				}
 
-				todoRemoveMe(8);
+				todoRemoveMe(9);
 
 				u64 actualSize = fileSize;
 				if (m_storeCompressed)
@@ -832,7 +835,7 @@ namespace uba
 					return true;
 				}
 
-				todoRemoveMe(9);
+				todoRemoveMe(10);
 
 				rlg.Cancel();
 				mvg.Cancel();
@@ -847,7 +850,7 @@ namespace uba
 				fetch.clientId = connectionInfo.GetId();
 				lock.Leave();
 
-				todoRemoveMe(10);
+				todoRemoveMe(11);
 
 				mappedView.size = fileSize;
 
@@ -926,15 +929,15 @@ namespace uba
 				CasEntry& casEntry = casInsres.first->second;
 				lookupLock.Leave();
 
-				SCOPED_WRITE_LOCK(casEntry.lock, entryLock);
+				CasEntryWriteLock casEntryLock(casEntry);
 				
-				if (!WaitForWritten(casEntry, entryLock, connectionInfo, TC("UNKNOWN")))
+				if (!WaitForWritten(casEntry, casEntryLock, connectionInfo, TC("UNKNOWN")))
 					return false;
 
 				bool exists = casEntry.verified && casEntry.exists;
 
 				if (!exists && casEntry.exists)
-					if (!VerifyExists(exists, casEntry, entryLock, casKey))
+					if (!VerifyExists(exists, casEntry, casEntryLock, casKey))
 						return false;
 
 				UBA_ASSERTF(!m_dropIndependentMappingsOnCopy, TC("Not supported... ref count is increased in next"));
@@ -963,7 +966,7 @@ namespace uba
 				auto failGuard = MakeGuard([&]() { ReleaseActiveRefCount(casEntry.key, clientId); });
 				auto success = [&]() { failGuard.Cancel(); return true; };
 
-				SCOPED_WRITE_LOCK(casEntry.lock, entryLock);
+				CasEntryWriteLock casEntryLock(casEntry);
 
 				// We need to take ref count on cas entry early since there are a bunch of sections
 				// that temporarily release cas entry lock where file could actually be dropped
@@ -982,7 +985,7 @@ namespace uba
 					if (findIt != m_deferredCasCreationLookup.end())
 					{
 						deferredLock.Leave();
-						entryLock.Leave();
+						casEntryLock.Leave();
 						if (HasCasFile(casKey))
 						{
 							writer.WriteU16(u16(~0));
@@ -991,7 +994,7 @@ namespace uba
 							UpdateWaitEntries(casEntry.key, true);
 							return success();
 						}
-						entryLock.Enter();
+						casEntryLock.Enter();
 					}
 				}
 
@@ -999,7 +1002,7 @@ namespace uba
 				{
 					if (bool exists = casEntry.exists)
 					{
-						if (!VerifyExists(exists, casEntry, entryLock, casKey))
+						if (!VerifyExists(exists, casEntry, casEntryLock, casKey))
 							return false;
 						if (exists)
 						{
@@ -1016,12 +1019,12 @@ namespace uba
 				}
 				else
 				{
-					if (!WaitForWritten(casEntry, entryLock, connectionInfo, hint.data))
+					if (!WaitForWritten(casEntry, casEntryLock, connectionInfo, hint.data))
 						return false;
 
 					if (casEntry.exists)
 					{
-						entryLock.Leave();
+						casEntryLock.Leave();
 						CasEntryAccessed(casEntry);
 						writer.WriteU16(u16(~0));
 						writer.WriteBool(m_traceStore);
@@ -1146,7 +1149,7 @@ namespace uba
 
 					CasEntry& casEntry = *activeStore.casEntry;
 					{
-						SCOPED_WRITE_LOCK(casEntry.lock, entryLock);
+						CasEntryWriteLock casEntryLock(casEntry);
 						casEntry.mappingHandle = activeStore.mappedView.handle;
 						casEntry.mappingOffset = activeStore.mappedView.offset;
 						casEntry.mappingSize = totalWritten;
diff --git a/Engine/Source/Programs/UnrealBuildAccelerator/Common/Public/UbaStorage.h b/Engine/Source/Programs/UnrealBuildAccelerator/Common/Public/UbaStorage.h
--- a/Engine/Source/Programs/UnrealBuildAccelerator/Common/Public/UbaStorage.h
+++ b/Engine/Source/Programs/UnrealBuildAccelerator/Common/Public/UbaStorage.h
@@ -113,6 +113,7 @@ namespace uba
 		bool writeToDisk = true;
 		bool keepTransientDataMapped = true;	// Will keep the transient data mapped in instead of doing map/unmap everytime we access it
 		bool allowDeleteVerified = false;
+		bool useDeferredDeletes = true;
 		MutexHandle exclusiveMutex = InvalidMutexHandle; // Set this to hand over pre-exclusive access to storage. Storage will release mutex on shutdown
 		u8 casCompressor = 0;
 		u8 casCompressionLevel = 0;
@@ -179,6 +180,7 @@ namespace uba
 		virtual void AddStats(StorageStats& stats) override;
 
 		struct CasEntry;
+		struct CasEntryWriteLock;
 
 		virtual bool HasCasFile(const CasKey& casKey, CasEntry** out = nullptr);
 		bool EnsureCasFile(const CasKey& casKey, const tchar* fileName);
@@ -194,7 +196,7 @@ namespace uba
 		virtual bool WriteCompressedFile(WriteResult& out, const tchar* from, FileHandle readHandle, u8* readMem, u64 fileSize, const tchar* toFile, const void* header, u64 headerSize, u64 lastWriteTime = 0) final;
 		bool WriteCasFileNoCheck(WriteResult& out, const StringKey& fileNameKey, const tchar* fileName, const CasKey& casKey, const tchar* casFile, bool storeCompressed);
 		bool WriteCasFile(WriteResult& out, const tchar* fileName, const CasKey& casKey);
-		bool VerifyExisting(bool& outReturnValue, ScopedWriteLock& entryLock, const CasKey& casKey, CasEntry& casEntry, StringBufferBase& casFile);
+		bool VerifyExisting(bool& outReturnValue, CasEntryWriteLock& casEntryLock, const CasKey& casKey, CasEntry& casEntry, StringBufferBase& casFile);
 		bool AddCasFile(StringKey fileNameKey, const tchar* fileName, const CasKey& casKey, bool deferCreation, u64 fileSize = InvalidValue, u64 lastWriteTime = 0);
 		void CasEntryAccessed(const CasKey& casKey);
 		virtual bool IsDisallowedPath(const tchar* fileName);
@@ -243,10 +245,29 @@ namespace uba
 		UnorderedMap<StringKey, FileEntry> m_fileTableLookup;
 
 		ReaderWriterLock m_casLookupLock;
+
+		struct CasEntryReadLock
+		{
+			CasEntryReadLock(CasEntry& e);
+			~CasEntryReadLock();
+			CasEntry& entry;
+			bool active;
+		};
+
+		struct CasEntryWriteLock
+		{
+			CasEntryWriteLock(CasEntry& e);
+			~CasEntryWriteLock();
+			void Enter();
+			void Leave();
+			CasEntry& entry;
+			bool active;
+		};
+
 		struct CasEntry
 		{
 			CasEntry(const CasKey& k) : key(k) {}
-			ReaderWriterLock lock;
+			ReaderWriterLock internalLock;
 			CasKey key;
 			CasEntry* prevAccessed = nullptr;
 			CasEntry* nextAccessed = nullptr;
@@ -282,6 +303,7 @@ namespace uba
 		bool m_manuallyHandleOverflow = false;
 		bool m_asyncUnmapViewOfFile = true;
 		bool m_allowDeleteVerified = false;
+		bool m_useDeferredDeletes;
 		bool m_writeToDisk;
 
 		MutexHandle m_exclusiveMutex = InvalidMutexHandle;
diff --git a/Engine/Source/Programs/UnrealBuildAccelerator/Common/Public/UbaStorageServer.h b/Engine/Source/Programs/UnrealBuildAccelerator/Common/Public/UbaStorageServer.h
--- a/Engine/Source/Programs/UnrealBuildAccelerator/Common/Public/UbaStorageServer.h
+++ b/Engine/Source/Programs/UnrealBuildAccelerator/Common/Public/UbaStorageServer.h
@@ -52,8 +52,8 @@ namespace uba
 		void ReleaseActiveRefCount(const CasKey& casKey, u32 clientId);
 		void UpdateWaitEntries(const CasKey& casKey, bool success);
 		bool HandleMessage(const ConnectionInfo& connectionInfo, const WorkContext& workContext, u8 messageType, BinaryReader& reader, BinaryWriter& writer);
-		bool WaitForWritten(CasEntry& casEntry, ScopedWriteLock& entryLock, const ConnectionInfo& connectionInfo, const tchar* hint);
-		bool VerifyExists(bool& outExists, CasEntry& casEntry, ScopedWriteLock& entryLock, const CasKey& casKey);
+		bool WaitForWritten(CasEntry& casEntry, CasEntryWriteLock& casEntryLock, const ConnectionInfo& connectionInfo, const tchar* hint);
+		bool VerifyExists(bool& outExists, CasEntry& casEntry, CasEntryWriteLock& casEntryLock, const CasKey& casKey);
 		bool GetProxy(BinaryWriter& writer, u32 clientId, bool writeCasHeader);
 
 		u16 PopId();
