[UBA]
* Changed so deleting cas files on overflow is done in jobs in batches of 10 (value is tweakable in [Session] deleteCasBatchSize)

diff --git a/Engine/Source/Programs/UnrealBuildAccelerator/Common/Private/UbaStorage.cpp b/Engine/Source/Programs/UnrealBuildAccelerator/Common/Private/UbaStorage.cpp
--- a/Engine/Source/Programs/UnrealBuildAccelerator/Common/Private/UbaStorage.cpp
+++ b/Engine/Source/Programs/UnrealBuildAccelerator/Common/Private/UbaStorage.cpp
@@ -91,12 +91,12 @@ namespace uba
 		const ConfigTable& table = *tablePtr;
 		table.GetValueAsString(rootDir, TC("RootDir"));
 		table.GetValueAsBool(asyncUnmapViewOfFile, TC("AsyncUnmapViewOfFile"));
-		table.GetValueAsBool(useDeferredDeletes, TC("UseDeferredDeletes"));
 		table.GetValueAsU64(casCapacityBytes, TC("CasCapacityBytes"));
 		table.GetValueAsBool(keepTransientDataMapped, TC("KeepTransientDataMapped"));
 		table.GetValueAsBool(storeCompressed, TC("StoreCompressed"));
 		table.GetValueAsU32(fileEntryTableWipeThreshold, TC("FileEntryTableWipeThreshold"));
 		table.GetValueAsU32(maxParallelCopyOrLink, TC("MaxParallelCopyOrLink"));
+		table.GetValueAsU32(deleteCasBatchSize, TC("DeleteCasBatchSize"));
 
 		const tchar* compressor;
 		if (table.GetValueAsString(compressor, TC("Compressor")))
@@ -244,154 +244,98 @@ namespace uba
 		UBA_ASSERT(!m_oldestAccessed || !m_oldestAccessed->nextAccessed);
 
 		struct Rec { CasEntry& casEntry; u64 size; };
-		Vector<Rec> toDelete;
 
-		for (CasEntry* it = m_oldestAccessed; it;)
-		{
-			CasEntry& ce = *it;
-			ce.internalLock.Enter();
+		List<Vector<Rec>> toDelete;
 
-			if (ce.verified && !m_allowDeleteVerified) // Can't remove these since they might be in mapping tables and actively being used. It is tempting to delete dropped cas files but we can't
+		{
+			//TrackHintScope ths(*TrackWorkScope::GetActiveScope(), TCV("DetachCasEntries"));
+			for (CasEntry* it = m_oldestAccessed; it;)
 			{
-				ce.internalLock.Leave();
-				break;
-			}
+				CasEntry& ce = *it;
+				ce.internalLock.Enter();
 
-			if (ce.beingWritten || ce.refCount)
-			{
-				if (ce.beingWritten)
-					m_logger.Info(TC("Want to delete file that is begin written. We should never end up here"));
-				ce.internalLock.Leave();
-				it = ce.prevAccessed;
-				continue;
-			}
+				if (ce.verified && !m_allowDeleteVerified) // Can't remove these since they might be in mapping tables and actively being used. It is tempting to delete dropped cas files but we can't
+				{
+					ce.internalLock.Leave();
+					break;
+				}
 
-			UBA_ASSERT(ce.key != CasKeyZero);
-			toDelete.push_back({*it, ce.size});
+				if (ce.beingWritten || ce.refCount)
+				{
+					if (ce.beingWritten)
+						m_logger.Info(TC("Want to delete file that is begin written. We should never end up here"));
+					ce.internalLock.Leave();
+					it = ce.prevAccessed;
+					continue;
+				}
 
-			if (m_trackedDeletes)
-				m_trackedDeletes->insert(ce.key);
+				UBA_ASSERT(ce.key != CasKeyZero);
+				Vector<Rec>& deleteBatch = (toDelete.empty() || toDelete.back().size() == m_deleteCasBatchSize) ? toDelete.emplace_back() : toDelete.back();
+				deleteBatch.push_back({*it, ce.size});
 
-			m_casEvictedBytes += ce.size;
-			m_casEvictedCount++;
-			m_casTotalBytes -= ce.size;
+				if (m_trackedDeletes)
+					m_trackedDeletes->insert(ce.key);
 
-			ce.exists = false;
-			ce.size = 0;
+				m_casEvictedBytes += ce.size;
+				m_casEvictedCount++;
+				m_casTotalBytes -= ce.size;
 
+				ce.exists = false;
+				ce.size = 0;
 
-			it = ce.prevAccessed;
-			DetachEntry(ce);
-			if (m_casTotalBytes <= m_casCapacityBytes)
-				break;
-		}
 
+				it = ce.prevAccessed;
+				DetachEntry(ce);
+				if (m_casTotalBytes <= m_casCapacityBytes)
+					break;
+			}
+		}
 		if (m_casTotalBytes > m_casCapacityBytes)
 		{
 			m_overflowReported = true;
 			m_logger.Info(TC("Exceeding maximum size set for cas (%s). Current session needs more storage to be able to finish (will now overflow). User memory reported on session exit"), BytesToText(m_casCapacityBytes).str);
 		}
-
 		lock.Leave();
 
-		u64 startTime = GetTime();
-
-		const tchar* operationType = TC("delete");
-#if PLATFORM_WINDOWS
-		if (m_useDeferredDeletes)
+		for (auto& batch : toDelete)
 		{
-			// On windows, deletefile is stupid slow.. so we move the file somewhere first so we can delete it there instead
-			operationType = TC("move");
-			Vector<u32> toDelete2;
-			toDelete2.reserve(toDelete.size());
-			for (Rec& rec : toDelete)
-			{
-				StringBuffer<> casFile;
-				StorageImpl::GetCasFileName(casFile, rec.casEntry.key);
-
-				static Atomic<u32> deleteIndexCounter;
-				u32 deleteIndex = deleteIndexCounter++;
-
-				StringBuffer<> deleteName;
-				deleteName.Append(m_tempPath).AppendValue(deleteIndex).Append(".del");
-
-				if (!MoveFileW(casFile.data, deleteName.data))
+			m_workManager.AddWork([toDelete = std::move(batch), this](const WorkContext&)
 				{
-					u32 error = GetLastError();
-					if (error != ERROR_FILE_NOT_FOUND && error != ERROR_PATH_NOT_FOUND)
+					for (const Rec& rec : toDelete)
 					{
-						m_logger.Error(TC("Failed to delete %s while handling overflow (%s)"), casFile.data, LastErrorToText(error).data);
-						rec.casEntry.exists = true;
-						rec.casEntry.size = rec.size;
-						rec.casEntry.internalLock.Leave();
-						continue;
-					}
-				}
+						StringBuffer<> casFile;
+						StorageImpl::GetCasFileName(casFile, rec.casEntry.key);
 
-				toDelete2.push_back(deleteIndex);
-
-				//m_logger.Warning(TC("Evicted %s due to overflow"), casFile.data);
-
-				casEntry.verified = true; // Verified to be deleted
+						if (!DeleteFileW(casFile.data))
+						{
+							u32 error = GetLastError();
+							if (error != ERROR_FILE_NOT_FOUND && error != ERROR_PATH_NOT_FOUND)
+							{
+								m_logger.Error(TC("Failed to delete %s while handling overflow (%s)"), casFile.data, LastErrorToText(error).data);
+								rec.casEntry.exists = true;
+								rec.casEntry.size = rec.size;
+								rec.casEntry.internalLock.Leave();
+
+								// TODO: Should this instead set some overflow
+								/*
+								SCOPED_FUTEX(m_accessLock, lock);
+								m_casEvictedBytes -= rec.size;
+								--m_casEvictedCount;
+								m_casTotalBytes += rec.size;
+								AttachEntry(&rec.casEntry); // TODO: This should be re-added in the end
+								*/
+								continue;
+							}
+						}
 
-				rec.casEntry.internalLock.Leave();
-			}
+						//m_logger.Warning(TC("Evicted %s due to overflow"), casFile.data);
 
-			m_workManager.AddWork([toDelete2 = std::move(toDelete2), this](const WorkContext&)
-				{
-					u64 startTime = GetTime();
-					for (u32 index : toDelete2)
-					{
-						StringBuffer<> deleteName;
-						deleteName.Append(m_tempPath).AppendValue(index).Append(".del");
-						DeleteFileW(deleteName.data);
-					}
-					u64 time = GetTime() - startTime;
-					if (TimeToMs(time) > 30*1000)
-						m_logger.Info(TC("Took delete job %s to delete %llu files"), TimeToText(time).str, toDelete2.size());
-				}, 1, TC("DeleteCasFiles"));
-		}
-		else
-#endif
-		{
-			for (Rec& rec : toDelete)
-			{
-				StringBuffer<> casFile;
-				StorageImpl::GetCasFileName(casFile, rec.casEntry.key);
+						rec.casEntry.verified = true; // Verified to be deleted
 
-				if (!DeleteFileW(casFile.data))
-				{
-					u32 error = GetLastError();
-					if (error != ERROR_FILE_NOT_FOUND && error != ERROR_PATH_NOT_FOUND)
-					{
-						m_logger.Error(TC("Failed to delete %s while handling overflow (%s)"), casFile.data, LastErrorToText(error).data);
-						rec.casEntry.exists = true;
-						rec.casEntry.size = rec.size;
 						rec.casEntry.internalLock.Leave();
-
-						// TODO: Should this instead set some overflow
-						/*
-						SCOPED_FUTEX(m_accessLock, lock);
-						m_casEvictedBytes -= rec.size;
-						--m_casEvictedCount;
-						m_casTotalBytes += rec.size;
-						AttachEntry(&rec.casEntry); // TODO: This should be re-added in the end
-						*/
-						continue;
 					}
-				}
-
-				//m_logger.Warning(TC("Evicted %s due to overflow"), casFile.data);
-
-				casEntry.verified = true; // Verified to be deleted
-
-				rec.casEntry.internalLock.Leave();
-			}
+				}, 1, TC("DeleteCasFiles"));
 		}
-		u64 time = GetTime() - startTime;
-		if (TimeToMs(time) > 30*1000)
-			m_logger.Info(TC("Took %s to %s %llu files"), TimeToText(time).str, operationType, toDelete.size());
-
 	}
 
 	void StorageImpl::CasEntryDeleted(CasEntry& casEntry, u64 size)
@@ -986,11 +930,11 @@ namespace uba
 	{
 		m_casCapacityBytes = info.casCapacityBytes;
 		m_fileEntryTableWipeThreshold = info.fileEntryTableWipeThreshold;
+		m_deleteCasBatchSize = info.deleteCasBatchSize;
 		m_storeCompressed = info.storeCompressed;
 		m_manuallyHandleOverflow = info.manuallyHandleOverflow;
 		m_asyncUnmapViewOfFile = info.asyncUnmapViewOfFile;
 		m_allowDeleteVerified = info.allowDeleteVerified;
-		m_useDeferredDeletes = info.useDeferredDeletes;
 		m_writeToDisk = info.writeToDisk;
 		if (!m_writeToDisk)
 			m_casCapacityBytes = 0;
@@ -2032,8 +1976,11 @@ namespace uba
 		casEntry.verified = true;
 		deferredLock.Leave();
 		WriteResult res;
-		if (!WriteCasFile(res, deferredCreation.data, casKey))
-			return m_logger.Error(TC("Failed to write deferred cas %s (%s)"), CasKeyString(casKey).str, deferredCreation.data);
+		{
+			//TrackHintScope ths(*TrackWorkScope::GetActiveScope(), TCV("WriteCasFile"));
+			if (!WriteCasFile(res, deferredCreation.data, casKey))
+				return m_logger.Error(TC("Failed to write deferred cas %s (%s)"), CasKeyString(casKey).str, deferredCreation.data);
+		}
 
 		if (res.sourceStoredSize != deferredFileSize) // TODO: Should also check last written time
 			m_logger.Warning(TC("Deferred cas %s was deferred with source file size %llu but was read with source file size %llu (%s)"), CasKeyString(casKey).str, deferredFileSize, res.sourceStoredSize, deferredCreation.data);
diff --git a/Engine/Source/Programs/UnrealBuildAccelerator/Common/Private/UbaStorageServer.cpp b/Engine/Source/Programs/UnrealBuildAccelerator/Common/Private/UbaStorageServer.cpp
--- a/Engine/Source/Programs/UnrealBuildAccelerator/Common/Private/UbaStorageServer.cpp
+++ b/Engine/Source/Programs/UnrealBuildAccelerator/Common/Private/UbaStorageServer.cpp
@@ -587,13 +587,22 @@ namespace uba
 				casKey = reader.ReadCasKey();
 				reader.ReadString(hint);
 
-				workContext.tracker.AddHint(StringView(hint).GetFileName());
+				StringBuffer<16> temp;
+				StringView logHint = hint;
+				if (hint.Equals(TCV("KnownInput")))
+				{
+					temp.Append(CasKeyString(casKey).SubStr(0, 10));
+					logHint = temp;
+				}
+
+				workContext.tracker.AddHint(logHint.GetFileName());
 
 				casKey = AsCompressed(casKey, m_storeCompressed);
 
 				bool detailedTrace = m_traceFetch;
 				if (detailedTrace)
-					m_trace->FileFetchBegin(connectionInfo.GetId(), casKey, hint);
+					m_trace->FileFetchBegin(connectionInfo.GetId(), casKey, logHint);
+
 				auto traceFetchGuard = MakeGuard([&](){ if (detailedTrace) m_trace->FileFetchEnd(connectionInfo.GetId(), casKey); });
 
 
diff --git a/Engine/Source/Programs/UnrealBuildAccelerator/Common/Private/UbaWorkManager.cpp b/Engine/Source/Programs/UnrealBuildAccelerator/Common/Private/UbaWorkManager.cpp
--- a/Engine/Source/Programs/UnrealBuildAccelerator/Common/Private/UbaWorkManager.cpp
+++ b/Engine/Source/Programs/UnrealBuildAccelerator/Common/Private/UbaWorkManager.cpp
@@ -8,6 +8,42 @@
 
 namespace uba
 {
+	thread_local TrackWorkScope* t_trackWorkScope;
+
+	TrackWorkScope::TrackWorkScope()
+	:	tracker(nullptr)
+	,	id(0)
+	{
+		outerScope = t_trackWorkScope;
+		t_trackWorkScope = this;
+	}
+
+	TrackWorkScope::TrackWorkScope(WorkTracker& t, StringView desc, const Color& color)
+	:	tracker(&t)
+	,	id(t.TrackWorkStart(desc, color))
+	{
+		outerScope = t_trackWorkScope;
+		t_trackWorkScope = this;
+	}
+
+	void TrackWorkScope::AddHint(StringView hint, u64 startTime)
+	{
+		if (tracker)
+			tracker->TrackWorkHint(id, hint, startTime);
+	}
+
+	TrackWorkScope::~TrackWorkScope()
+	{
+		t_trackWorkScope = outerScope;
+		if (tracker)
+			tracker->TrackWorkEnd(id);
+	}
+
+	TrackWorkScope* TrackWorkScope::GetActiveScope()
+	{
+		return t_trackWorkScope;
+	}
+
 	u32 WorkManager::TrackWorkStart(const StringView& desc, const Color& color)
 	{
 		if (auto t = m_workTracker.load())
diff --git a/Engine/Source/Programs/UnrealBuildAccelerator/Common/Public/UbaStorage.h b/Engine/Source/Programs/UnrealBuildAccelerator/Common/Public/UbaStorage.h
--- a/Engine/Source/Programs/UnrealBuildAccelerator/Common/Public/UbaStorage.h
+++ b/Engine/Source/Programs/UnrealBuildAccelerator/Common/Public/UbaStorage.h
@@ -107,6 +107,7 @@ namespace uba
 		u64 casCapacityBytes = 20llu * 1024 * 1024 * 1024;
 		u32 fileEntryTableWipeThreshold = 1'000'000;
 		u32 maxParallelCopyOrLink = 1000;
+		u32 deleteCasBatchSize = 10;
 		bool storeCompressed = true;
 		bool manuallyHandleOverflow = false;
 		bool asyncUnmapViewOfFile = true;
@@ -298,12 +299,12 @@ namespace uba
 		u64 m_casDroppedBytes = 0;
 		u32 m_casDroppedCount = 0;
 		u32 m_fileEntryTableWipeThreshold = 0;
+		u32 m_deleteCasBatchSize;
 		bool m_overflowReported = false;
 		bool m_storeCompressed = false;
 		bool m_manuallyHandleOverflow = false;
 		bool m_asyncUnmapViewOfFile = true;
 		bool m_allowDeleteVerified = false;
-		bool m_useDeferredDeletes;
 		bool m_writeToDisk;
 
 		MutexHandle m_exclusiveMutex = InvalidMutexHandle;
diff --git a/Engine/Source/Programs/UnrealBuildAccelerator/Common/Public/UbaWorkManager.h b/Engine/Source/Programs/UnrealBuildAccelerator/Common/Public/UbaWorkManager.h
--- a/Engine/Source/Programs/UnrealBuildAccelerator/Common/Public/UbaWorkManager.h
+++ b/Engine/Source/Programs/UnrealBuildAccelerator/Common/Public/UbaWorkManager.h
@@ -87,11 +87,13 @@ namespace uba
 
 	struct TrackWorkScope
 	{
-		TrackWorkScope() : tracker(nullptr), id(0) {}
-		TrackWorkScope(WorkTracker& t, StringView desc, const Color& color = ColorWork) : tracker(&t), id(t.TrackWorkStart(desc, color)) {}
-		void AddHint(StringView hint, u64 startTime = 0) { if (tracker) tracker->TrackWorkHint(id, hint, startTime); }
-		~TrackWorkScope() { if (tracker) tracker->TrackWorkEnd(id); }
+		TrackWorkScope();
+		TrackWorkScope(WorkTracker& t, StringView desc, const Color& color = ColorWork);
+		void AddHint(StringView hint, u64 startTime = 0);
+		~TrackWorkScope();
+		static TrackWorkScope* GetActiveScope();
 		WorkTracker* tracker;
+		TrackWorkScope* outerScope;
 		u32 id;
 	};
 
diff --git a/Engine/Source/Programs/UnrealBuildAccelerator/Core/Public/UbaStringBuffer.h b/Engine/Source/Programs/UnrealBuildAccelerator/Core/Public/UbaStringBuffer.h
--- a/Engine/Source/Programs/UnrealBuildAccelerator/Core/Public/UbaStringBuffer.h
+++ b/Engine/Source/Programs/UnrealBuildAccelerator/Core/Public/UbaStringBuffer.h
@@ -130,7 +130,8 @@ namespace uba
 		bool Equals(const StringView& str, bool ignoreCase = true) const { return count == str.count && uba::Equals(data, str.data, count, ignoreCase); }
 		StringView GetFileName() const { const tchar* str = uba::GetFileName(data); return StringView(str, count - u32(str - data)); }
 		StringView GetPath() const;
-		StringView Skip(u32 skipCount) { return StringView(data + skipCount, count - skipCount); }
+		StringView Skip(u32 skipCount) const { return StringView(data + skipCount, count - skipCount); }
+		StringView SubStr(u32 startOffset, u32 length) const { return StringView(data + startOffset, length); }
 
 		TString ToString() const { return TString(data, data + count); }
 
