[UBA]
* Changed nameToHashLookupLock reader writer lock to futex to prevent write lock starvation on linux/mac. We saw a hang where all threads were sitting on this lock and on linux/mac lots of write locks can end up starving out each other causing a hang

diff --git a/Engine/Source/Programs/UnrealBuildAccelerator/Common/Private/UbaSessionServer.cpp b/Engine/Source/Programs/UnrealBuildAccelerator/Common/Private/UbaSessionServer.cpp
--- a/Engine/Source/Programs/UnrealBuildAccelerator/Common/Private/UbaSessionServer.cpp
+++ b/Engine/Source/Programs/UnrealBuildAccelerator/Common/Private/UbaSessionServer.cpp
@@ -389,7 +389,7 @@ namespace uba
 
 
 							// Update name to hash table
-							SCOPED_WRITE_LOCK(m_nameToHashLookupLock, lock);
+							SCOPED_FUTEX(m_nameToHashLookupLock, lock);
 							CasKey& lookupCasKey = m_nameToHashLookup[fileNameKey];
 							if (lookupCasKey != casKey)
 							{
@@ -1112,7 +1112,7 @@ namespace uba
 		u64 serverTime;
 		if (m_nameToHashInitialized && casKey != CasKeyIsDirectory)
 		{
-			SCOPED_WRITE_LOCK(m_nameToHashLookupLock, lock);
+			SCOPED_FUTEX(m_nameToHashLookupLock, lock);
 			serverTime = GetTime();
 			CasKey& lookupCasKey = m_nameToHashLookup[fileNameKey];
 			if (lookupCasKey != casKey)
@@ -1369,7 +1369,7 @@ namespace uba
 	{
 		u32 requestedSize = reader.ReadU32();
 
-		SCOPED_READ_LOCK(m_nameToHashLookupLock, lock);
+		SCOPED_FUTEX_READ(m_nameToHashLookupLock, lock);
 		if (requestedSize == ~0u)
 		{
 			requestedSize = u32(m_nameToHashTableMem.writtenSize);
@@ -1469,7 +1469,7 @@ namespace uba
 		u32 neededDirectoryTableSize = GetDirectoryTableSize();
 		u32 neededHashTableSize;
 		{
-			SCOPED_READ_LOCK(m_nameToHashLookupLock, l);
+			SCOPED_FUTEX_READ(m_nameToHashLookupLock, l);
 			neededHashTableSize = u32(m_nameToHashTableMem.writtenSize);
 		}
 
@@ -2187,7 +2187,7 @@ namespace uba
 
 	void SessionServer::FileEntryAdded(StringKey fileNameKey, u64 lastWritten, u64 size)
 	{
-		SCOPED_WRITE_LOCK(m_nameToHashLookupLock, lock);
+		SCOPED_FUTEX(m_nameToHashLookupLock, lock);
 
 		if (!m_nameToHashInitialized)
 			return;
@@ -2343,7 +2343,7 @@ namespace uba
 		if (!m_nameToHashTableEnabled || m_nameToHashInitialized)
 			return true;
 
-		SCOPED_WRITE_LOCK(m_nameToHashLookupLock, lock);
+		SCOPED_FUTEX(m_nameToHashLookupLock, lock);
 		m_nameToHashTableMem.Init(NameToHashMemSize);
 		m_nameToHashInitialized = true;
 		lock.Leave();
@@ -2378,7 +2378,7 @@ namespace uba
 				}
 			}
 		}
-		SCOPED_WRITE_LOCK(m_nameToHashLookupLock, lock2);
+		SCOPED_FUTEX(m_nameToHashLookupLock, lock2);
 		u64 entryCount = m_nameToHashLookup.size();
 		lock2.Leave();
 
@@ -2460,7 +2460,7 @@ namespace uba
 					LogLine(logStr);
 				}
 				{
-					SCOPED_READ_LOCK(m_nameToHashLookupLock, hlock);
+					SCOPED_FUTEX_READ(m_nameToHashLookupLock, hlock);
 					auto findIt = m_nameToHashLookup.find(fileNameKey);
 					if (findIt != m_nameToHashLookup.end())
 						logStr.Clear().Appendf(TC("File %s found in name-to-hash lookup. CasKey is %s"), path.data, CasKeyString(findIt->second).str);
diff --git a/Engine/Source/Programs/UnrealBuildAccelerator/Common/Public/UbaSessionServer.h b/Engine/Source/Programs/UnrealBuildAccelerator/Common/Public/UbaSessionServer.h
--- a/Engine/Source/Programs/UnrealBuildAccelerator/Common/Public/UbaSessionServer.h
+++ b/Engine/Source/Programs/UnrealBuildAccelerator/Common/Public/UbaSessionServer.h
@@ -194,7 +194,7 @@ namespace uba
 		UnorderedMap<StringKey, CustomCasKey> m_customCasKeys;
 
 		UnorderedMap<StringKey, CasKey> m_nameToHashLookup;
-		ReaderWriterLock m_nameToHashLookupLock;
+		Futex m_nameToHashLookupLock;
 		Atomic<bool> m_nameToHashInitialized;
 
 		ReaderWriterLock m_receivedFilesLock;
