[UBA]
* Added option allowDeferredDeletes and disabled it on cache service and on non-windows (to see if a farm issue goes away)
* Improved some error logging

diff --git a/Engine/Source/Programs/UnrealBuildAccelerator/CacheService/Private/UbaCacheService.cpp b/Engine/Source/Programs/UnrealBuildAccelerator/CacheService/Private/UbaCacheService.cpp
--- a/Engine/Source/Programs/UnrealBuildAccelerator/CacheService/Private/UbaCacheService.cpp
+++ b/Engine/Source/Programs/UnrealBuildAccelerator/CacheService/Private/UbaCacheService.cpp
@@ -513,6 +513,7 @@ namespace uba
 		storageInfo.allowHintAsFallback = false;
 		storageInfo.writeReceivedCasFilesToDisk = true;
 		storageInfo.allowDeleteVerified = true;
+		storageInfo.allowDeferredDeletes = false;
 		StorageServer storageServer(storageInfo);
 
 		Trace trace(logger.m_writer);
diff --git a/Engine/Source/Programs/UnrealBuildAccelerator/Common/Private/UbaSession.cpp b/Engine/Source/Programs/UnrealBuildAccelerator/Common/Private/UbaSession.cpp
--- a/Engine/Source/Programs/UnrealBuildAccelerator/Common/Private/UbaSession.cpp
+++ b/Engine/Source/Programs/UnrealBuildAccelerator/Common/Private/UbaSession.cpp
@@ -4002,7 +4002,7 @@ namespace uba
 				if (fa.OpenMemoryRead(0, false))
 					return func(fa.GetData(), fa.GetSize());
 				else
-					return m_logger.Warning(TC("Failed to open %s"), out.fileName.data);
+					return m_logger.Warning(TC("Failed to open %s (%s)"), out.fileName.data, LastErrorToText().data);
 			};
 
 		auto DevirtualizePathFunc = [this, rootsHandle = startInfo.rootsHandle](StringBufferBase& inOut) { return DevirtualizePath(inOut, rootsHandle, false); };
diff --git a/Engine/Source/Programs/UnrealBuildAccelerator/Common/Private/UbaStorage.cpp b/Engine/Source/Programs/UnrealBuildAccelerator/Common/Private/UbaStorage.cpp
--- a/Engine/Source/Programs/UnrealBuildAccelerator/Common/Private/UbaStorage.cpp
+++ b/Engine/Source/Programs/UnrealBuildAccelerator/Common/Private/UbaStorage.cpp
@@ -94,6 +94,7 @@ namespace uba
 		table.GetValueAsU64(casCapacityBytes, TC("CasCapacityBytes"));
 		table.GetValueAsBool(keepTransientDataMapped, TC("KeepTransientDataMapped"));
 		table.GetValueAsBool(storeCompressed, TC("StoreCompressed"));
+		table.GetValueAsBool(allowDeferredDeletes, TC("AllowDeferredDeletes"));
 		table.GetValueAsU32(fileEntryTableWipeThreshold, TC("FileEntryTableWipeThreshold"));
 		table.GetValueAsU32(maxParallelCopyOrLink, TC("MaxParallelCopyOrLink"));
 		table.GetValueAsU32(deleteCasBatchSize, TC("DeleteCasBatchSize"));
@@ -297,44 +298,58 @@ namespace uba
 		}
 		lock.Leave();
 
-		for (auto& batch : toDelete)
-		{
-			m_workManager.AddWork([toDelete = std::move(batch), this](const WorkContext&)
+		static auto deleteBatch = [](StorageImpl& storage, const Vector<Rec>& toDelete)
+			{
+				for (const Rec& rec : toDelete)
 				{
-					for (const Rec& rec : toDelete)
-					{
-						StringBuffer<> casFile;
-						StorageImpl::GetCasFileName(casFile, rec.casEntry.key);
+					StringBuffer<> casFile;
+					storage.StorageImpl::GetCasFileName(casFile, rec.casEntry.key);
 
-						if (!DeleteFileW(casFile.data))
+					if (!DeleteFileW(casFile.data))
+					{
+						u32 error = GetLastError();
+						if (error != ERROR_FILE_NOT_FOUND && error != ERROR_PATH_NOT_FOUND)
 						{
-							u32 error = GetLastError();
-							if (error != ERROR_FILE_NOT_FOUND && error != ERROR_PATH_NOT_FOUND)
-							{
-								m_logger.Error(TC("Failed to delete %s while handling overflow (%s)"), casFile.data, LastErrorToText(error).data);
-								rec.casEntry.exists = true;
-								rec.casEntry.size = rec.size;
-								rec.casEntry.internalLock.Leave();
-
-								// TODO: Should this instead set some overflow
-								/*
-								SCOPED_FUTEX(m_accessLock, lock);
-								m_casEvictedBytes -= rec.size;
-								--m_casEvictedCount;
-								m_casTotalBytes += rec.size;
-								AttachEntry(&rec.casEntry); // TODO: This should be re-added in the end
-								*/
-								continue;
-							}
+							storage.m_logger.Error(TC("Failed to delete %s while handling overflow (%s)"), casFile.data, LastErrorToText(error).data);
+							rec.casEntry.exists = true;
+							rec.casEntry.size = rec.size;
+							rec.casEntry.internalLock.Leave();
+
+							// TODO: Should this instead set some overflow
+							/*
+							SCOPED_FUTEX(m_accessLock, lock);
+							m_casEvictedBytes -= rec.size;
+							--m_casEvictedCount;
+							m_casTotalBytes += rec.size;
+							AttachEntry(&rec.casEntry); // TODO: This should be re-added in the end
+							*/
+							continue;
 						}
+					}
 
-						//m_logger.Warning(TC("Evicted %s due to overflow"), casFile.data);
+					//m_logger.Warning(TC("Evicted %s due to overflow"), casFile.data);
 
-						rec.casEntry.verified = true; // Verified to be deleted
+					rec.casEntry.verified = true; // Verified to be deleted
 
-						rec.casEntry.internalLock.Leave();
-					}
-				}, 1, TC("DeleteCasFiles"));
+					rec.casEntry.internalLock.Leave();
+				}
+			};
+
+		// TODO: Defer deletes can cause deadlocks if all active jobs are sitting on locks of the cas that is being deleted
+		// .. since they are now in a locked state but queued... need to make a solution for this
+		bool deferDeletes = m_allowDeferredDeletes;
+
+		if (deferDeletes)
+		{
+			for (auto& batch : toDelete)
+				m_workManager.AddWork([toDelete = std::move(batch), this](const WorkContext&) { deleteBatch(*this, toDelete); }, 1, TC("DeleteCasFiles"));
+		}
+		else
+		{
+			m_workManager.ParallelFor(u32(toDelete.size()), toDelete, [this](const WorkContext&, auto& it)
+				{
+					deleteBatch(*this, *it);
+				}, TCV("DeleteCasFiles"));
 		}
 	}
 
@@ -935,6 +950,7 @@ namespace uba
 		m_manuallyHandleOverflow = info.manuallyHandleOverflow;
 		m_asyncUnmapViewOfFile = info.asyncUnmapViewOfFile;
 		m_allowDeleteVerified = info.allowDeleteVerified;
+		m_allowDeferredDeletes = info.allowDeferredDeletes;
 		m_writeToDisk = info.writeToDisk;
 		if (!m_writeToDisk)
 			m_casCapacityBytes = 0;
diff --git a/Engine/Source/Programs/UnrealBuildAccelerator/Common/Public/UbaStorage.h b/Engine/Source/Programs/UnrealBuildAccelerator/Common/Public/UbaStorage.h
--- a/Engine/Source/Programs/UnrealBuildAccelerator/Common/Public/UbaStorage.h
+++ b/Engine/Source/Programs/UnrealBuildAccelerator/Common/Public/UbaStorage.h
@@ -114,7 +114,7 @@ namespace uba
 		bool writeToDisk = true;
 		bool keepTransientDataMapped = true;	// Will keep the transient data mapped in instead of doing map/unmap everytime we access it
 		bool allowDeleteVerified = false;
-		bool useDeferredDeletes = true;
+		bool allowDeferredDeletes = true;
 		MutexHandle exclusiveMutex = InvalidMutexHandle; // Set this to hand over pre-exclusive access to storage. Storage will release mutex on shutdown
 		u8 casCompressor = 0;
 		u8 casCompressionLevel = 0;
@@ -305,6 +305,7 @@ namespace uba
 		bool m_manuallyHandleOverflow = false;
 		bool m_asyncUnmapViewOfFile = true;
 		bool m_allowDeleteVerified = false;
+		bool m_allowDeferredDeletes = true;
 		bool m_writeToDisk;
 
 		MutexHandle m_exclusiveMutex = InvalidMutexHandle;
