[UBA]
* Changed so crawler always bottlenecks work.. and never more than work manager work thread count - 2  (we need to make sure there is always room for garbage collecting of cas entries)

diff --git a/Engine/Source/Programs/UnrealBuildAccelerator/Common/Private/UbaSession.cpp b/Engine/Source/Programs/UnrealBuildAccelerator/Common/Private/UbaSession.cpp
--- a/Engine/Source/Programs/UnrealBuildAccelerator/Common/Private/UbaSession.cpp
+++ b/Engine/Source/Programs/UnrealBuildAccelerator/Common/Private/UbaSession.cpp
@@ -1536,8 +1536,10 @@ namespace uba
 		return result;
 	}
 
-	u32 GetCrawlerBottleneck(int linkCrawlerMaxParallel)
+	u32 GetCrawlerBottleneck(bool runningRemote, int linkCrawlerMaxParallel)
 	{
+		if (runningRemote)
+			return ~0u;
 		if (linkCrawlerMaxParallel <= 0)
 			return u32(Max(0, int(GetLogicalProcessorCount()) + linkCrawlerMaxParallel));
 		else
@@ -1557,7 +1559,7 @@ namespace uba
 	,	m_writeFilesBottleneck(info.writeFilesBottleneck)
 	,	m_writeFilesFileMapMax(u64(info.writeFilesFileMapMaxMb)*1024*1024)
 	,	m_writeFilesNoBufferingMin(u64(info.writeFilesNoBufferingMinMb)*1024*1024)
-	,	m_dependencyCrawler(m_logger, workManager, GetCrawlerBottleneck(info.linkCrawlerMaxParallel))
+	,	m_dependencyCrawler(m_logger, workManager, GetCrawlerBottleneck(runningRemote, info.linkCrawlerMaxParallel))
 	{
 		UBA_ASSERTF(info.rootDir && *info.rootDir, TC("No root dir set when creating session"));
 		m_rootDir.count = GetFullPathNameW(info.rootDir, m_rootDir.capacity, m_rootDir.data, NULL);
diff --git a/Engine/Source/Programs/UnrealBuildAccelerator/Common/Public/UbaDependencyCrawler.h b/Engine/Source/Programs/UnrealBuildAccelerator/Common/Public/UbaDependencyCrawler.h
--- a/Engine/Source/Programs/UnrealBuildAccelerator/Common/Public/UbaDependencyCrawler.h
+++ b/Engine/Source/Programs/UnrealBuildAccelerator/Common/Public/UbaDependencyCrawler.h
@@ -28,7 +28,6 @@ namespace uba
 	{
 	public:
 		inline DependencyCrawler(Logger& logger, WorkManager& workManager, u32 bottleneckMax = ~0u);
-		inline ~DependencyCrawler();
 
 		using FileExistsFunc = Function<bool(const StringView& fileName, u32& outAttr)>;
 		using FileFunc = Function<void(const StringView& file, bool isDirectory)>;
@@ -97,7 +96,7 @@ namespace uba
 		FileExistsFunc m_fileExistsFunc;
 		TraverseFilesFunc m_traverseFilesFunc;
 
-		Bottleneck* m_linkerDecompressBottleneck = nullptr;
+		Bottleneck m_bottleneck;
 
 		struct CodeFile { TString path; bool hasPch = false; };
 		using CodeFiles = List<CodeFile>;
@@ -109,20 +108,23 @@ namespace uba
 		inline bool ParseCodeFile2(TrackWorkScope& tracker, Instance& instance, HandledFiles& handledFiles, const void* data, u64 dataSize, const StringView& codeFile, bool parseDefines);
 		inline bool HandleInclude(TrackWorkScope& tracker, Instance& instance, HandledFiles& handledFiles, const StringView& rootPath, const StringView& include, const StringView& codeFile, bool parseDefines);
 		inline bool TraverseInclude(TrackWorkScope& tracker, Instance& instance, const StringView& rootPathWithSlash);
-	};
 
+		inline u32 GetBottleneckMax(WorkManager& workManager, u32 bottleneckMax)
+		{
+			// We never allow the crawler to use all workers since there could be jobs that do garbage collecting of cas entries 
+			u32 max = workManager.GetWorkerCount();
+			UBA_ASSERT(max > 2);
+			if (max <= 2)
+				return 0;
+			return Min(max - 2u, bottleneckMax);
+		}
+	};
 
 	DependencyCrawler::DependencyCrawler(Logger& logger, WorkManager& workManager, u32 bottleneckMax)
 	:	m_logger(logger)
 	,	m_workManager(workManager)
+	,	m_bottleneck(GetBottleneckMax(workManager, bottleneckMax))
 	{
-		if (bottleneckMax != ~0u)
-			m_linkerDecompressBottleneck = new Bottleneck(bottleneckMax);
-	}
-
-	DependencyCrawler::~DependencyCrawler()
-	{
-		delete m_linkerDecompressBottleneck;
 	}
 
 	void DependencyCrawler::Init(FileExistsFunc&& fileExistsFunc, TraverseFilesFunc&& traverseFilesFunc, bool useBloomFilter)
@@ -193,7 +195,7 @@ namespace uba
 							{
 								bool stopCrawling = false;
 								ref.instance.createFileFunc(context.tracker, codeFile, {}, stopCrawling);
-							}, 1, TC("CrawlForDecomp"), ColorWork, false, m_linkerDecompressBottleneck);
+							}, 1, TC("CrawlForDecomp"), ColorWork, false, &m_bottleneck);
 					}
 					else
 					{
@@ -229,7 +231,7 @@ namespace uba
 					}
 				}
 
-			}, 1, TC("CrawlRsp"), ColorWork);
+			}, 1, TC("CrawlRsp"), ColorWork, false, &m_bottleneck);
 		return true;
 	}
 
@@ -809,7 +811,7 @@ namespace uba
 						if (parseDefines)
 							work({tracker});
 						else
-							m_workManager.AddWork(work, 1, TC("CrawlIncludes"), ColorWork);
+							m_workManager.AddWork(work, 1, TC("CrawlIncludes"), ColorWork, false, &m_bottleneck);
 					}
 				}
 				else if (parseDefines && strncmp(it, "define", 6) == 0)
@@ -905,7 +907,7 @@ namespace uba
 					m_workManager.AddWork([this, ref = InstanceRef(instance), codeFile = codeFilePath.ToString()](const WorkContext& context)
 						{
 							TraverseInclude(context.tracker, ref.instance, codeFile);
-						}, 1, TC("CrawlBiDir"), ColorWork);
+						}, 1, TC("CrawlBiDir"), ColorWork, false, &m_bottleneck);
 				}
 				else
 				{
@@ -920,7 +922,7 @@ namespace uba
 							if (!ref.instance.createFileFunc(context.tracker, codeFile, {}, stopCrawling))
 								if (!stopCrawling)
 									m_logger.Warning(TC("Failed to open file %s from builtindir"), codeFile.c_str());
-						}, 1, TC("CrawlBiFile"), ColorWork);
+						}, 1, TC("CrawlBiFile"), ColorWork, false, &m_bottleneck);
 				}
 			});
 		return true;
