[UBA]
* Added support for sending in bottleneck to parallel for

diff --git a/Engine/Source/Programs/UnrealBuildAccelerator/Common/Public/UbaWorkManager.h b/Engine/Source/Programs/UnrealBuildAccelerator/Common/Public/UbaWorkManager.h
--- a/Engine/Source/Programs/UnrealBuildAccelerator/Common/Public/UbaWorkManager.h
+++ b/Engine/Source/Programs/UnrealBuildAccelerator/Common/Public/UbaWorkManager.h
@@ -51,7 +51,7 @@ namespace uba
 		WorkTracker* GetWorkTracker() { return m_workTracker; }
 
 		template<u32 BatchSize = 1, typename TContainer, typename TFunc>
-		void ParallelFor(u32 workCount, TContainer& container, TFunc&& func, const StringView& description = AsView(TC("")), bool highPriority = false);
+		void ParallelFor(u32 workCount, TContainer& container, TFunc&& func, const StringView& description = AsView(TC("")), bool highPriority = false, Bottleneck* bottleneck = nullptr);
 
 	protected:
 		Atomic<WorkTracker*> m_workTracker;
@@ -107,7 +107,7 @@ namespace uba
 	};
 
 	template<u32 BatchSize, typename TContainer, typename TFunc>
-	void WorkManager::ParallelFor(u32 workCount, TContainer& container, TFunc&& func, const StringView& description, bool highPriority)
+	void WorkManager::ParallelFor(u32 workCount, TContainer& container, TFunc&& func, const StringView& description, bool highPriority, Bottleneck* bottleneck)
 	{
 		#if !defined( __clang_analyzer__ ) // Static analyzer claims context can leak but it can only leak when process terminates (and then it doesn't matter)
 
@@ -193,7 +193,7 @@ namespace uba
 				}
 			};
 		
-		AddWork(work, workCount, description.data, ColorWork, highPriority);
+		AddWork(work, workCount, description.data, ColorWork, highPriority, bottleneck);
 
 		{
 			#if UBA_TRACK_WORK
