[UBA]
* Don't compress command line in CreateProcess message unless it is very long (CommunicationMemSize - <value>, which is like 60kb)

diff --git a/Engine/Source/Programs/UnrealBuildAccelerator/Detours/Private/PosixOS/UbaDetoursFunctionsPosix.cpp b/Engine/Source/Programs/UnrealBuildAccelerator/Detours/Private/PosixOS/UbaDetoursFunctionsPosix.cpp
--- a/Engine/Source/Programs/UnrealBuildAccelerator/Detours/Private/PosixOS/UbaDetoursFunctionsPosix.cpp
+++ b/Engine/Source/Programs/UnrealBuildAccelerator/Detours/Private/PosixOS/UbaDetoursFunctionsPosix.cpp
@@ -2186,7 +2186,7 @@ int shared_posix_spawn(pid_t* pid, const char* path, const posix_spawn_file_acti
 		BinaryWriter writer;
 		writer.WriteByte(MessageType_CreateProcess);
 		writer.WriteString(path); // Application
-		writer.WriteLongString(cmdLineWithoutApplication);
+		writer.WriteLongString(cmdLineWithoutApplication, CommunicationMemSize - 4096);
 		writer.WriteString(pwd); // Current dir
 		writer.WriteBool(false); // Start suspended
 		writer.WriteBool(true); // Is child
diff --git a/Engine/Source/Programs/UnrealBuildAccelerator/Detours/Private/Windows/UbaDetoursFunctionsKernelBase.inl b/Engine/Source/Programs/UnrealBuildAccelerator/Detours/Private/Windows/UbaDetoursFunctionsKernelBase.inl
--- a/Engine/Source/Programs/UnrealBuildAccelerator/Detours/Private/Windows/UbaDetoursFunctionsKernelBase.inl
+++ b/Engine/Source/Programs/UnrealBuildAccelerator/Detours/Private/Windows/UbaDetoursFunctionsKernelBase.inl
@@ -3226,7 +3226,7 @@ BOOL Detoured_CreateProcessW(LPCWSTR lpApplicationName, LPWSTR lpCommandLine, LP
 	{
 		RPC_MESSAGE(CreateProcess, createProcess)
 		writer.WriteString(application);
-		writer.WriteLongString(commandLineWithoutApplication);
+		writer.WriteLongString(commandLineWithoutApplication, CommunicationMemSize - 4096); // don't compress if there is room in the writer memory
 		writer.WriteString(workingDir);
 		writer.WriteBool(startSuspended);
 		writer.WriteBool(isChild);
