[UBA]
* Added WaitForSingleObject after TerminateProcess with the hopes this will speed up freeing of memory on agent side when processes are terminated because of memory pressure

diff --git a/Engine/Source/Programs/UnrealBuildAccelerator/Common/Private/UbaProcess.cpp b/Engine/Source/Programs/UnrealBuildAccelerator/Common/Private/UbaProcess.cpp
--- a/Engine/Source/Programs/UnrealBuildAccelerator/Common/Private/UbaProcess.cpp
+++ b/Engine/Source/Programs/UnrealBuildAccelerator/Common/Private/UbaProcess.cpp
@@ -2174,8 +2174,8 @@ namespace uba
 		SCOPED_FUTEX(m_initLock, lock);
 		Logger& logger = m_session.m_logger;
 
-		ProcHandle handle = m_nativeProcessHandle;
-		if (handle == InvalidProcHandle)
+		ProcHandle procHandle = m_nativeProcessHandle;
+		if (procHandle == InvalidProcHandle)
 			return ~0u;
 
 		if (m_parentProcess)
@@ -2183,28 +2183,31 @@ namespace uba
 		m_nativeProcessHandle = InvalidProcHandle;
 
 #if PLATFORM_WINDOWS
-
+		HANDLE handle = (HANDLE)procHandle;
 		auto closeHandleGuard = MakeGuard([&]()
 			{
 				IO_COUNTERS ioCounters;
-				if (GetProcessIoCounters((HANDLE)handle, &ioCounters))
+				if (GetProcessIoCounters(handle, &ioCounters))
 				{
 					m_processStats.iopsRead = ioCounters.ReadOperationCount;
 					m_processStats.iopsWrite = ioCounters.WriteOperationCount;
 					m_processStats.iopsOther = ioCounters.OtherOperationCount;
 				}
 
-				CloseHandle((HANDLE)handle);
+				CloseHandle(handle);
 			});
 
 		bool hadTimeout = false;
 		if (cancel)
-			TerminateProcess((HANDLE)handle, ProcessCancelExitCode);
+		{
+			TerminateProcess(handle, ProcessCancelExitCode);
+			WaitForSingleObject(handle, 1000); // On windows this should return in milliseconds. On wine it might take a little while. We need to wait to make sure process is reaped
+		}
 		else
 		{	
 			while (true)
 			{
-				DWORD res = WaitForSingleObject((HANDLE)handle, 120 * 1000);
+				DWORD res = WaitForSingleObject(handle, 120 * 1000);
 				if (res == WAIT_OBJECT_0)
 				{
 					break;
@@ -2217,7 +2220,7 @@ namespace uba
 						hadTimeout = true;
 						const tchar* gotMessage = m_gotExitMessage ? TC("Got") : TC("Did not get");
 						logger.Info(TC("WaitForSingleObject timed out after 120 seconds waiting for process %s to exit (Exit code %u, %s ExitMessage and wrote %u files. Runtime: %s). Will terminate and wait again"), m_startInfo.GetDescription(), m_nativeProcessExitCode, gotMessage, u32(m_shared.writtenFiles.size()), TimeToText(GetTime() - m_startTime).str);
-						TerminateProcess((HANDLE)handle, m_nativeProcessExitCode);
+						TerminateProcess(handle, m_nativeProcessExitCode);
 						continue;
 					}
 					logger.Error(TC("WaitForSingleObject failed while waiting for process %s to exit even after terminating it (%s)"), m_startInfo.GetDescription(), LastErrorToText().data);
@@ -2226,7 +2229,8 @@ namespace uba
 					logger.Error(TC("WaitForSingleObject failed while waiting for process to exit (%s)"), LastErrorToText().data);
 				else if (res == WAIT_ABANDONED)
 					logger.Error(TC("Abandoned, this should never happen"));
-				TerminateProcess((HANDLE)handle, UBA_EXIT_CODE(13));
+				TerminateProcess(handle, UBA_EXIT_CODE(13));
+				WaitForSingleObject(handle, 1000);
 				return UBA_EXIT_CODE(13);
 			}
 		}
@@ -2235,7 +2239,7 @@ namespace uba
 		if (!hadTimeout)
 		{
 			DWORD nativeExitCode = 0;
-			res = GetExitCodeProcess((HANDLE)handle, (DWORD*)&nativeExitCode);
+			res = GetExitCodeProcess(handle, (DWORD*)&nativeExitCode);
 			if (!res && GetLastError() == ERROR_INVALID_HANDLE) // Was already terminated
 				return ~0u;
 			if (m_gotExitMessage || !m_detourEnabled)
