[UBA]
* Added so when iterating through to kill processes, if Cancel returns falls we search for next one to cancel
* Added SetMaxIdleTime  to be able to change session max idle time before disconnecting
* Changed so agent that has one or more other agents connected to its proxy will have twice as long idle timeout time

diff --git a/Engine/Source/Programs/UnrealBuildAccelerator/Agent/Private/UbaAgent.cpp b/Engine/Source/Programs/UnrealBuildAccelerator/Agent/Private/UbaAgent.cpp
--- a/Engine/Source/Programs/UnrealBuildAccelerator/Agent/Private/UbaAgent.cpp
+++ b/Engine/Source/Programs/UnrealBuildAccelerator/Agent/Private/UbaAgent.cpp
@@ -1469,6 +1469,9 @@ namespace uba
 					u32 clientCount = proxyServer->GetClientCount();
 					if (clientCount > 1) // when having a proxy the agent itself is always connected to it
 					{
+						if (maxIdleSeconds < 30) // A bit hacky solution.. but we really want to try preventing idling out the proxy if possible when other agents are connected
+							sessionClient->SetMaxIdleTime(maxIdleSeconds*2);
+
 						u32 processToFree = (clientCount - 1)/3 + 1; // Always free one, and then one per three helpers.. (so 16 helpers connected will remove 6 process count)
 						u32 newProcessCount;
 						if (processToFree < maxProcessCount)
@@ -1482,6 +1485,11 @@ namespace uba
 							sessionClient->SetMaxProcessCount(newProcessCount);
 						}
 					}
+					else
+					{
+						// Restore idle time if reaching 0 connections.
+						sessionClient->SetMaxIdleTime(maxIdleSeconds);
+					}
 				}
 				// This is an estimation based on tcp limitations (ack and sliding windows).
 				// For every 15ms latency on "best ping") we increase targetConnectionCount up to maxConnectionCount
diff --git a/Engine/Source/Programs/UnrealBuildAccelerator/Common/Private/UbaSessionClient.cpp b/Engine/Source/Programs/UnrealBuildAccelerator/Common/Private/UbaSessionClient.cpp
--- a/Engine/Source/Programs/UnrealBuildAccelerator/Common/Private/UbaSessionClient.cpp
+++ b/Engine/Source/Programs/UnrealBuildAccelerator/Common/Private/UbaSessionClient.cpp
@@ -212,6 +212,11 @@ namespace uba
 		m_maxProcessCount = count;
 	}
 
+	void SessionClient::SetMaxIdleTime(u32 seconds)
+	{
+		m_maxIdleSeconds = seconds;
+	}
+
 	void SessionClient::SetAllowSpawn(bool allow)
 	{
 		m_allowSpawn = allow;
@@ -1919,7 +1924,7 @@ namespace uba
 				u32 idleTime = u32(TimeToS(GetTime() - idleStartTime));
 				if (idleTime > m_maxIdleSeconds)
 				{
-					m_logger.Info(TC("Session has been idle longer than max idle time (%u seconds). Disconnecting (Did %u process requests during idle)"), m_maxIdleSeconds, processRequestCount);
+					m_logger.Info(TC("Session has been idle longer than max idle time (%u seconds). Disconnecting (Did %u process requests during idle)"), m_maxIdleSeconds.load(), processRequestCount);
 					SendNotification(AsView(TC("Idle time timeout")));
 
 					m_waitToSendEvent.Set();
@@ -1951,7 +1956,8 @@ namespace uba
 					SCOPED_FUTEX(rec.lock, lock);
 					if (rec.isDone)
 						continue;
-					rec.handle.Cancel();
+					if (!rec.handle.Cancel())
+						continue;
 					rec.isKilled = true;
 					SendReturnProcess(rec.handle.GetId(), TC("Running out of memory"));
 					++m_stats.killCount;
diff --git a/Engine/Source/Programs/UnrealBuildAccelerator/Common/Public/UbaSessionClient.h b/Engine/Source/Programs/UnrealBuildAccelerator/Common/Public/UbaSessionClient.h
--- a/Engine/Source/Programs/UnrealBuildAccelerator/Common/Public/UbaSessionClient.h
+++ b/Engine/Source/Programs/UnrealBuildAccelerator/Common/Public/UbaSessionClient.h
@@ -45,6 +45,7 @@ namespace uba
 		void SendSummary(const Function<void(Logger&)>& extraInfo);
 		void SetIsTerminating(const tchar* reason = TC("Terminating"), u64 delayMs = 0); // Session stores pointer directly. Can't be temporary
 		void SetMaxProcessCount(u32 count);
+		void SetMaxIdleTime(u32 seconds);
 		void SetAllowSpawn(bool allow);
 
 		u64 GetBestPing();
@@ -120,7 +121,7 @@ namespace uba
 		u32 m_sessionId = 0;
 		u32 m_uiLanguage = 0;
 		u32 m_defaultPriorityClass;
-		u32 m_maxIdleSeconds = ~0u;
+		Atomic<u32> m_maxIdleSeconds = ~0u;
 		u32 m_osVersion = 0;
 		u32 m_killRandomIndex = ~0u;
 		u32 m_killRandomCounter = 0;
