[UBA]
* Fixed shadowed variable warning on macos

diff --git a/Engine/Source/Programs/UnrealBuildAccelerator/Common/Private/UbaThread.cpp b/Engine/Source/Programs/UnrealBuildAccelerator/Common/Private/UbaThread.cpp
--- a/Engine/Source/Programs/UnrealBuildAccelerator/Common/Private/UbaThread.cpp
+++ b/Engine/Source/Programs/UnrealBuildAccelerator/Common/Private/UbaThread.cpp
@@ -458,12 +458,13 @@ namespace uba
 			return false;
 		auto tsg = MakeGuard([&] { vm_deallocate(mach_task_self(), (vm_address_t)threads, thread_count * sizeof(thread_t)); });
 
-		for (int i = 0; i < thread_count; i++)
+		for (int threadIt = 0; threadIt < thread_count; threadIt++)
 		{
-			if (threads[i] == mach_thread_self())
+			auto thread = threads[threadIt];
+
+			if (thread == mach_thread_self())
 				continue;
 
-			auto thread = threads[i];
 			thread_suspend(thread);
 			auto g = MakeGuard([&] { thread_resume(thread); });
 
