[UBA]
* Added commented out code for testing having tcp send logic send multiple segments of data at the same time to reduce kernel calls

diff --git a/Engine/Source/Programs/UnrealBuildAccelerator/Common/Private/UbaNetworkBackendTcp.cpp b/Engine/Source/Programs/UnrealBuildAccelerator/Common/Private/UbaNetworkBackendTcp.cpp
--- a/Engine/Source/Programs/UnrealBuildAccelerator/Common/Private/UbaNetworkBackendTcp.cpp
+++ b/Engine/Source/Programs/UnrealBuildAccelerator/Common/Private/UbaNetworkBackendTcp.cpp
@@ -51,6 +51,8 @@
 #define UBA_USE_OVERLAPPED_SEND_WITH_LOCK 0
 #define UBA_USE_IOCP PLATFORM_WINDOWS
 
+#define UBA_EXPERIMENTAL_WRITE_QUEUE 0
+
 namespace uba
 {
 	constexpr u32 MaxHeaderSize = 33;
@@ -84,6 +86,15 @@ namespace uba
 
 	////////////////////////////////////////////////////////////////////////////////////////////////////
 
+	#if UBA_EXPERIMENTAL_WRITE_QUEUE
+	struct SendBlock
+	{
+		const void* data;
+		u32 dataLen;
+		SendBlock* next;
+	};
+	#endif
+
 	struct NetworkBackendTcp::Connection
 	{
 		Connection(Logger& l, SOCKET s) : logger(l), socket(s), ready(true) { CreateGuid(uid); }
@@ -132,6 +143,12 @@ namespace uba
 		bool receivingHeader = true;
 		#endif
 
+		#if UBA_EXPERIMENTAL_WRITE_QUEUE
+		SendBlock* firstBlock = nullptr;
+		SendBlock* lastBlock = nullptr;
+		Futex blockQueueLock;
+		#endif
+
 		Connection(const Connection&) = delete;
 		void operator=(const Connection&) = delete;
 	};
@@ -934,6 +951,95 @@ namespace uba
 		}
 		#endif
 
+#if UBA_EXPERIMENTAL_WRITE_QUEUE
+		if (true)
+		{
+			SendBlock newBlock { b, u32(bufferLen), nullptr };
+
+			SCOPED_FUTEX(connection.blockQueueLock, queueLock);
+			if (connection.lastBlock)
+				connection.lastBlock->next = &newBlock;
+			else
+				connection.firstBlock = &newBlock;
+			connection.lastBlock = &newBlock;
+			queueLock.Leave();
+
+			SCOPED_FUTEX(connection.sendLock, sendLock);
+			//if (newBlock.dataLen == 0)
+			//	return true;
+
+			queueLock.Enter();
+			SendBlock* first = connection.firstBlock;
+			if (!first)
+				return true;
+			connection.firstBlock = nullptr;
+			connection.lastBlock = nullptr;
+			queueLock.Leave();
+
+			while (true)
+			{
+				WSABUF buf[256];
+				u32 count = 0;
+				u32 totalSend = 0;
+				while (first)
+				{
+					if (totalSend + u32(first->dataLen) > 128*1024*1024)
+						break;
+					buf[count++] = { u32(first->dataLen), (char*)first->data };
+					totalSend += first->dataLen;
+					first = first->next;
+					if (count == 255)
+						break;
+				}
+
+				if (count > 20)
+				{
+					logger.Info(TC("Multisend: %u"), count);
+				}
+
+				Event ev(true);
+				OVERLAPPED overlapped {};
+				overlapped.hEvent = ev.GetHandle();
+
+				SOCKET socket = connection.socket;
+				if (socket == INVALID_SOCKET)
+					return false;
+				int sent = WSASend(socket, buf, count, NULL, 0, &overlapped, NULL);
+				if (sent == SOCKET_ERROR)
+				{
+					u32 lastError = WSAGetLastError();
+					if (lastError != WSA_IO_PENDING)
+					{
+						#if UBA_LOG_SOCKET_ERRORS
+						logger.Info(TC("SendSocket - send returned an error for socket %i: %s (%s)"), socket, LastErrorToText(WSAGetLastError()).data, hint);
+						#endif
+						return false;
+					}
+				}
+
+				ev.IsSet();
+
+				DWORD bytesSent;
+				DWORD flags;
+				if (!WSAGetOverlappedResult(socket, &overlapped, &bytesSent, FALSE, &flags))
+				{
+					#if UBA_LOG_SOCKET_ERRORS
+					logger.Info(TC("WSAGetOverlappedResult - error for socket %i: %s (%s)"), socket, LastErrorToText().data, hint);
+					#endif
+					return false;
+				}
+
+				if (bytesSent != totalSend)
+				{
+					logger.Info(TC("ERROR TRIED TO SEND %llu but sent %u"), totalSend, bytesSent);
+				}
+				if (!first)
+					return true;
+				ev.Reset();
+			}
+		}
+#endif
+
 
 #if UBA_USE_OVERLAPPED_SEND
 		if (m_useOverlappedSend)
@@ -2054,14 +2160,14 @@ namespace uba
 		//ULONG bytesReturned;
 		//if (WSAIoctl(socket, SIO_TCP_INFO, &version, sizeof(ULONG), &ti, sizeof(TCP_INFO_v0), &bytesReturned, NULL, NULL) == SOCKET_ERROR)
 		//	return logger.Error(TC("WSAIoctl SIO_TCP_INFO failed (%s)"), LastErrorToText(WSAGetLastError()).data);
-		//logger.Info(TC("RTT=%uÂµs  InFlight=%u  Retrans=%u  Dupacks=%u  State=%u RcvWnd=%u  RcvBuf=%u"), ti.RttUs, ti.BytesInFlight, ti.BytesRetrans, ti.DupAcksIn, ti.State, ti.RcvWnd, ti.RcvBuf);
+		//logger.Info(TC("RTT=%uµs  InFlight=%u  Retrans=%u  Dupacks=%u  State=%u RcvWnd=%u  RcvBuf=%u"), ti.RttUs, ti.BytesInFlight, ti.BytesRetrans, ti.DupAcksIn, ti.State, ti.RcvWnd, ti.RcvBuf);
 	#else
 		/*
 		struct tcp_info ti = {};
 		socklen_t len = sizeof(ti);
 		if (getsockopt(s, IPPROTO_TCP, TCP_INFO, &ti, &len) == 0)
 		{
-			printf("%s  rtt=%uÂµs  unacked=%u  retrans=%u  snd_cwnd=%u  state=%u\n",
+			printf("%s  rtt=%uµs  unacked=%u  retrans=%u  snd_cwnd=%u  state=%u\n",
 				   tag,
 				   ti.tcpi_rtt, ti.tcpi_unacked, ti.tcpi_retrans,
 				   ti.tcpi_snd_cwnd, ti.tcpi_state);
