Uba: Handle .d.tmp files

#rnx

diff --git a/Engine/Source/Programs/UnrealBuildAccelerator/Core/Private/UbaApplicationRules.cpp b/Engine/Source/Programs/UnrealBuildAccelerator/Core/Private/UbaApplicationRules.cpp
--- a/Engine/Source/Programs/UnrealBuildAccelerator/Core/Private/UbaApplicationRules.cpp
+++ b/Engine/Source/Programs/UnrealBuildAccelerator/Core/Private/UbaApplicationRules.cpp
@@ -306,7 +306,7 @@ namespace uba
 		{
 			if (file.EndsWith(TCV(".obj")) || file.EndsWith(TCV(".o")) || file.EndsWith(TCV(".o.tmp")) || file.EndsWith(TCV(".obj.tmp")))
 				return 1024 * 1024 * 1024; // There are huge obj files when building with -stresstestunity
-			if (file.EndsWith(TCV(".d")))
+			if (file.EndsWith(TCV(".d")) || file.EndsWith(TCV(".d.tmp")))
 				return 32 * 1024 * 1024;
 			return Super::FileTypeMaxSize(file, isSystemOrTempFile);
 		}
@@ -324,10 +324,8 @@ namespace uba
 
 		virtual bool IsOutputFile(StringView fileName, StringView systemTemp) const override
 		{
-			return fileName.EndsWith(TCV(".c.d"))
-				|| fileName.EndsWith(TCV(".h.d"))
-				|| fileName.EndsWith(TCV(".cc.d"))
-				|| fileName.EndsWith(TCV(".cpp.d"))
+			return fileName.EndsWith(TCV(".d"))
+				|| fileName.EndsWith(TCV(".d.tmp"))
 				|| fileName.EndsWith(TCV(".o"))
 				|| fileName.EndsWith(TCV(".o.tmp"))
 				|| fileName.EndsWith(TCV(".gch"))
@@ -346,7 +344,7 @@ namespace uba
 
 		virtual bool IsRarelyReadAfterWritten(const StringView& fileName) const override
 		{
-			return fileName.EndsWith(TCV(".d"));
+			return fileName.EndsWith(TCV(".d")) || fileName.EndsWith(TCV(".d.tmp"));
 		}
 
 		virtual bool ShouldExtractSymbols(const StringView& fileName) const override
@@ -357,7 +355,7 @@ namespace uba
 		virtual bool ShouldDevirtualizeFile(const StringView& fileName, bool& outEscapeSpaces) const override
 		{
 			outEscapeSpaces = true;
-			return fileName.EndsWith(TCV(".d"));
+			return fileName.EndsWith(TCV(".d")) || fileName.EndsWith(TCV(".d.tmp"));
 		}
 
 		virtual DependencyCrawlerType GetDependencyCrawlerType() const override
@@ -663,10 +661,8 @@ namespace uba
 	{
 		virtual bool IsOutputFile(StringView fileName, StringView systemTemp) const override
 		{
-			return fileName.EndsWith(TCV(".c.d"))
-				|| fileName.EndsWith(TCV(".h.d"))
-				|| fileName.EndsWith(TCV(".cc.d"))
-				|| fileName.EndsWith(TCV(".cpp.d"))
+			return fileName.EndsWith(TCV(".d"))
+				|| fileName.EndsWith(TCV(".d.tmp"))
 				|| fileName.EndsWith(TCV(".iwyu"))
 				;
 		}
