[UBA]
* Fixed bad charadcters in UbaNetworkBackendTcp
* Fixed another shadowed variable on macos

diff --git a/Engine/Source/Programs/UnrealBuildAccelerator/Common/Private/UbaNetworkBackendTcp.cpp b/Engine/Source/Programs/UnrealBuildAccelerator/Common/Private/UbaNetworkBackendTcp.cpp
--- a/Engine/Source/Programs/UnrealBuildAccelerator/Common/Private/UbaNetworkBackendTcp.cpp
+++ b/Engine/Source/Programs/UnrealBuildAccelerator/Common/Private/UbaNetworkBackendTcp.cpp
@@ -2160,14 +2160,14 @@ namespace uba
 		//ULONG bytesReturned;
 		//if (WSAIoctl(socket, SIO_TCP_INFO, &version, sizeof(ULONG), &ti, sizeof(TCP_INFO_v0), &bytesReturned, NULL, NULL) == SOCKET_ERROR)
 		//	return logger.Error(TC("WSAIoctl SIO_TCP_INFO failed (%s)"), LastErrorToText(WSAGetLastError()).data);
-		//logger.Info(TC("RTT=%uµs  InFlight=%u  Retrans=%u  Dupacks=%u  State=%u RcvWnd=%u  RcvBuf=%u"), ti.RttUs, ti.BytesInFlight, ti.BytesRetrans, ti.DupAcksIn, ti.State, ti.RcvWnd, ti.RcvBuf);
+		//logger.Info(TC("RTT=%uÂµs  InFlight=%u  Retrans=%u  Dupacks=%u  State=%u RcvWnd=%u  RcvBuf=%u"), ti.RttUs, ti.BytesInFlight, ti.BytesRetrans, ti.DupAcksIn, ti.State, ti.RcvWnd, ti.RcvBuf);
 	#else
 		/*
 		struct tcp_info ti = {};
 		socklen_t len = sizeof(ti);
 		if (getsockopt(s, IPPROTO_TCP, TCP_INFO, &ti, &len) == 0)
 		{
-			printf("%s  rtt=%uµs  unacked=%u  retrans=%u  snd_cwnd=%u  state=%u\n",
+			printf("%s  rtt=%uÂµs  unacked=%u  retrans=%u  snd_cwnd=%u  state=%u\n",
 				   tag,
 				   ti.tcpi_rtt, ti.tcpi_unacked, ti.tcpi_retrans,
 				   ti.tcpi_snd_cwnd, ti.tcpi_state);
diff --git a/Engine/Source/Programs/UnrealBuildAccelerator/Common/Private/UbaThread.cpp b/Engine/Source/Programs/UnrealBuildAccelerator/Common/Private/UbaThread.cpp
--- a/Engine/Source/Programs/UnrealBuildAccelerator/Common/Private/UbaThread.cpp
+++ b/Engine/Source/Programs/UnrealBuildAccelerator/Common/Private/UbaThread.cpp
@@ -468,7 +468,6 @@ namespace uba
 			thread_suspend(thread);
 			auto g = MakeGuard([&] { thread_resume(thread); });
 
-			kern_return_t kr;
 			uint64_t pc = 0, fp = 0;
 
 			#if defined(__x86_64__)
