[UBA]
* Added VirtualLock to file mapping functions
* Lock communication shared memory since it is written to all the time but we never want it to be flushed to a page file

diff --git a/Engine/Source/Programs/UnrealBuildAccelerator/Agent/Private/UbaAgent.cpp b/Engine/Source/Programs/UnrealBuildAccelerator/Agent/Private/UbaAgent.cpp
--- a/Engine/Source/Programs/UnrealBuildAccelerator/Agent/Private/UbaAgent.cpp
+++ b/Engine/Source/Programs/UnrealBuildAccelerator/Agent/Private/UbaAgent.cpp
@@ -744,6 +744,7 @@ namespace uba
 		}
 
 		InitMemory();
+		SetMemoryWorkingSet(DefaultMemoryWorkingSet);
 
 		if (useExceptionHandler)
 			AddExceptionHandler();
diff --git a/Engine/Source/Programs/UnrealBuildAccelerator/Cli/Private/UbaCli.cpp b/Engine/Source/Programs/UnrealBuildAccelerator/Cli/Private/UbaCli.cpp
--- a/Engine/Source/Programs/UnrealBuildAccelerator/Cli/Private/UbaCli.cpp
+++ b/Engine/Source/Programs/UnrealBuildAccelerator/Cli/Private/UbaCli.cpp
@@ -148,6 +148,7 @@ namespace uba
 
 		AddExceptionHandler();
 		InitMemory();
+		SetMemoryWorkingSet(DefaultMemoryWorkingSet);
 
 		u32 storageCapacityGb = DefaultCapacityGb;
 		StringBuffer<256> workDir;
diff --git a/Engine/Source/Programs/UnrealBuildAccelerator/Common/Private/UbaFileMapping.cpp b/Engine/Source/Programs/UnrealBuildAccelerator/Common/Private/UbaFileMapping.cpp
--- a/Engine/Source/Programs/UnrealBuildAccelerator/Common/Private/UbaFileMapping.cpp
+++ b/Engine/Source/Programs/UnrealBuildAccelerator/Common/Private/UbaFileMapping.cpp
@@ -26,6 +26,7 @@ namespace uba
 	UnorderedMap<const void*, HANDLE> g_viewMappings;
 	#endif
 
+	bool g_memoryLockEnabled;
 
 #if PLATFORM_WINDOWS
 	Bottleneck& g_createFileHandleBottleneck = *new Bottleneck(8); // Allocated and leaked just to prevent shutdown asserts in debug
@@ -423,6 +424,40 @@ namespace uba
 		memcpy(dest, source, size);
 	}
 
+	void SetMemoryWorkingSet(Logger& logger, u64 minWorkingSet)
+	{
+		if (!SetMemoryWorkingSet(minWorkingSet))
+		{
+			LogEntryType type = LogEntryType_Detail;
+			#if UBA_DEBUG
+			type = LogEntryType_Error;
+			#endif
+			logger.Logf(type, TC("Failed to set memory working set (%s)"), LastErrorToText().data);
+			return;
+		}
+
+		logger.Detail(TC("MemoryWorkingSet min set to %s"), BytesToText(minWorkingSet).str);
+		g_memoryLockEnabled = true;
+	}
+
+	void LockMemory(Logger& logger, void* dest, u64 size, const tchar* hint)
+	{
+		if (!g_memoryLockEnabled)
+			return;
+
+		#if PLATFORM_WINDOWS
+		if (!VirtualLock(dest, size))
+		{
+			g_memoryLockEnabled = false;
+			LogEntryType type = LogEntryType_Detail;
+			#if UBA_DEBUG
+			type = LogEntryType_Error;
+			#endif
+			logger.Logf(type, TC("LockMemory of %s failed for %s (%s)"), BytesToText(size).str, hint, LastErrorToText().data);
+		}
+		#endif
+	}
+
 	////////////////////////////////////////////////////////////////////////////////////////////////////
 
 	FileMappingBuffer::FileMappingBuffer(Logger& logger, WorkManager* workManager)
@@ -904,13 +939,15 @@ namespace uba
 				return { {}, 0, 0 };
 			}
 		}
+		LockMemory(m_logger, data, m_blockSize, hint);
+
 		return {m_mappingHandle, offset, data};
 	}
 
 	void FileMappingAllocator::Free(Allocation allocation)
 	{
 		UBA_ASSERT(allocation.handle == m_mappingHandle);
-		if (!UnmapViewOfFile(m_logger, allocation.memory, m_blockSize, m_name))
+		if (!UnmapViewOfFile(m_logger, allocation.memory, m_blockSize, m_name, true))
 			m_logger.Error(TC("%s - Failed to unmap view of file (%s)"), m_name, LastErrorToText().data);
 		u64 index = allocation.offset / m_blockSize;
 		SCOPED_FUTEX(m_mappingLock, lock);
diff --git a/Engine/Source/Programs/UnrealBuildAccelerator/Common/Private/UbaSession.cpp b/Engine/Source/Programs/UnrealBuildAccelerator/Common/Private/UbaSession.cpp
--- a/Engine/Source/Programs/UnrealBuildAccelerator/Common/Private/UbaSession.cpp
+++ b/Engine/Source/Programs/UnrealBuildAccelerator/Common/Private/UbaSession.cpp
@@ -1704,6 +1704,7 @@ namespace uba
 		m_fileMappingTableHandle = uba::CreateMemoryMappingW(m_logger, PAGE_READWRITE, FileMappingTableMemSize, nullptr, TC("FileMappings"));
 		UBA_ASSERT(m_fileMappingTableHandle.IsValid());
 		m_fileMappingTableMem = MapViewOfFile(m_logger, m_fileMappingTableHandle, FILE_MAP_WRITE, 0, FileMappingTableMemSize, m_allowDiscardVirtualMemory);
+		LockMemory(m_logger, m_fileMappingTableMem, FileMappingTableMemSize, TC("FileMappings"));
 		UBA_ASSERT(m_fileMappingTableMem);
 
 		m_directoryTableHandle = uba::CreateMemoryMappingW(m_logger, PAGE_READWRITE|SEC_RESERVE, DirTableMemSize, nullptr, TC("DirMappings"));
@@ -3726,6 +3727,7 @@ namespace uba
 		if (!MapViewCommit(m_logger, address, toCommit, m_allowDiscardVirtualMemory))
 			m_logger.Error(TC("Failed to commit memory for directory table (Committed: %llu, ToCommit: %llu) (%s)"), m_directoryTableMemCommitted, toCommit, LastErrorToText().data);
 		m_directoryTableMemCommitted += toCommit;
+		LockMemory(m_logger, address, toCommit, TC("DirectoryTable"));
 	}
 
 	void Session::GetSessionInfo(StringBufferBase& out)
diff --git a/Engine/Source/Programs/UnrealBuildAccelerator/Common/Private/UbaSessionServer.cpp b/Engine/Source/Programs/UnrealBuildAccelerator/Common/Private/UbaSessionServer.cpp
--- a/Engine/Source/Programs/UnrealBuildAccelerator/Common/Private/UbaSessionServer.cpp
+++ b/Engine/Source/Programs/UnrealBuildAccelerator/Common/Private/UbaSessionServer.cpp
@@ -2206,7 +2206,7 @@ namespace uba
 
 		lookupCasKey = cachedInfo.casKey;
 		BinaryWriter w(m_nameToHashTableMem.memory, m_nameToHashTableMem.writtenSize, NameToHashMemSize);
-		m_nameToHashTableMem.AllocateNoLock(sizeof(StringKey) + sizeof(CasKey), 1, TC("NameToHashTable"));
+		m_nameToHashTableMem.AllocateNoLock(sizeof(StringKey) + sizeof(CasKey), 1, TC("NameToHashTable"), true);
 		w.WriteStringKey(fileNameKey);
 		w.WriteCasKey(lookupCasKey);
 	}
diff --git a/Engine/Source/Programs/UnrealBuildAccelerator/Common/Private/UbaTrace.cpp b/Engine/Source/Programs/UnrealBuildAccelerator/Common/Private/UbaTrace.cpp
--- a/Engine/Source/Programs/UnrealBuildAccelerator/Common/Private/UbaTrace.cpp
+++ b/Engine/Source/Programs/UnrealBuildAccelerator/Common/Private/UbaTrace.cpp
@@ -110,7 +110,9 @@ namespace uba
 		if (!m_memoryHandle.IsValid())
 			return false;
 		if (GetLastError() != ERROR_ALREADY_EXISTS)
+		{
 			m_memoryBegin = MapViewOfFile(m_logger, m_memoryHandle, FILE_MAP_WRITE, 0, m_memoryCapacity);
+		}
 
 		if (!m_memoryBegin)
 		{
@@ -318,8 +320,11 @@ namespace uba
 		if (m_memoryCommitted >= committedMemoryNeeded)
 			return true;
 
-		if (MapViewCommit(m_logger, m_memoryBegin + m_memoryCommitted, committedMemoryNeeded - m_memoryCommitted))
+		void* address = m_memoryBegin + m_memoryCommitted;
+		u64 toCommit = committedMemoryNeeded - m_memoryCommitted;
+		if (MapViewCommit(m_logger, address, toCommit))
 		{
+			LockMemory(m_logger, address, toCommit, TC("Trace"));
 			m_memoryCommitted = committedMemoryNeeded;
 			return true;
 		}
diff --git a/Engine/Source/Programs/UnrealBuildAccelerator/Common/Public/UbaDefaultConstants.h b/Engine/Source/Programs/UnrealBuildAccelerator/Common/Public/UbaDefaultConstants.h
--- a/Engine/Source/Programs/UnrealBuildAccelerator/Common/Public/UbaDefaultConstants.h
+++ b/Engine/Source/Programs/UnrealBuildAccelerator/Common/Public/UbaDefaultConstants.h
@@ -17,4 +17,6 @@ namespace uba
 
 	static constexpr u32 ProcessPriority_Normal = 0x00000020;
 	static constexpr u32 ProcessPriority_BelowNormal = 0x00004000;
+
+	static constexpr u32 DefaultMemoryWorkingSet = 512*1024*1024;
 }
diff --git a/Engine/Source/Programs/UnrealBuildAccelerator/Common/Public/UbaFileMapping.h b/Engine/Source/Programs/UnrealBuildAccelerator/Common/Public/UbaFileMapping.h
--- a/Engine/Source/Programs/UnrealBuildAccelerator/Common/Public/UbaFileMapping.h
+++ b/Engine/Source/Programs/UnrealBuildAccelerator/Common/Public/UbaFileMapping.h
@@ -37,6 +37,8 @@ namespace uba
 	bool CloseFileMapping(Logger& logger, FileMappingHandle h, const tchar* hint);
 	bool DuplicateFileMapping(Logger& logger, ProcHandle sourceProcessHandle, FileMappingHandle sourceHandle, ProcHandle targetProcessHandle, FileMappingHandle& targetHandle, u32 dwDesiredAccess, bool bInheritHandle, u32 dwOptions, const tchar* hint);
 	void MapMemoryCopy(void* dest, const void* source, u64 size);
+	void SetMemoryWorkingSet(Logger& logger, u64 minWorkingSet);
+	void LockMemory(Logger& logger, void* dest, u64 size, const tchar* hint);
 
 	////////////////////////////////////////////////////////////////////////////////////////////////////
 
diff --git a/Engine/Source/Programs/UnrealBuildAccelerator/Core/Private/UbaMemory.cpp b/Engine/Source/Programs/UnrealBuildAccelerator/Core/Private/UbaMemory.cpp
--- a/Engine/Source/Programs/UnrealBuildAccelerator/Core/Private/UbaMemory.cpp
+++ b/Engine/Source/Programs/UnrealBuildAccelerator/Core/Private/UbaMemory.cpp
@@ -93,7 +93,7 @@ namespace uba
 		return AllocateNoLock(bytes, alignment, hint);
 	}
 
-	void* MemoryBlock::AllocateNoLock(u64 bytes, u64 alignment, const tchar* hint)
+	void* MemoryBlock::AllocateNoLock(u64 bytes, u64 alignment, const tchar* hint, bool lockMemory)
 	{
 		u64 startPos = AlignUp(writtenSize, alignment);
 		u64 newPos = startPos + bytes;
@@ -107,8 +107,11 @@ namespace uba
 			u64 toCommit = AlignUp(newPos - committedSize, MemoryBlock_ReserveAlign);
 			if (committedSize + toCommit > reserveSize)
 				toCommit = reserveSize - committedSize;
-			if (!VirtualAlloc(memory + committedSize, toCommit, MEM_COMMIT, PAGE_READWRITE))
+			void* addr = memory + committedSize;
+			if (!VirtualAlloc(addr, toCommit, MEM_COMMIT, PAGE_READWRITE))
 				FatalError(9883, TC("Failed to commit virtual memory for memory block. Total size %llu (%u) (%s)"), committedSize + toCommit, GetLastError(), hint);
+			if (lockMemory)
+				VirtualLock(addr, toCommit);
 			committedSize += toCommit;
 		}
 		#endif
diff --git a/Engine/Source/Programs/UnrealBuildAccelerator/Core/Private/UbaPlatform.cpp b/Engine/Source/Programs/UnrealBuildAccelerator/Core/Private/UbaPlatform.cpp
--- a/Engine/Source/Programs/UnrealBuildAccelerator/Core/Private/UbaPlatform.cpp
+++ b/Engine/Source/Programs/UnrealBuildAccelerator/Core/Private/UbaPlatform.cpp
@@ -635,6 +635,21 @@ namespace uba
 		#endif
 	}
 
+	bool SetMemoryWorkingSet(u64 minWorkingSet)
+	{
+		#if PLATFORM_WINDOWS
+		SIZE_T minWS, maxWS;
+		DWORD flags;
+		if (!GetProcessWorkingSetSizeEx(GetCurrentProcess(), &minWS, &maxWS, &flags))
+			return false;
+		if (minWorkingSet < minWS)
+			return true;
+		if (!SetProcessWorkingSetSizeEx(GetCurrentProcess(), minWorkingSet, SIZE_T(-1), flags))
+			return false;
+		#endif
+		return true;
+	}
+
 
 #if !PLATFORM_WINDOWS
 	void strcpy_s(tchar* dest, u64 destCapacity, const tchar* source)
diff --git a/Engine/Source/Programs/UnrealBuildAccelerator/Core/Public/UbaMemory.h b/Engine/Source/Programs/UnrealBuildAccelerator/Core/Public/UbaMemory.h
--- a/Engine/Source/Programs/UnrealBuildAccelerator/Core/Public/UbaMemory.h
+++ b/Engine/Source/Programs/UnrealBuildAccelerator/Core/Public/UbaMemory.h
@@ -99,7 +99,7 @@ namespace uba
 		bool Init(u64 reserveSize_, void* baseAddress_ = nullptr, bool useHugePages = false);
 		void Deinit();
 		void* Allocate(u64 bytes, u64 alignment, const tchar* hint);
-		void* AllocateNoLock(u64 bytes, u64 alignment, const tchar* hint);
+		void* AllocateNoLock(u64 bytes, u64 alignment, const tchar* hint, bool lockMemory = false);
 		void* CommitNoLock(u64 bytes, const tchar* hint);
 		void Free(void* p);
 		StringView Strdup(const StringView& str);
diff --git a/Engine/Source/Programs/UnrealBuildAccelerator/Core/Public/UbaPlatform.h b/Engine/Source/Programs/UnrealBuildAccelerator/Core/Public/UbaPlatform.h
--- a/Engine/Source/Programs/UnrealBuildAccelerator/Core/Public/UbaPlatform.h
+++ b/Engine/Source/Programs/UnrealBuildAccelerator/Core/Public/UbaPlatform.h
@@ -111,6 +111,7 @@ namespace uba
 
 	void AddExceptionHandler();
 	void InitMemory();
+	bool SetMemoryWorkingSet(u64 minWorkingSet);
 
 #if PLATFORM_WINDOWS
 	inline constexpr bool CaseInsensitiveFs = true;
diff --git a/Engine/Source/Programs/UnrealBuildAccelerator/Host/Private/UbaExports.cpp b/Engine/Source/Programs/UnrealBuildAccelerator/Host/Private/UbaExports.cpp
--- a/Engine/Source/Programs/UnrealBuildAccelerator/Host/Private/UbaExports.cpp
+++ b/Engine/Source/Programs/UnrealBuildAccelerator/Host/Private/UbaExports.cpp
@@ -126,6 +126,7 @@ namespace uba
 	MutexHandle g_exclusiveMutex = InvalidMutexHandle;
 
 	Trace* g_globalTrace;
+	LogWriter* g_globalLogWriter;
 }
 
 extern "C"
@@ -142,12 +143,23 @@ extern "C"
 
 	uba::LogWriter* LogWriter_CreateCallback(uba::CallbackLogWriter::BeginScopeCallback begin, uba::CallbackLogWriter::EndScopeCallback end, uba::CallbackLogWriter::LogCallback log)
 	{
+		using namespace uba;
 		#if !PLATFORM_LINUX
-		uba::AddExceptionHandler();
+		AddExceptionHandler();
 		#endif
-		uba::InitMemory();
+		InitMemory();
+
+		LogWriter* writer = new CallbackLogWriter(begin, end, log);
+
+		LoggerWithWriter logger(*writer, TC(""));
+		SetMemoryWorkingSet(logger, DefaultMemoryWorkingSet);
 
-		return new uba::CallbackLogWriter(begin, end, log);
+		return writer;
+	}
+
+	void LogWriter_SetGlobal(uba::LogWriter* writer)
+	{
+		uba::g_globalLogWriter = writer;
 	}
 
 	void LogWriter_DestroyCallback(uba::LogWriter* writer)
@@ -1059,7 +1071,7 @@ extern "C"
 	uba::Trace* Trace_Create(const uba::tchar* name, const uba::tchar* channel)
 	{
 		using namespace uba;
-		auto trace = new Trace(g_consoleLogWriter);
+		auto trace = new Trace(g_globalLogWriter ? *g_globalLogWriter : g_consoleLogWriter);
 
 		StringBuffer<256> fullName;
 
diff --git a/Engine/Source/Programs/UnrealBuildAccelerator/Host/Public/UbaExports.h b/Engine/Source/Programs/UnrealBuildAccelerator/Host/Public/UbaExports.h
--- a/Engine/Source/Programs/UnrealBuildAccelerator/Host/Public/UbaExports.h
+++ b/Engine/Source/Programs/UnrealBuildAccelerator/Host/Public/UbaExports.h
@@ -48,6 +48,7 @@ extern "C"
 	UBA_API uba::LogWriter* LogWriter_GetDefault();
 	UBA_API uba::LogWriter* LogWriter_GetNull();
 	UBA_API uba::LogWriter* LogWriter_CreateCallback(uba::CallbackLogWriter::BeginScopeCallback begin, uba::CallbackLogWriter::EndScopeCallback end, uba::CallbackLogWriter::LogCallback log);
+	UBA_API void LogWriter_SetGlobal(uba::LogWriter* writer);
 	UBA_API void LogWriter_DestroyCallback(uba::LogWriter* writer);
 
 	// Config
