[UBA]
* Moved detoured handle start to 600000 (from 300000) to create more room for system using crazy many handles

diff --git a/Engine/Source/Programs/UnrealBuildAccelerator/Detours/Private/Windows/UbaDetoursFunctionsWin.cpp b/Engine/Source/Programs/UnrealBuildAccelerator/Detours/Private/Windows/UbaDetoursFunctionsWin.cpp
--- a/Engine/Source/Programs/UnrealBuildAccelerator/Detours/Private/Windows/UbaDetoursFunctionsWin.cpp
+++ b/Engine/Source/Programs/UnrealBuildAccelerator/Detours/Private/Windows/UbaDetoursFunctionsWin.cpp
@@ -165,7 +165,7 @@ u8 g_emptyMemoryFileMem;
 MemoryFile& g_emptyMemoryFile = *new MemoryFile(&g_emptyMemoryFileMem, true);
 
 constexpr u64 DetouredHandleMaxCount = 200*1024; // ~200000 handles enough?
-constexpr u64 DetouredHandleStart = 300000; // Let's hope noone uses the handles starting at 300000! :)
+constexpr u64 DetouredHandleStart = 600000; // Let's hope noone uses the handles starting at 600000! :)
 constexpr u64 DetouredHandleEnd = DetouredHandleStart + DetouredHandleMaxCount;
 constexpr u64 DetouredHandlesMemReserve = AlignUp(DetouredHandleMaxCount*sizeof(DetouredHandle), 64*1024);
 constexpr u64 DetouredHandlesMemStart = 0;
