[UBA]
* Added EmergencyShutdown code that can be sent from detoured process
* Fixed so we pull the emergency shutdown when there is a write access violation in detoured WriteFile while running on wine. This is likely out of memory and no good way to recover

diff --git a/Engine/Source/Programs/UnrealBuildAccelerator/Common/Private/UbaProcess.cpp b/Engine/Source/Programs/UnrealBuildAccelerator/Common/Private/UbaProcess.cpp
--- a/Engine/Source/Programs/UnrealBuildAccelerator/Common/Private/UbaProcess.cpp
+++ b/Engine/Source/Programs/UnrealBuildAccelerator/Common/Private/UbaProcess.cpp
@@ -659,6 +659,11 @@ namespace uba
 		}
 		#endif
 
+		if (m_nativeProcessExitCode == EmergencyShutdownExitCode)
+		{
+			m_session.EmergencyShutdown();
+		}
+
 		if (IsCancelled())
 			m_exitCode = ProcessCancelExitCode;
 		else
diff --git a/Engine/Source/Programs/UnrealBuildAccelerator/Common/Private/UbaSession.cpp b/Engine/Source/Programs/UnrealBuildAccelerator/Common/Private/UbaSession.cpp
--- a/Engine/Source/Programs/UnrealBuildAccelerator/Common/Private/UbaSession.cpp
+++ b/Engine/Source/Programs/UnrealBuildAccelerator/Common/Private/UbaSession.cpp
@@ -3559,6 +3559,11 @@ namespace uba
 		return true;
 	}
 
+	void Session::EmergencyShutdown()
+	{
+		FatalError(EmergencyShutdownExitCode, TC("Emergency shutdown. Likely out-of-memory"));
+	}
+
 	void Session::PrintSessionStats(Logger& logger)
 	{
 		u64 mappingBufferSize;
diff --git a/Engine/Source/Programs/UnrealBuildAccelerator/Common/Public/UbaSession.h b/Engine/Source/Programs/UnrealBuildAccelerator/Common/Public/UbaSession.h
--- a/Engine/Source/Programs/UnrealBuildAccelerator/Common/Public/UbaSession.h
+++ b/Engine/Source/Programs/UnrealBuildAccelerator/Common/Public/UbaSession.h
@@ -186,6 +186,7 @@ namespace uba
 		virtual bool FlushWrittenFiles(ProcessImpl& process);
 		virtual bool UpdateEnvironment(ProcessImpl& process, const StringView& reason, bool resetStats);
 		virtual bool LogLine(ProcessImpl& process, const tchar* line, LogEntryType logType);
+		virtual void EmergencyShutdown();
 
 		static constexpr CasKey CasKeyIsDirectory = { ~u64(0), ~u64(0), ~u32(0) };
 
diff --git a/Engine/Source/Programs/UnrealBuildAccelerator/Core/Public/UbaProtocol.h b/Engine/Source/Programs/UnrealBuildAccelerator/Core/Public/UbaProtocol.h
--- a/Engine/Source/Programs/UnrealBuildAccelerator/Core/Public/UbaProtocol.h
+++ b/Engine/Source/Programs/UnrealBuildAccelerator/Core/Public/UbaProtocol.h
@@ -53,6 +53,8 @@ namespace uba
 
 	inline constexpr u32 FileMappingTableMemSize = 16 * 1024 * 1024;
 	inline constexpr u32 DirTableMemSize = 128 * 1024 * 1024;
+
+	inline constexpr u32 EmergencyShutdownExitCode = 9886;
 }
 
 #define UBA_ENABLE_ON_DISK_FILE_MAPPINGS 0
diff --git a/Engine/Source/Programs/UnrealBuildAccelerator/Detours/Private/Windows/UbaDetoursFunctionsKernelBase.inl b/Engine/Source/Programs/UnrealBuildAccelerator/Detours/Private/Windows/UbaDetoursFunctionsKernelBase.inl
--- a/Engine/Source/Programs/UnrealBuildAccelerator/Detours/Private/Windows/UbaDetoursFunctionsKernelBase.inl
+++ b/Engine/Source/Programs/UnrealBuildAccelerator/Detours/Private/Windows/UbaDetoursFunctionsKernelBase.inl
@@ -1,7 +1,5 @@
 // Copyright Epic Games, Inc. All Rights Reserved.
 
-#define UBA_FIND_MEMCPY_ACCESS_VIOLATION 1
-
 DWORD Local_GetLongPathNameW(LPCWSTR lpszShortPath, LPWSTR lpszLongPath, DWORD cchBuffer)
 {
 	SCOPED_WRITE_LOCK(g_longPathNameCacheLock, lock);
@@ -656,28 +654,14 @@ HANDLE Detoured_CreateToolhelp32Snapshot(DWORD dwFlags, DWORD th32ProcessID)
 	//return True_CreateToolhelp32Snapshot(dwFlags, th32ProcessID);
 }
 
-#if UBA_FIND_MEMCPY_ACCESS_VIOLATION
-struct MemoryFileData
-{
-	u8* baseAddress = nullptr;
-	u64 reserveSize = 0;
-	u64 mappedSize = 0;
-	u64 committedSize = 0;
-	u64 writtenSize = 0;
-};
-int LogAccessViolation(PEXCEPTION_POINTERS ep, MemoryFile& mf, MemoryFileData& data, HANDLE hFile, DWORD nNumberOfBytesToWrite)
+int WriteFileAccessViolation(PEXCEPTION_POINTERS ep)
 {
-	if (ep->ExceptionRecord->ExceptionCode == EXCEPTION_ACCESS_VIOLATION)// && r->NumberParameters >= 2 && r->ExceptionInformation[0] == 1)
-	{
-		StringBuffer<> buf;
-		buf.Appendf(L"Failed to write %llu bytes to memory file %s\r\n", nNumberOfBytesToWrite, HandleToName(hFile));
-		buf.Appendf(L" Before %p %llu %llu %llu %llu\r\n", data.baseAddress, data.reserveSize, data.mappedSize, data.committedSize, data.writtenSize);
-		buf.Appendf(L"  After %p %llu %llu %llu %llu", mf.baseAddress, mf.reserveSize, mf.mappedSize, mf.committedSize, mf.writtenSize);
-		Rpc_WriteLog(buf.data, buf.count, true, true);
-	}
+	// This likely is OOM in the file mapping system on wine (tmpfs running out of space or something).
+	// We need to improve how uba tracks memory usage on wine to catch these before it is too late. but for now, just terminate
+	if (g_isRunningWine && ep->ExceptionRecord->ExceptionCode == EXCEPTION_ACCESS_VIOLATION)
+		TerminateProcess(GetCurrentProcess(), EmergencyShutdownExitCode);
 	return EXCEPTION_CONTINUE_SEARCH;
 }
-#endif
 
 BOOL Detoured_WriteFile(HANDLE hFile, LPCVOID lpBuffer, DWORD nNumberOfBytesToWrite, LPDWORD lpNumberOfBytesWritten, LPOVERLAPPED lpOverlapped)
 {
@@ -727,26 +711,13 @@ BOOL Detoured_WriteFile(HANDLE hFile, LPCVOID lpBuffer, DWORD nNumberOfBytesToWr
 			{
 				SCOPED_WRITE_LOCK(mf->lock, lock);
 
-#if UBA_FIND_MEMCPY_ACCESS_VIOLATION
-				MemoryFileData data;
-				data.baseAddress = mf->baseAddress;
-				data.reserveSize = mf->reserveSize;
-				data.mappedSize = mf->mappedSize;
-				data.committedSize = mf->committedSize;
-				data.writtenSize = mf->writtenSize;
-
 				__try
 				{
-					//*(int*)nullptr = 123;
-#endif
 					mf->Write(dh, lpBuffer, nNumberOfBytesToWrite);
-#if UBA_FIND_MEMCPY_ACCESS_VIOLATION
 				}
-				__except(LogAccessViolation(GetExceptionInformation(), *mf, data, hFile, nNumberOfBytesToWrite))
+				__except(WriteFileAccessViolation(GetExceptionInformation()))
 				{
-					int dummy;(void)dummy;
 				}
-#endif
 			}
 
 			*lpNumberOfBytesWritten = nNumberOfBytesToWrite;
