[UBA]
* Lock communication shared memory on detoured process side
* improved error message when closing in on detoured handle range with normal handles

diff --git a/Engine/Source/Programs/UnrealBuildAccelerator/Detours/Private/Windows/UbaDetoursFunctionsNtDll.inl b/Engine/Source/Programs/UnrealBuildAccelerator/Detours/Private/Windows/UbaDetoursFunctionsNtDll.inl
--- a/Engine/Source/Programs/UnrealBuildAccelerator/Detours/Private/Windows/UbaDetoursFunctionsNtDll.inl
+++ b/Engine/Source/Programs/UnrealBuildAccelerator/Detours/Private/Windows/UbaDetoursFunctionsNtDll.inl
@@ -703,7 +703,7 @@ NTSTATUS NTAPI Local_NtCreateFile(bool IsCreateFunc, PHANDLE hFileHandle, ACCESS
 		//Rpc_WriteLogf(L"NtCreateFile: %ls", ObjectAttributes->ObjectName->Buffer);
 		auto res = True_NtCreateFile(hFileHandle, DesiredAccess, ObjectAttributes, IoStatusBlock, AllocationSize, FileAttributes, ShareAccess, CreateDisposition, CreateOptions, EaBuffer, EaLength);
 
-		UBA_ASSERTF(res != STATUS_SUCCESS || u64(*hFileHandle) < DetouredHandleStart - 10000, L"Normal handle range is closing in on detoured. Bump detour range (normal: %llu, detour start: %llu)", u64(*hFileHandle), DetouredHandleStart);
+		UBA_ASSERTF(res != STATUS_SUCCESS || u64(*hFileHandle) < DetouredHandleStart - 10000, L"Normal handle range is closing in on detoured. Bump detour range (normal: %llu, detour start: %llu) (%.*s)", u64(*hFileHandle), DetouredHandleStart, ObjectAttributes->ObjectName->Length / 2, ObjectAttributes->ObjectName->Buffer);
 
 		// I have no idea why we get this sometimes when trying to open pch for read after recently being written..
 		// All scenarios I've seen this succeeds after 1 second.
diff --git a/Engine/Source/Programs/UnrealBuildAccelerator/Detours/Private/Windows/UbaDetoursMainWin.cpp b/Engine/Source/Programs/UnrealBuildAccelerator/Detours/Private/Windows/UbaDetoursMainWin.cpp
--- a/Engine/Source/Programs/UnrealBuildAccelerator/Detours/Private/Windows/UbaDetoursMainWin.cpp
+++ b/Engine/Source/Programs/UnrealBuildAccelerator/Detours/Private/Windows/UbaDetoursMainWin.cpp
@@ -54,9 +54,18 @@ namespace uba
 		}
 
 		u8* mem = (u8*)::MapViewOfFile(g_messageMappingHandle, FILE_MAP_READ | FILE_MAP_WRITE, ToHigh(payload.communicationOffset), ToLow(payload.communicationOffset), CommunicationMemSize);
+		if (!mem)
+		{
+			return -5;
+		}
+		
 		g_readMessageMappingMem = mem;
 		g_writeMessageMappingMem = mem;// +CommunicationMemSize / 2;
 
+		SetMemoryWorkingSet(CommunicationMemSize*2);
+
+		VirtualLock(mem, CommunicationMemSize);
+
 		return 0;
 	}
 
