[UBA]
* Changed work priority to be an enum.

diff --git a/Engine/Source/Programs/UnrealBuildAccelerator/Common/Private/UbaCacheServer.cpp b/Engine/Source/Programs/UnrealBuildAccelerator/Common/Private/UbaCacheServer.cpp
--- a/Engine/Source/Programs/UnrealBuildAccelerator/Common/Private/UbaCacheServer.cpp
+++ b/Engine/Source/Programs/UnrealBuildAccelerator/Common/Private/UbaCacheServer.cpp
@@ -1134,7 +1134,7 @@ namespace uba
 					bucket.m_cacheEntryLookup.erase(key);
 
 				totalEntryCount += bucket.totalEntryCount;
-			}, TCV("ParallelForBucket1"), true);
+			}, TCV("ParallelForBucket1"), WorkPriority_High);
 
 			// Reset deleted cas files and update it again..
 			deletedCasFiles.clear();
@@ -1388,7 +1388,7 @@ namespace uba
 			if (bucket.oldestUsedTime)
 				longestUnusedSeconds = GetFileTimeAsSeconds(now - m_creationTime - bucket.oldestUsedTime);
 			m_logger.Info(TC("    Bucket %u Done (%s). CacheEntries: %s (%s) PathTable: %s CasTable: %s LongestUnused: %s"), bucket.index, TimeToText(GetTime() - bucketStartTime).str, CountToText(bucket.totalEntryCount).str, BytesToText(bucket.totalEntrySize).str, BytesToText(bucket.m_pathTable.GetSize()).str, BytesToText(bucket.m_casKeyTable.GetSize()).str, TimeToText(MsToTime(longestUnusedSeconds*1000), true).str);
-		}, TCV("ParallelForBucket2"), true);
+		}, TCV("ParallelForBucket2"), WorkPriority_High);
 
 		// Need to make sure all cas entries are dropped before saving cas table
 		u64 dropStartTime = GetTime();
diff --git a/Engine/Source/Programs/UnrealBuildAccelerator/Common/Private/UbaNetworkServer.cpp b/Engine/Source/Programs/UnrealBuildAccelerator/Common/Private/UbaNetworkServer.cpp
--- a/Engine/Source/Programs/UnrealBuildAccelerator/Common/Private/UbaNetworkServer.cpp
+++ b/Engine/Source/Programs/UnrealBuildAccelerator/Common/Private/UbaNetworkServer.cpp
@@ -1051,13 +1051,13 @@ namespace uba
 		}
 	}
 
-	void NetworkServer::AddWork(const WorkFunction& work, u32 count, const tchar* desc, const Color& color, bool highPriority, Bottleneck* bottleneck)
+	void NetworkServer::AddWork(const WorkFunction& work, u32 count, const tchar* desc, const Color& color, WorkPriority priority, Bottleneck* bottleneck)
 	{
 		UBA_ASSERT(*desc);
 		SCOPED_FUTEX(m_additionalWorkLock, lock);
 		for (u32 i = 0; i != count; ++i)
 		{
-			if (highPriority)
+			if (priority == WorkPriority_High)
 			{
 				m_additionalWork.push_front({ work, bottleneck });
 				if (m_workTracker)
diff --git a/Engine/Source/Programs/UnrealBuildAccelerator/Common/Private/UbaSessionClient.cpp b/Engine/Source/Programs/UnrealBuildAccelerator/Common/Private/UbaSessionClient.cpp
--- a/Engine/Source/Programs/UnrealBuildAccelerator/Common/Private/UbaSessionClient.cpp
+++ b/Engine/Source/Programs/UnrealBuildAccelerator/Common/Private/UbaSessionClient.cpp
@@ -431,7 +431,7 @@ namespace uba
 				StringBuffer<MaxPath> temp;
 				if (!WriteBinFile(temp, ToView(moduleName), newCasKey, keyStr, m.attributes))
 					success = false;
-			}, AsView(TC("FetchModule")), true);
+			}, AsView(TC("FetchModule")), WorkPriority_High);
 
 		if (!success)
 			return false;
diff --git a/Engine/Source/Programs/UnrealBuildAccelerator/Common/Private/UbaSessionServer.cpp b/Engine/Source/Programs/UnrealBuildAccelerator/Common/Private/UbaSessionServer.cpp
--- a/Engine/Source/Programs/UnrealBuildAccelerator/Common/Private/UbaSessionServer.cpp
+++ b/Engine/Source/Programs/UnrealBuildAccelerator/Common/Private/UbaSessionServer.cpp
@@ -399,7 +399,7 @@ namespace uba
 								w.WriteStringKey(fileNameKey);
 								w.WriteCasKey(lookupCasKey);
 							}
-						}, AsView(TC("KnownInputsLoop")), true);
+						}, AsView(TC("KnownInputsLoop")), WorkPriority_High);
 
 					remoteProcess->m_knownInputsCount = keysIndex;
 					remoteProcess->m_knownInputsDone.Set();
diff --git a/Engine/Source/Programs/UnrealBuildAccelerator/Common/Private/UbaWorkManager.cpp b/Engine/Source/Programs/UnrealBuildAccelerator/Common/Private/UbaWorkManager.cpp
--- a/Engine/Source/Programs/UnrealBuildAccelerator/Common/Private/UbaWorkManager.cpp
+++ b/Engine/Source/Programs/UnrealBuildAccelerator/Common/Private/UbaWorkManager.cpp
@@ -164,14 +164,14 @@ namespace uba
 	}
 
 
-	void WorkManagerImpl::AddWork(const WorkFunction& work, u32 count, const tchar* desc, const Color& color, bool highPriority, Bottleneck* bottleneck)
+	void WorkManagerImpl::AddWork(const WorkFunction& work, u32 count, const tchar* desc, const Color& color, WorkPriority priority, Bottleneck* bottleneck)
 	{
 		UBA_ASSERT(*desc);
 		SCOPED_FUTEX(m_workLock, lock);
 		bool trackWork = m_workTracker.load();
 		for (u32 i = 0; i != count; ++i)
 		{
-			if (highPriority)
+			if (priority == WorkPriority_High)
 			{
 				m_work.push_front({ work, bottleneck });
 				if (trackWork)
diff --git a/Engine/Source/Programs/UnrealBuildAccelerator/Common/Public/UbaDependencyCrawler.h b/Engine/Source/Programs/UnrealBuildAccelerator/Common/Public/UbaDependencyCrawler.h
--- a/Engine/Source/Programs/UnrealBuildAccelerator/Common/Public/UbaDependencyCrawler.h
+++ b/Engine/Source/Programs/UnrealBuildAccelerator/Common/Public/UbaDependencyCrawler.h
@@ -195,7 +195,7 @@ namespace uba
 							{
 								bool stopCrawling = false;
 								ref.instance.createFileFunc(context.tracker, codeFile, {}, stopCrawling);
-							}, 1, TC("CrawlForDecomp"), ColorWork, false, &m_bottleneck);
+							}, 1, TC("CrawlForDecomp"), ColorWork, WorkPriority_Low, &m_bottleneck);
 					}
 					else
 					{
@@ -231,7 +231,7 @@ namespace uba
 					}
 				}
 
-			}, 1, TC("CrawlRsp"), ColorWork, false, &m_bottleneck);
+			}, 1, TC("CrawlRsp"), ColorWork, WorkPriority_Low, &m_bottleneck);
 		return true;
 	}
 
@@ -811,7 +811,7 @@ namespace uba
 						if (parseDefines)
 							work({tracker});
 						else
-							m_workManager.AddWork(work, 1, TC("CrawlIncludes"), ColorWork, false, &m_bottleneck);
+							m_workManager.AddWork(work, 1, TC("CrawlIncludes"), ColorWork, WorkPriority_Low, &m_bottleneck);
 					}
 				}
 				else if (parseDefines && strncmp(it, "define", 6) == 0)
@@ -907,7 +907,7 @@ namespace uba
 					m_workManager.AddWork([this, ref = InstanceRef(instance), codeFile = codeFilePath.ToString()](const WorkContext& context)
 						{
 							TraverseInclude(context.tracker, ref.instance, codeFile);
-						}, 1, TC("CrawlBiDir"), ColorWork, false, &m_bottleneck);
+						}, 1, TC("CrawlBiDir"), ColorWork, WorkPriority_Low, &m_bottleneck);
 				}
 				else
 				{
@@ -922,7 +922,7 @@ namespace uba
 							if (!ref.instance.createFileFunc(context.tracker, codeFile, {}, stopCrawling))
 								if (!stopCrawling)
 									m_logger.Warning(TC("Failed to open file %s from builtindir"), codeFile.c_str());
-						}, 1, TC("CrawlBiFile"), ColorWork, false, &m_bottleneck);
+						}, 1, TC("CrawlBiFile"), ColorWork, WorkPriority_Low, &m_bottleneck);
 				}
 			});
 		return true;
diff --git a/Engine/Source/Programs/UnrealBuildAccelerator/Common/Public/UbaNetworkServer.h b/Engine/Source/Programs/UnrealBuildAccelerator/Common/Public/UbaNetworkServer.h
--- a/Engine/Source/Programs/UnrealBuildAccelerator/Common/Public/UbaNetworkServer.h
+++ b/Engine/Source/Programs/UnrealBuildAccelerator/Common/Public/UbaNetworkServer.h
@@ -90,7 +90,7 @@ namespace uba
 		//void RegisterOnConnection(u8 id, const OnConnectionFunction& func);
 		//void UnregisterOnConnection(u8 id);
 
-		virtual void AddWork(const WorkFunction& work, u32 count, const tchar* desc, const Color& color = ColorWork, bool highPriority = false, Bottleneck* bottleneck = nullptr) override final;
+		virtual void AddWork(const WorkFunction& work, u32 count, const tchar* desc, const Color& color = ColorWork, WorkPriority priority = WorkPriority_Normal, Bottleneck* bottleneck = nullptr) override final;
 		virtual void DoWork(u32 count = 1) override final;
 		virtual u32 GetWorkerCount() override final;
 		virtual bool FlushWork(u32 timeoutMs = 0) override final;
diff --git a/Engine/Source/Programs/UnrealBuildAccelerator/Common/Public/UbaWorkManager.h b/Engine/Source/Programs/UnrealBuildAccelerator/Common/Public/UbaWorkManager.h
--- a/Engine/Source/Programs/UnrealBuildAccelerator/Common/Public/UbaWorkManager.h
+++ b/Engine/Source/Programs/UnrealBuildAccelerator/Common/Public/UbaWorkManager.h
@@ -34,11 +34,18 @@ namespace uba
 		Thread* thread = nullptr;
 	};
 
+	enum WorkPriority : u8
+	{
+		WorkPriority_Normal,
+		WorkPriority_High,
+		WorkPriority_Low,
+	};
+
 	class WorkManager : public WorkTracker
 	{
 	public:
 		using WorkFunction = Function<void(const WorkContext& context)>;
-		virtual void AddWork(const WorkFunction& work, u32 count, const tchar* desc, const Color& color = ColorWork, bool highPriority = false, Bottleneck* bottleneck = nullptr) = 0;
+		virtual void AddWork(const WorkFunction& work, u32 count, const tchar* desc, const Color& color = ColorWork, WorkPriority priority = WorkPriority_Normal, Bottleneck* bottleneck = nullptr) = 0;
 		virtual u32 GetWorkerCount() = 0;
 		virtual void DoWork(u32 count = 1) = 0;
 		virtual bool FlushWork(u32 timeoutMs = 0) = 0;
@@ -51,7 +58,7 @@ namespace uba
 		WorkTracker* GetWorkTracker() { return m_workTracker; }
 
 		template<u32 BatchSize = 1, typename TContainer, typename TFunc>
-		void ParallelFor(u32 workCount, TContainer& container, TFunc&& func, const StringView& description = AsView(TC("")), bool highPriority = false, Bottleneck* bottleneck = nullptr);
+		void ParallelFor(u32 workCount, TContainer& container, TFunc&& func, const StringView& description = AsView(TC("")), WorkPriority priority = WorkPriority_Normal, Bottleneck* bottleneck = nullptr);
 
 	protected:
 		Atomic<WorkTracker*> m_workTracker;
@@ -63,7 +70,7 @@ namespace uba
 	public:
 		WorkManagerImpl(u32 maxWorkerCount, const tchar* workerDesc = TC("UbaWrk"), u32 precreateWorkerCount = 1);
 		virtual ~WorkManagerImpl();
-		virtual void AddWork(const WorkFunction& work, u32 count, const tchar* desc, const Color& color = ColorWork, bool highPriority = false, Bottleneck* bottleneck = nullptr) override final;
+		virtual void AddWork(const WorkFunction& work, u32 count, const tchar* desc, const Color& color = ColorWork, WorkPriority priority = WorkPriority_Normal, Bottleneck* bottleneck = nullptr) override final;
 		virtual u32 GetWorkerCount() override final;
 		virtual void DoWork(u32 count = 1) override final;
 		virtual bool FlushWork(u32 timeoutMs = 0) override final;
@@ -107,7 +114,7 @@ namespace uba
 	};
 
 	template<u32 BatchSize, typename TContainer, typename TFunc>
-	void WorkManager::ParallelFor(u32 workCount, TContainer& container, TFunc&& func, const StringView& description, bool highPriority, Bottleneck* bottleneck)
+	void WorkManager::ParallelFor(u32 workCount, TContainer& container, TFunc&& func, const StringView& description, WorkPriority priority, Bottleneck* bottleneck)
 	{
 		#if !defined( __clang_analyzer__ ) // Static analyzer claims context can leak but it can only leak when process terminates (and then it doesn't matter)
 
@@ -193,7 +200,7 @@ namespace uba
 				}
 			};
 		
-		AddWork(work, workCount, description.data, ColorWork, highPriority, bottleneck);
+		AddWork(work, workCount, description.data, ColorWork, priority, bottleneck);
 
 		{
 			#if UBA_TRACK_WORK
