[UBA]
* Added so work manager handles different work priorities

diff --git a/Engine/Source/Programs/UnrealBuildAccelerator/Common/Private/UbaNetworkServer.cpp b/Engine/Source/Programs/UnrealBuildAccelerator/Common/Private/UbaNetworkServer.cpp
--- a/Engine/Source/Programs/UnrealBuildAccelerator/Common/Private/UbaNetworkServer.cpp
+++ b/Engine/Source/Programs/UnrealBuildAccelerator/Common/Private/UbaNetworkServer.cpp
@@ -726,27 +726,15 @@ namespace uba
 
 	void NetworkServer::Worker::DoAdditionalWorkAndSignalAvailable(NetworkServer& server)
 	{
-		auto& additionalWork = server.m_additionalWork;
-		auto& additionalWorkLock = server.m_additionalWorkLock;
+		auto& additionalWorkQueue = server.m_additionalWorkQueue;
 		BottleneckTicket ticket;
 
 		while (true)
 		{
 			while (true)
 			{
-				AdditionalWork work;
-				SCOPED_FUTEX(additionalWorkLock, lock);
-				for (auto i=additionalWork.begin(), e=additionalWork.end(); i!=e; ++i)
-				{
-					if (i->bottleneck && !i->bottleneck->TryEnter())
-						continue;
-					work = std::move(*i);
-					additionalWork.erase(i);
-					break;
-				}
-				lock.Leave();
-
-				if (!work.func)
+				Work work;
+				if (!additionalWorkQueue.Dequeue(work))
 					break;
 
 				#if UBA_TRACK_WORK
@@ -766,10 +754,10 @@ namespace uba
 			// a race where AddWork would not see this thread in the
 			// available list after adding some work.
 			SCOPED_FUTEX(server.m_availableWorkersLock, lock1);
-			SCOPED_FUTEX_READ(additionalWorkLock, lock2);
+			SCOPED_FUTEX_READ(additionalWorkQueue.lock, lock2);
 			// Verify there is not additional work while we hold both lock
 			// and only add ourself as available if no additional work is present.
-			if (!additionalWork.empty())
+			if (!additionalWorkQueue.IsEmptyNoLock())
 				continue;
 			server.PushWorkerNoLock(this);
 			break;
@@ -1053,24 +1041,7 @@ namespace uba
 
 	void NetworkServer::AddWork(const WorkFunction& work, u32 count, const tchar* desc, const Color& color, WorkPriority priority, Bottleneck* bottleneck)
 	{
-		UBA_ASSERT(*desc);
-		SCOPED_FUTEX(m_additionalWorkLock, lock);
-		for (u32 i = 0; i != count; ++i)
-		{
-			if (priority == WorkPriority_High)
-			{
-				m_additionalWork.push_front({ work, bottleneck });
-				if (m_workTracker)
-					m_additionalWork.front().desc = desc;
-			}
-			else
-			{
-				m_additionalWork.push_back({ work, bottleneck });
-				if (m_workTracker)
-					m_additionalWork.back().desc = desc;
-			}
-		}
-		lock.Leave();
+		m_additionalWorkQueue.Enqueue(work, count, desc, color, priority, bottleneck, m_workTracker != nullptr);
 
 		SCOPED_FUTEX(m_availableWorkersLock, lock2);
 		if (!m_workersEnabled)
@@ -1106,9 +1077,7 @@ namespace uba
 
 		while (true)
 		{
-			SCOPED_FUTEX_READ(m_additionalWorkLock, lock1);
-			bool noAdditionalWork = m_additionalWork.empty();
-			lock1.Leave();
+			bool noAdditionalWork = m_additionalWorkQueue.IsEmpty();
 
 			SCOPED_FUTEX_READ(m_availableWorkersLock, lock2);
 			bool noActiveWorkers = !m_firstActiveWorker;
@@ -1204,67 +1173,60 @@ namespace uba
 
 	bool NetworkServer::DoAdditionalWork()
 	{
-		SCOPED_FUTEX(m_additionalWorkLock, lock);
-		if (m_additionalWork.empty())
+		Work work;
+		if (m_additionalWorkQueue.Dequeue(work))
 		{
-			lock.Leave();
-
-			SCOPED_FUTEX(m_availableWorkersLock, lock2);
-			if (m_createdWorkerCount != m_maxWorkerCount)
-				return false;
-			lock2.Leave();
-
-			auto worker = t_worker;
-			if (!worker)
-				return false;
-
-			auto oldContext = worker->m_context;
-			WorkerContext context(*this);
-			worker->m_context = &context;
-
-			PushWorker(worker);
-			bool workAvail = context.workAvailable.IsSet(10);
-			lock2.Enter();
-			if (worker->m_inUse)
-			{
-				lock2.Leave();
-				if (!workAvail)
-					context.workAvailable.IsSet(~0u);
-				worker->Update(context);
-				UBA_ASSERT(worker->m_inUse);
-			}
-			else
-			{
-				// Take worker back from free list
-				if (m_firstAvailableWorker == worker)
-					m_firstAvailableWorker = worker->m_nextWorker;
-				else
-					worker->m_prevWorker->m_nextWorker = worker->m_nextWorker;
-				if (worker->m_nextWorker)
-					worker->m_nextWorker->m_prevWorker = worker->m_prevWorker;
-				worker->m_prevWorker = nullptr;
-				worker->m_nextWorker = m_firstActiveWorker;
-				if (m_firstActiveWorker)
-					m_firstActiveWorker->m_prevWorker = worker;
-				m_firstActiveWorker = worker;
-				worker->m_inUse = true;
-			}
-			
-			worker->m_context = oldContext;
+			#if UBA_TRACK_WORK
+			TrackWorkScope tws(*this, work.desc, ColorWork);
+			#else
+			TrackWorkScope tws;
+			#endif
+			work.func({tws});
 			return true;
 		}
-		AdditionalWork work = std::move(m_additionalWork.front());
-		m_additionalWork.pop_front();
+
+		SCOPED_FUTEX(m_availableWorkersLock, lock);
+		if (m_createdWorkerCount != m_maxWorkerCount)
+			return false;
 		lock.Leave();
 
-		#if UBA_TRACK_WORK
-		TrackWorkScope tws(*this, work.desc, ColorWork);
-		#else
-		TrackWorkScope tws;
-		#endif
+		auto worker = t_worker;
+		if (!worker)
+			return false;
 
-		work.func({tws});
+		auto oldContext = worker->m_context;
+		WorkerContext context(*this);
+		worker->m_context = &context;
 
+		PushWorker(worker);
+		bool workAvail = context.workAvailable.IsSet(10);
+		lock.Enter();
+		if (worker->m_inUse)
+		{
+			lock.Leave();
+			if (!workAvail)
+				context.workAvailable.IsSet(~0u);
+			worker->Update(context);
+			UBA_ASSERT(worker->m_inUse);
+		}
+		else
+		{
+			// Take worker back from free list
+			if (m_firstAvailableWorker == worker)
+				m_firstAvailableWorker = worker->m_nextWorker;
+			else
+				worker->m_prevWorker->m_nextWorker = worker->m_nextWorker;
+			if (worker->m_nextWorker)
+				worker->m_nextWorker->m_prevWorker = worker->m_prevWorker;
+			worker->m_prevWorker = nullptr;
+			worker->m_nextWorker = m_firstActiveWorker;
+			if (m_firstActiveWorker)
+				m_firstActiveWorker->m_prevWorker = worker;
+			m_firstActiveWorker = worker;
+			worker->m_inUse = true;
+		}
+			
+		worker->m_context = oldContext;
 		return true;
 	}
 
diff --git a/Engine/Source/Programs/UnrealBuildAccelerator/Common/Private/UbaWorkManager.cpp b/Engine/Source/Programs/UnrealBuildAccelerator/Common/Private/UbaWorkManager.cpp
--- a/Engine/Source/Programs/UnrealBuildAccelerator/Common/Private/UbaWorkManager.cpp
+++ b/Engine/Source/Programs/UnrealBuildAccelerator/Common/Private/UbaWorkManager.cpp
@@ -68,6 +68,114 @@ namespace uba
 		m_workTracker = workTracker;
 	}
 
+	WorkQueue::~WorkQueue()
+	{
+		for (u32 i=0; i != sizeof_array(first); ++i)
+		{
+			Work* w = first[i];
+			while (w)
+			{
+				Work* n = w->next;
+				delete w;
+				w = n;
+			}
+		}
+	}
+
+	void WorkQueue::Enqueue(const WorkFunction& func, u32 count, const tchar* desc, const Color& color, WorkPriority priority, Bottleneck* bottleneck, bool trackWork)
+	{
+		UBA_ASSERT(*desc);
+		SCOPED_FUTEX(lock, lk);
+		for (u32 i = 0; i != count; ++i)
+		{
+			Work* work;
+			if (free)
+			{
+				work = free;
+				free = free->next;
+			}
+			else
+				work = new Work();
+		
+			work->next = nullptr;
+			work->func = func;
+			work->bottleneck = bottleneck;
+			if (trackWork)
+				work->desc = desc;
+
+			Work*& l = last[priority];
+
+			if (!l)
+			{
+				l = work;
+				first[priority] = work;
+			}
+			else
+			{
+				l->next = work;
+				l = work;
+			}
+		}
+	}
+
+	bool WorkQueue::Dequeue(Work& out)
+	{
+		SCOPED_FUTEX(lock, l);
+		return DequeueNoLock(out);
+	}
+
+	bool WorkQueue::IsEmpty()
+	{
+		SCOPED_FUTEX(lock, l);
+		return IsEmptyNoLock();
+	}
+
+	bool WorkQueue::DequeueNoLock(Work& out)
+	{
+		for (u32 i=0; i != sizeof_array(first); ++i)
+		{
+			Work* p = nullptr;
+			Work* w = first[i];
+			while (w)
+			{
+				if (w->bottleneck && !w->bottleneck->TryEnter())
+				{
+					p = w;
+					w = w->next;
+					continue;
+				}
+
+				if (!w->next)
+				{
+					last[i] = p;
+				}
+				if (p)
+				{
+					p->next = w->next;
+				}
+				else
+				{
+					first[i] = w->next;
+				}
+
+				w->next = free;
+				free = w;
+				out = *w;
+				return true;
+			}
+		}
+		return false;
+	}
+
+	bool WorkQueue::IsEmptyNoLock()
+	{
+		for (u32 i=0; i != sizeof_array(first); ++i)
+			if (first[i])
+				return false;
+		return true;
+	}
+
+
 	struct WorkManagerImpl::Worker
 	{
 		Worker(WorkManagerImpl& manager, const tchar* workerDesc) : m_workAvailable(false)
@@ -89,6 +197,7 @@ namespace uba
 		void ThreadWorker(WorkManagerImpl& manager)
 		{
 			BottleneckTicket ticket;
+			WorkQueue& workQueue = manager.m_workQueue;
 
 			while (true)
 			{
@@ -102,17 +211,7 @@ namespace uba
 					while (true)
 					{
 						Work work;
-						SCOPED_FUTEX(manager.m_workLock, lock);
-						for (auto i=manager.m_work.begin(), e=manager.m_work.end(); i!=e; ++i)
-						{
-							if (i->bottleneck && !i->bottleneck->TryEnter())
-								continue;
-							work = std::move(*i);
-							manager.m_work.erase(i);
-							break;
-						}
-						lock.Leave();
-						if (!work.func)
+						if (!workQueue.Dequeue(work))
 							break;
 
 						#if UBA_TRACK_WORK
@@ -127,8 +226,8 @@ namespace uba
 					}
 
 					SCOPED_FUTEX(manager.m_availableWorkersLock, lock1);
-					SCOPED_FUTEX_READ(manager.m_workLock, lock2);
-					if (!manager.m_work.empty())
+					SCOPED_FUTEX_READ(workQueue.lock, lock2);
+					if (!workQueue.IsEmptyNoLock())
 						continue;
 
 					manager.PushWorkerNoLock(this);
@@ -166,25 +265,7 @@ namespace uba
 
 	void WorkManagerImpl::AddWork(const WorkFunction& work, u32 count, const tchar* desc, const Color& color, WorkPriority priority, Bottleneck* bottleneck)
 	{
-		UBA_ASSERT(*desc);
-		SCOPED_FUTEX(m_workLock, lock);
-		bool trackWork = m_workTracker.load();
-		for (u32 i = 0; i != count; ++i)
-		{
-			if (priority == WorkPriority_High)
-			{
-				m_work.push_front({ work, bottleneck });
-				if (trackWork)
-					m_work.front().desc = desc;
-			}
-			else
-			{
-				m_work.push_back({ work, bottleneck });
-				if (trackWork)
-					m_work.back().desc = desc;
-			}
-		}
-		lock.Leave();
+		m_workQueue.Enqueue(work, count, desc, color, priority, bottleneck, m_workTracker.load());
 
 		SCOPED_FUTEX(m_availableWorkersLock, lock2);
 		while (count--)
@@ -215,13 +296,9 @@ namespace uba
 	{
 		while (count--)
 		{
-			SCOPED_FUTEX(m_workLock, lock);
-			if (m_work.empty())
+			Work work;
+			if (!m_workQueue.Dequeue(work))
 				break;
-			Work work = m_work.front();
-			m_work.pop_front();
-			lock.Leave();
-
 			#if UBA_TRACK_WORK
 			TrackWorkScope tws(*this, work.desc);
 			#else
@@ -238,11 +315,10 @@ namespace uba
 
 		while (true)
 		{
-			SCOPED_FUTEX_READ(m_workLock, lock);
-			bool workEmpty = m_work.empty();
-			lock.Leave();
-			if (workEmpty)
+			SCOPED_FUTEX_READ(m_workQueue.lock, lock);
+			if (m_workQueue.IsEmptyNoLock())
 				break;
+			lock.Leave();
 			if (hasTimedOut())
 				return false;
 			Sleep(5);
diff --git a/Engine/Source/Programs/UnrealBuildAccelerator/Common/Public/UbaNetworkServer.h b/Engine/Source/Programs/UnrealBuildAccelerator/Common/Public/UbaNetworkServer.h
--- a/Engine/Source/Programs/UnrealBuildAccelerator/Common/Public/UbaNetworkServer.h
+++ b/Engine/Source/Programs/UnrealBuildAccelerator/Common/Public/UbaNetworkServer.h
@@ -158,9 +158,7 @@ namespace uba
 
 		u32 m_maxWorkerCount = 0;
 
-		Futex m_additionalWorkLock;
-		struct AdditionalWork { WorkFunction func; Bottleneck* bottleneck; TString desc; };
-		List<AdditionalWork> m_additionalWork;
+		WorkQueue m_additionalWorkQueue;
 
 		Futex m_availableWorkersLock;
 		Worker* m_firstAvailableWorker = nullptr;
diff --git a/Engine/Source/Programs/UnrealBuildAccelerator/Common/Public/UbaWorkManager.h b/Engine/Source/Programs/UnrealBuildAccelerator/Common/Public/UbaWorkManager.h
--- a/Engine/Source/Programs/UnrealBuildAccelerator/Common/Public/UbaWorkManager.h
+++ b/Engine/Source/Programs/UnrealBuildAccelerator/Common/Public/UbaWorkManager.h
@@ -34,17 +34,18 @@ namespace uba
 		Thread* thread = nullptr;
 	};
 
+	using WorkFunction = Function<void(const WorkContext& context)>;
+
 	enum WorkPriority : u8
 	{
-		WorkPriority_Normal,
 		WorkPriority_High,
+		WorkPriority_Normal,
 		WorkPriority_Low,
 	};
 
 	class WorkManager : public WorkTracker
 	{
 	public:
-		using WorkFunction = Function<void(const WorkContext& context)>;
 		virtual void AddWork(const WorkFunction& work, u32 count, const tchar* desc, const Color& color = ColorWork, WorkPriority priority = WorkPriority_Normal, Bottleneck* bottleneck = nullptr) = 0;
 		virtual u32 GetWorkerCount() = 0;
 		virtual void DoWork(u32 count = 1) = 0;
@@ -64,6 +65,24 @@ namespace uba
 		Atomic<WorkTracker*> m_workTracker;
 	};
 
+	struct Work { WorkFunction func; Bottleneck* bottleneck; TString desc; Work* next; };
+
+	struct WorkQueue
+	{
+		~WorkQueue();
+		void Enqueue(const WorkFunction& func, u32 count, const tchar* desc, const Color& color, WorkPriority priority, Bottleneck* bottleneck, bool trackWork);
+		bool Dequeue(Work& out);
+		bool IsEmpty();
+
+		bool DequeueNoLock(Work& out);
+		bool IsEmptyNoLock();
+
+		Futex lock;
+		Work* first[3] = { 0 };
+		Work* last[3] = { 0 };
+		Work* free = nullptr;
+	};
+
 
 	class WorkManagerImpl : public WorkManager
 	{
@@ -83,9 +102,7 @@ namespace uba
 		u32 m_maxWorkerCount;
 		TString m_workerDesc;
 		List<Worker> m_workers;
-		struct Work { WorkFunction func; Bottleneck* bottleneck; TString desc; };
-		Futex m_workLock;
-		List<Work> m_work;
+		WorkQueue m_workQueue;
 		Atomic<u32> m_activeWorkerCount;
 
 		Futex m_availableWorkersLock;
