[UBA]
* Fixed bugs in work manager dequeue where bottleneck->leave was never called

diff --git a/Engine/Source/Programs/UnrealBuildAccelerator/Common/Private/UbaNetworkServer.cpp b/Engine/Source/Programs/UnrealBuildAccelerator/Common/Private/UbaNetworkServer.cpp
--- a/Engine/Source/Programs/UnrealBuildAccelerator/Common/Private/UbaNetworkServer.cpp
+++ b/Engine/Source/Programs/UnrealBuildAccelerator/Common/Private/UbaNetworkServer.cpp
@@ -733,20 +733,19 @@ namespace uba
 		{
 			while (true)
 			{
-				Work work;
-				if (!additionalWorkQueue.Dequeue(work))
+				Work* work = additionalWorkQueue.Dequeue(ticket);
+				if (!work)
 					break;
 
 				#if UBA_TRACK_WORK
-				TrackWorkScope tws(server, work.desc, ColorWork);
+				TrackWorkScope tws(server, work->desc, ColorWork);
 				#else
 				TrackWorkScope tws;
 				#endif
 
-				work.func({tws});
+				work->func({tws});
 
-				if (work.bottleneck)
-					work.bottleneck->Leave(ticket);
+				additionalWorkQueue.Finish(work, ticket);
 			}
 
 			// Both locks needs to be taken to verify if additional work
@@ -1173,15 +1172,19 @@ namespace uba
 
 	bool NetworkServer::DoAdditionalWork()
 	{
-		Work work;
-		if (m_additionalWorkQueue.Dequeue(work))
+		BottleneckTicket ticket;
+
+		Work* work = m_additionalWorkQueue.Dequeue(ticket);
+		if (work)
 		{
 			#if UBA_TRACK_WORK
-			TrackWorkScope tws(*this, work.desc, ColorWork);
+			TrackWorkScope tws(*this, work->desc, ColorWork);
 			#else
 			TrackWorkScope tws;
 			#endif
-			work.func({tws});
+			work->func({tws});
+
+			m_additionalWorkQueue.Finish(work, ticket);
 			return true;
 		}
 
diff --git a/Engine/Source/Programs/UnrealBuildAccelerator/Common/Private/UbaWorkManager.cpp b/Engine/Source/Programs/UnrealBuildAccelerator/Common/Private/UbaWorkManager.cpp
--- a/Engine/Source/Programs/UnrealBuildAccelerator/Common/Private/UbaWorkManager.cpp
+++ b/Engine/Source/Programs/UnrealBuildAccelerator/Common/Private/UbaWorkManager.cpp
@@ -118,20 +118,9 @@ namespace uba
 		}
 	}
 
-	bool WorkQueue::Dequeue(Work& out)
+	Work* WorkQueue::Dequeue(BottleneckTicket& ticket)
 	{
 		SCOPED_FUTEX(lock, l);
-		return DequeueNoLock(out);
-	}
-
-	bool WorkQueue::IsEmpty()
-	{
-		SCOPED_FUTEX(lock, l);
-		return IsEmptyNoLock();
-	}
-
-	bool WorkQueue::DequeueNoLock(Work& out)
-	{
 		for (u32 i=0; i != sizeof_array(first); ++i)
 		{
 			Work* p = nullptr;
@@ -158,13 +147,32 @@ namespace uba
 					first[i] = w->next;
 				}
 
-				w->next = free;
-				free = w;
-				out = *w;
-				return true;
+				return w;
 			}
 		}
-		return false;
+		return nullptr;
+	}
+
+	void WorkQueue::Finish(Work* work, BottleneckTicket& ticket)
+	{
+		Bottleneck* bottleneck = work->bottleneck;
+		work->bottleneck = nullptr;
+		work->func = {};
+		work->desc.clear();
+
+		SCOPED_FUTEX(lock, l);
+		work->next = free;
+		free = work;
+		l.Leave();
+
+		if (bottleneck)
+			bottleneck->Leave(ticket);
+	}
+
+	bool WorkQueue::IsEmpty()
+	{
+		SCOPED_FUTEX(lock, l);
+		return IsEmptyNoLock();
 	}
 
 	bool WorkQueue::IsEmptyNoLock()
@@ -210,19 +218,17 @@ namespace uba
 				{
 					while (true)
 					{
-						Work work;
-						if (!workQueue.Dequeue(work))
+						Work* work = workQueue.Dequeue(ticket);
+						if (!work)
 							break;
 
 						#if UBA_TRACK_WORK
-						TrackWorkScope tws(manager, work.desc);
+						TrackWorkScope tws(manager, work->desc);
 						#else
 						TrackWorkScope tws;
 						#endif
-						work.func({tws, &m_thread});
-
-						if (work.bottleneck)
-							work.bottleneck->Leave(ticket);
+						work->func({tws, &m_thread});
+						workQueue.Finish(work, ticket);
 					}
 
 					SCOPED_FUTEX(manager.m_availableWorkersLock, lock1);
@@ -294,17 +300,21 @@ namespace uba
 
 	void WorkManagerImpl::DoWork(u32 count)
 	{
+		BottleneckTicket ticket;
+
 		while (count--)
 		{
-			Work work;
-			if (!m_workQueue.Dequeue(work))
+			Work* work = m_workQueue.Dequeue(ticket);
+			if (!work)
 				break;
 			#if UBA_TRACK_WORK
-			TrackWorkScope tws(*this, work.desc);
+			TrackWorkScope tws(*this, work->desc);
 			#else
 			TrackWorkScope tws;
 			#endif
-			work.func({tws});
+			work->func({tws});
+
+			m_workQueue.Finish(work, ticket);
 		}
 	}
 
diff --git a/Engine/Source/Programs/UnrealBuildAccelerator/Common/Public/UbaBottleneck.h b/Engine/Source/Programs/UnrealBuildAccelerator/Common/Public/UbaBottleneck.h
--- a/Engine/Source/Programs/UnrealBuildAccelerator/Common/Public/UbaBottleneck.h
+++ b/Engine/Source/Programs/UnrealBuildAccelerator/Common/Public/UbaBottleneck.h
@@ -35,7 +35,7 @@ namespace uba
 			m_last = &ticket;
 
 			if (!ticket.ev.IsCreated())
-				ticket.ev.Create(true);
+				ticket.ev.Create(false);
 
 			lock.Leave();
 
diff --git a/Engine/Source/Programs/UnrealBuildAccelerator/Common/Public/UbaWorkManager.h b/Engine/Source/Programs/UnrealBuildAccelerator/Common/Public/UbaWorkManager.h
--- a/Engine/Source/Programs/UnrealBuildAccelerator/Common/Public/UbaWorkManager.h
+++ b/Engine/Source/Programs/UnrealBuildAccelerator/Common/Public/UbaWorkManager.h
@@ -12,6 +12,7 @@ namespace uba
 {
 	class Thread;
 	struct Bottleneck;
+	struct BottleneckTicket;
 	struct StringView;
 	struct TrackWorkScope;
 
@@ -71,10 +72,10 @@ namespace uba
 	{
 		~WorkQueue();
 		void Enqueue(const WorkFunction& func, u32 count, const tchar* desc, const Color& color, WorkPriority priority, Bottleneck* bottleneck, bool trackWork);
-		bool Dequeue(Work& out);
+		Work* Dequeue(BottleneckTicket& ticket);
+		void Finish(Work* work, BottleneckTicket& ticket);
 		bool IsEmpty();
 
-		bool DequeueNoLock(Work& out);
 		bool IsEmptyNoLock();
 
 		Futex lock;
