[UBA]
* Moved lambda out to be a normal function. Has this feeling codegen is somehow wrong on linux.. with this change I can't repro a hang that happened 9/10 times when building on the farm
* Added so UbaCli -checkCas can be combined with -deleteCas and then it will only delete bad cas files

diff --git a/Engine/Source/Programs/UnrealBuildAccelerator/Cli/Private/UbaCli.cpp b/Engine/Source/Programs/UnrealBuildAccelerator/Cli/Private/UbaCli.cpp
--- a/Engine/Source/Programs/UnrealBuildAccelerator/Cli/Private/UbaCli.cpp
+++ b/Engine/Source/Programs/UnrealBuildAccelerator/Cli/Private/UbaCli.cpp
@@ -81,8 +81,8 @@ namespace uba
 		logger.Info(TC("   -workdir=<dir>          Working directory"));
 		logger.Info(TC("   -config=<file>          Config file that contains options for various systems"));
 		logger.Info(TC("   -vfs=<virtual>;<local>  Will convert virtual path to local under the hood. Can have multiple -vfs"));
-		logger.Info(TC("   -checkcas               Check so all cas entries are correct"));
-		logger.Info(TC("   -checkfiletable         Check so file table has correct cas stored"));
+		logger.Info(TC("   -checkcas               Check so all cas entries are correct. Combine with -deleteCas to delete bad ones"));
+		logger.Info(TC("   -checkfiletable         Check so file table has correct cas stored. if bad entries it is enough to delete 'casdb' file"));
 		logger.Info(TC("   -checkcloud             Check if we are inside cloud and output information about cloud"));
 		logger.Info(TC("   -deletecas              Deletes the casdb"));
 		logger.Info(TC("   -getcas                 Will print hash of application"));
@@ -595,18 +595,6 @@ namespace uba
 				logger.Warning(TC("  Running x64 binary on arm64 system. Use arm binaries instead"));
 
 		bool exit = false;
-		if (deleteCas)
-		{
-			WorkManagerImpl workManager(maxProcessCount, TC("UbaWrk/DelCas"));
-			StorageImpl(StorageCreateInfo(g_rootDir.data, logWriter, workManager)).DeleteAllCas();
-			for (u32 i=0; i!=agentCount; ++i)
-			{
-				StringBuffer<> clientRootDir;
-				clientRootDir.Append(g_rootDir).Append("Agent").AppendValue(i);
-				StorageImpl(StorageCreateInfo(clientRootDir.data, logWriter, workManager)).DeleteAllCas();
-			}
-			exit = true;
-		}
 
 		if (!addCas.empty())
 		{
@@ -628,8 +616,30 @@ namespace uba
 			storageInfo.casCapacityBytes = 0;
 			storageInfo.storeCompressed = storeCompressed;
 			StorageImpl storage(storageInfo);
-			if (!storage.CheckCasContent())
+			Vector<TString> badFiles;
+			if (!storage.CheckCasContent(&badFiles))
 				return false;
+
+			if (deleteCas)
+			{
+				for (auto& badFile : badFiles)
+					DeleteFileW(badFile.c_str());
+				deleteCas = false;
+			}
+
+			exit = true;
+		}
+
+		if (deleteCas)
+		{
+			WorkManagerImpl workManager(maxProcessCount, TC("UbaWrk/DelCas"));
+			StorageImpl(StorageCreateInfo(g_rootDir.data, logWriter, workManager)).DeleteAllCas();
+			for (u32 i=0; i!=agentCount; ++i)
+			{
+				StringBuffer<> clientRootDir;
+				clientRootDir.Append(g_rootDir).Append("Agent").AppendValue(i);
+				StorageImpl(StorageCreateInfo(clientRootDir.data, logWriter, workManager)).DeleteAllCas();
+			}
 			exit = true;
 		}
 
diff --git a/Engine/Source/Programs/UnrealBuildAccelerator/Common/Private/UbaStorage.cpp b/Engine/Source/Programs/UnrealBuildAccelerator/Common/Private/UbaStorage.cpp
--- a/Engine/Source/Programs/UnrealBuildAccelerator/Common/Private/UbaStorage.cpp
+++ b/Engine/Source/Programs/UnrealBuildAccelerator/Common/Private/UbaStorage.cpp
@@ -244,9 +244,7 @@ namespace uba
 		UBA_ASSERT(!m_newestAccessed || !m_newestAccessed->prevAccessed);
 		UBA_ASSERT(!m_oldestAccessed || !m_oldestAccessed->nextAccessed);
 
-		struct Rec { CasEntry& casEntry; u64 size; };
-
-		List<Vector<Rec>> toDelete;
+		List<Vector<DeleteRec>> toDelete;
 
 		{
 			//TrackHintScope ths(*TrackWorkScope::GetActiveScope(), TCV("DetachCasEntries"));
@@ -271,8 +269,8 @@ namespace uba
 				}
 
 				UBA_ASSERT(ce.key != CasKeyZero);
-				Vector<Rec>& deleteBatch = (toDelete.empty() || toDelete.back().size() == m_deleteCasBatchSize) ? toDelete.emplace_back() : toDelete.back();
-				deleteBatch.push_back({*it, ce.size});
+				Vector<DeleteRec>& deleteBatch = (toDelete.empty() || toDelete.back().size() == m_deleteCasBatchSize) ? toDelete.emplace_back() : toDelete.back();
+				deleteBatch.emplace_back(*it, ce.size);
 
 				if (m_trackedDeletes)
 					m_trackedDeletes->insert(ce.key);
@@ -298,42 +296,10 @@ namespace uba
 		}
 		lock.Leave();
 
-		static auto deleteBatch = [](StorageImpl& storage, const Vector<Rec>& toDelete)
-			{
-				for (const Rec& rec : toDelete)
-				{
-					StringBuffer<> casFile;
-					storage.StorageImpl::GetCasFileName(casFile, rec.casEntry.key);
-
-					if (!DeleteFileW(casFile.data))
-					{
-						u32 error = GetLastError();
-						if (error != ERROR_FILE_NOT_FOUND && error != ERROR_PATH_NOT_FOUND)
-						{
-							storage.m_logger.Error(TC("Failed to delete %s while handling overflow (%s)"), casFile.data, LastErrorToText(error).data);
-							rec.casEntry.exists = true;
-							rec.casEntry.size = rec.size;
-							rec.casEntry.internalLock.Leave();
-
-							// TODO: Should this instead set some overflow
-							/*
-							SCOPED_FUTEX(m_accessLock, lock);
-							m_casEvictedBytes -= rec.size;
-							--m_casEvictedCount;
-							m_casTotalBytes += rec.size;
-							AttachEntry(&rec.casEntry); // TODO: This should be re-added in the end
-							*/
-							continue;
-						}
-					}
-
-					//m_logger.Warning(TC("Evicted %s due to overflow"), casFile.data);
-
-					rec.casEntry.verified = true; // Verified to be deleted
-
-					rec.casEntry.internalLock.Leave();
-				}
-			};
+#if 0//!PLATFORM_WINDOWS
+		for (auto& batch : toDelete)
+			DeleteBatch(batch);
+#else
 
 		// TODO: Defer deletes can cause deadlocks if all active jobs are sitting on locks of the cas that is being deleted
 		// .. since they are now in a locked state but queued... need to make a solution for this
@@ -342,14 +308,52 @@ namespace uba
 		if (deferDeletes)
 		{
 			for (auto& batch : toDelete)
-				m_workManager.AddWork([toDelete = std::move(batch), this](const WorkContext&) { deleteBatch(*this, toDelete); }, 1, TC("DeleteCasFiles"));
+				m_workManager.AddWork([toDelete = std::move(batch), this](const WorkContext&) { DeleteBatch(toDelete); }, 1, TC("DeleteCasFiles"));
 		}
 		else
 		{
 			m_workManager.ParallelFor(u32(toDelete.size()), toDelete, [this](const WorkContext&, auto& it)
 				{
-					deleteBatch(*this, *it);
-				}, TCV("DeleteCasFiles"));
+					DeleteBatch(*it);
+				}, TCV("DeleteCasFiles"), WorkPriority_High);
+		}
+#endif
+	}
+
+	void StorageImpl::DeleteBatch(const Vector<DeleteRec>& toDelete)
+	{
+		for (const DeleteRec& rec : toDelete)
+		{
+			StringBuffer<> casFile;
+			StorageImpl::GetCasFileName(casFile, rec.casEntry.key);
+
+			if (!DeleteFileW(casFile.data))
+			{
+				u32 error = GetLastError();
+				if (error != ERROR_FILE_NOT_FOUND && error != ERROR_PATH_NOT_FOUND)
+				{
+					m_logger.Error(TC("Failed to delete %s while handling overflow (%s)"), casFile.data, LastErrorToText(error).data);
+					rec.casEntry.exists = true;
+					rec.casEntry.size = rec.size;
+					rec.casEntry.internalLock.Leave();
+
+					// TODO: Should this instead set some overflow
+					/*
+					SCOPED_FUTEX(m_accessLock, lock);
+					m_casEvictedBytes -= rec.size;
+					--m_casEvictedCount;
+					m_casTotalBytes += rec.size;
+					AttachEntry(&rec.casEntry); // TODO: This should be re-added in the end
+					*/
+					continue;
+				}
+			}
+
+			//m_logger.Warning(TC("Evicted %s due to overflow"), casFile.data);
+
+			rec.casEntry.verified = true; // Verified to be deleted
+
+			rec.casEntry.internalLock.Leave();
 		}
 	}
 
@@ -1388,7 +1392,7 @@ namespace uba
 		return true;
 	}
 
-	bool StorageImpl::CheckCasContent()
+	bool StorageImpl::CheckCasContent(Vector<TString>* badFiles)
 	{
 		StringBuffer<> casRoot;
 		casRoot.Append(m_rootDir.data, m_rootDir.count - 1);
@@ -1435,6 +1439,8 @@ namespace uba
 								++errorCount;
 								if (lastWritten > newestWrittenError)
 									newestWrittenError = lastWritten;
+								if (badFiles)
+									badFiles->push_back(filePath);
 							});
 
 						CasKey checkedKey = EmptyFileKey;
@@ -1467,7 +1473,7 @@ namespace uba
 							return;
 						}
 						m_logger.Error(TC("Cas key and content mismatch for key %s (expected %s) (%s ago)"), CasKeyString(casKey).str, CasKeyString(checkedKey).str, timeStr.data);
-						
+
 					}, 1, TC("CheckCasContent"));
 			});
 		m_logger.Info(TC("Validating %u entries..."), entryCount);
diff --git a/Engine/Source/Programs/UnrealBuildAccelerator/Common/Public/UbaStorage.h b/Engine/Source/Programs/UnrealBuildAccelerator/Common/Public/UbaStorage.h
--- a/Engine/Source/Programs/UnrealBuildAccelerator/Common/Public/UbaStorage.h
+++ b/Engine/Source/Programs/UnrealBuildAccelerator/Common/Public/UbaStorage.h
@@ -141,7 +141,7 @@ namespace uba
 		virtual ~StorageImpl();
 
 		bool LoadCasTable(bool logStats = true, bool alwaysCheckAllFiles = false, bool* outWasTerminated = nullptr);
-		bool CheckCasContent();
+		bool CheckCasContent(Vector<TString>* badFiles = nullptr);
 		bool CheckFileTable(const tchar* searchPath);
 		const tchar* GetTempPath();
 
@@ -212,6 +212,9 @@ namespace uba
 		void TraverseAllCasFiles(const Function<void(const CasKey& key, u64 size)>& func, bool allowParallel = false);
 		bool CheckAllCasFiles(u64 checkContentOfFilesNewerThanTime = ~u64(0));
 		void HandleOverflow(UnorderedSet<CasKey>* outDeletedFiles);
+		
+		struct DeleteRec { CasEntry& casEntry; u64 size; };
+		void DeleteBatch(const Vector<DeleteRec>& toDelete);
 
 		static MutexHandle GetExclusiveAccess(Logger& logger, const StringView& rootDir, bool reportError);
 		bool DeleteIsRunningFile();
