[UBA]
* Changed how cas deletion on overflow works. When overflow is discovered we estimate how many files that needs to be deleted. We then kick jobs with batches of deletes given of how many files we think we need to delete. Since the sorting changes over time our estimates might be wrong but will fix itself up over time.

diff --git a/Engine/Source/Programs/UnrealBuildAccelerator/Common/Private/UbaStorage.cpp b/Engine/Source/Programs/UnrealBuildAccelerator/Common/Private/UbaStorage.cpp
--- a/Engine/Source/Programs/UnrealBuildAccelerator/Common/Private/UbaStorage.cpp
+++ b/Engine/Source/Programs/UnrealBuildAccelerator/Common/Private/UbaStorage.cpp
@@ -230,7 +230,8 @@ namespace uba
 		if (!m_casCapacityBytes || m_overflowReported || m_manuallyHandleOverflow)
 			return;
 
-		if (m_casTotalBytes <= m_casCapacityBytes)
+		u64 casTotalBytes = m_casTotalBytes - m_casQueuedDeleteBytes;
+		if (casTotalBytes <= m_casCapacityBytes)
 		{
 			if (!casEntry.prevAccessed && !casEntry.nextAccessed && (!m_newestAccessed || m_newestAccessed != &casEntry))
 			{
@@ -240,92 +241,97 @@ namespace uba
 			return;
 		}
 
-		TimerScope ts(Stats().handleOverflow);
-		UBA_ASSERT(!m_newestAccessed || !m_newestAccessed->prevAccessed);
-		UBA_ASSERT(!m_oldestAccessed || !m_oldestAccessed->nextAccessed);
+		// Locking cas entries and then create work to delete them does not work since it can cause deadlocks
+		// We need to estimate approximately how many files we need to delete and then kick the work
+		u32 deleteCountNeeded = 0;
+		for (CasEntry* it = m_oldestAccessed; it; it = it->prevAccessed)
+		{
+			++deleteCountNeeded;
+			casTotalBytes -= it->size;
+			m_casQueuedDeleteBytes += it->size;
+			if (casTotalBytes <= m_casCapacityBytes)
+				break;
+		}
 
-		List<Vector<DeleteRec>> toDelete;
+		lock.Leave();
 
+		u32 workCount = (deleteCountNeeded + m_deleteCasBatchSize-1) / m_deleteCasBatchSize;
+
+		if (m_allowDeferredDeletes)
 		{
-			//TrackHintScope ths(*TrackWorkScope::GetActiveScope(), TCV("DetachCasEntries"));
-			for (CasEntry* it = m_oldestAccessed; it;)
-			{
-				CasEntry& ce = *it;
-				ce.internalLock.Enter();
+			m_workManager.AddWork([this](const WorkContext&) { DeleteOldestCasEntries(); }, workCount, TC("DeleteCasFiles"));
+		}
+		else
+		{
+			while (workCount--)
+				DeleteOldestCasEntries();
+		}
+	}
 
-				if (ce.verified && !m_allowDeleteVerified) // Can't remove these since they might be in mapping tables and actively being used. It is tempting to delete dropped cas files but we can't
-				{
-					ce.internalLock.Leave();
-					break;
-				}
+	void StorageImpl::DeleteOldestCasEntries()
+	{
+		struct DeleteRec { CasEntry& casEntry; u64 size; };
+		List<DeleteRec> toDelete;
+		u32 left = m_deleteCasBatchSize;
 
-				if (ce.beingWritten || ce.refCount)
-				{
-					if (ce.beingWritten)
-						m_logger.Info(TC("Want to delete file that is begin written. We should never end up here"));
-					ce.internalLock.Leave();
-					it = ce.prevAccessed;
-					continue;
-				}
+		SCOPED_FUTEX(m_accessLock, lock);
+		for (CasEntry* it = m_oldestAccessed; it;)
+		{
+			CasEntry& ce = *it;
+			ce.internalLock.Enter();
 
-				UBA_ASSERT(ce.key != CasKeyZero);
-				Vector<DeleteRec>& deleteBatch = (toDelete.empty() || toDelete.back().size() == m_deleteCasBatchSize) ? toDelete.emplace_back() : toDelete.back();
-				deleteBatch.emplace_back(*it, ce.size);
+			if (ce.verified && !m_allowDeleteVerified) // Can't remove these since they might be in mapping tables and actively being used. It is tempting to delete dropped cas files but we can't
+			{
+				ce.internalLock.Leave();
+				break;
+			}
 
-				if (m_trackedDeletes)
-					m_trackedDeletes->insert(ce.key);
+			if (ce.beingWritten || ce.refCount)
+			{
+				if (ce.beingWritten)
+					m_logger.Info(TC("Want to delete file that is begin written. We should never end up here"));
+				ce.internalLock.Leave();
+				it = ce.prevAccessed;
+				continue;
+			}
 
-				m_casEvictedBytes += ce.size;
-				m_casEvictedCount++;
-				m_casTotalBytes -= ce.size;
+			UBA_ASSERT(ce.key != CasKeyZero);
+			toDelete.emplace_back(ce, ce.size);
 
-				ce.exists = false;
-				ce.size = 0;
+			if (m_trackedDeletes)
+				m_trackedDeletes->insert(ce.key);
 
+			m_casEvictedBytes += ce.size;
+			++m_casEvictedCount;
+			m_casTotalBytes -= ce.size;
+			if (ce.size <= m_casQueuedDeleteBytes)
+				m_casQueuedDeleteBytes -= ce.size;
+			else
+				m_casQueuedDeleteBytes = 0;
 
-				it = ce.prevAccessed;
-				DetachEntry(ce);
-				if (m_casTotalBytes <= m_casCapacityBytes)
-					break;
-			}
+			ce.exists = false;
+			ce.size = 0;
+
+			it = ce.prevAccessed;
+			DetachEntry(ce);
+
+			if (!--left)
+				break;
 		}
-		if (m_casTotalBytes > m_casCapacityBytes)
+
+		if (left != 0 && !m_overflowReported)
 		{
 			m_overflowReported = true;
 			m_logger.Info(TC("Exceeding maximum size set for cas (%s). Current session needs more storage to be able to finish (will now overflow). User memory reported on session exit"), BytesToText(m_casCapacityBytes).str);
 		}
-		lock.Leave();
-
-#if 0//!PLATFORM_WINDOWS
-		for (auto& batch : toDelete)
-			DeleteBatch(batch);
-#else
-
-		// TODO: Defer deletes can cause deadlocks if all active jobs are sitting on locks of the cas that is being deleted
-		// .. since they are now in a locked state but queued... need to make a solution for this
-		bool deferDeletes = m_allowDeferredDeletes;
 
-		if (deferDeletes)
-		{
-			for (auto& batch : toDelete)
-				m_workManager.AddWork([toDelete = std::move(batch), this](const WorkContext&) { DeleteBatch(toDelete); }, 1, TC("DeleteCasFiles"));
-		}
-		else
-		{
-			m_workManager.ParallelFor(u32(toDelete.size()), toDelete, [this](const WorkContext&, auto& it)
-				{
-					DeleteBatch(*it);
-				}, TCV("DeleteCasFiles"), WorkPriority_High);
-		}
-#endif
-	}
+		lock.Leave();
 
-	void StorageImpl::DeleteBatch(const Vector<DeleteRec>& toDelete)
-	{
-		for (const DeleteRec& rec : toDelete)
+		for (DeleteRec& rec : toDelete)
 		{
+			CasEntry& ce = rec.casEntry;
 			StringBuffer<> casFile;
-			StorageImpl::GetCasFileName(casFile, rec.casEntry.key);
+			StorageImpl::GetCasFileName(casFile, ce.key);
 
 			if (!DeleteFileW(casFile.data))
 			{
@@ -333,27 +339,14 @@ namespace uba
 				if (error != ERROR_FILE_NOT_FOUND && error != ERROR_PATH_NOT_FOUND)
 				{
 					m_logger.Error(TC("Failed to delete %s while handling overflow (%s)"), casFile.data, LastErrorToText(error).data);
-					rec.casEntry.exists = true;
-					rec.casEntry.size = rec.size;
-					rec.casEntry.internalLock.Leave();
-
-					// TODO: Should this instead set some overflow
-					/*
-					SCOPED_FUTEX(m_accessLock, lock);
-					m_casEvictedBytes -= rec.size;
-					--m_casEvictedCount;
-					m_casTotalBytes += rec.size;
-					AttachEntry(&rec.casEntry); // TODO: This should be re-added in the end
-					*/
+					ce.exists = true;
+					ce.size = rec.size;
+					ce.internalLock.Leave();
 					continue;
 				}
 			}
-
-			//m_logger.Warning(TC("Evicted %s due to overflow"), casFile.data);
-
-			rec.casEntry.verified = true; // Verified to be deleted
-
-			rec.casEntry.internalLock.Leave();
+			ce.verified = true; // Verified to be deleted
+			ce.internalLock.Leave();
 		}
 	}
 
diff --git a/Engine/Source/Programs/UnrealBuildAccelerator/Common/Public/UbaStorage.h b/Engine/Source/Programs/UnrealBuildAccelerator/Common/Public/UbaStorage.h
--- a/Engine/Source/Programs/UnrealBuildAccelerator/Common/Public/UbaStorage.h
+++ b/Engine/Source/Programs/UnrealBuildAccelerator/Common/Public/UbaStorage.h
@@ -114,7 +114,7 @@ namespace uba
 		bool writeToDisk = true;
 		bool keepTransientDataMapped = true;	// Will keep the transient data mapped in instead of doing map/unmap everytime we access it
 		bool allowDeleteVerified = false;
-		bool allowDeferredDeletes = false;
+		bool allowDeferredDeletes = true;
 		MutexHandle exclusiveMutex = InvalidMutexHandle; // Set this to hand over pre-exclusive access to storage. Storage will release mutex on shutdown
 		u8 casCompressor = 0;
 		u8 casCompressionLevel = 0;
@@ -212,10 +212,8 @@ namespace uba
 		void TraverseAllCasFiles(const Function<void(const CasKey& key, u64 size)>& func, bool allowParallel = false);
 		bool CheckAllCasFiles(u64 checkContentOfFilesNewerThanTime = ~u64(0));
 		void HandleOverflow(UnorderedSet<CasKey>* outDeletedFiles);
+		void DeleteOldestCasEntries();
 		
-		struct DeleteRec { CasEntry& casEntry; u64 size; };
-		void DeleteBatch(const Vector<DeleteRec>& toDelete);
-
 		static MutexHandle GetExclusiveAccess(Logger& logger, const StringView& rootDir, bool reportError);
 		bool DeleteIsRunningFile();
 		bool TouchIsRunningFile();
@@ -295,6 +293,7 @@ namespace uba
 		CasEntry* m_newestAccessed = nullptr;
 		CasEntry* m_oldestAccessed = nullptr;
 		u64 m_casTotalBytes = 0;
+		u64 m_casQueuedDeleteBytes = 0;
 		u64 m_casMaxBytes = 0;
 		u64 m_casCapacityBytes = 0;
 		u64 m_casEvictedBytes = 0;
