[UBA]
* Fixed unlikely bug where cas entry could be garbage collected while being sent to cache server from client

diff --git a/Engine/Source/Programs/UnrealBuildAccelerator/Common/Private/UbaCacheClient.cpp b/Engine/Source/Programs/UnrealBuildAccelerator/Common/Private/UbaCacheClient.cpp
--- a/Engine/Source/Programs/UnrealBuildAccelerator/Common/Private/UbaCacheClient.cpp
+++ b/Engine/Source/Programs/UnrealBuildAccelerator/Common/Private/UbaCacheClient.cpp
@@ -1319,6 +1319,12 @@ namespace uba
 			{
 				createCasTs.Leave();
 
+				auto refCountGuard = MakeGuard([&]()
+					{
+						StorageImpl::CasEntryWriteLock casEntryLock(*casEntry);
+						--casEntry->refCount; // HasCasFile increase refcount
+					});
+
 				UBA_TRACE_WRITE_HINT("SendFile");
 				TimerScope ts(stats.sendFile);
 
@@ -1335,20 +1341,16 @@ namespace uba
 
 				FileAccessor file(m_logger, casKeyFileName.data);
 
-				StorageImpl::CasEntryWriteLock casEntryLock(*casEntry);
-				--casEntry->refCount; // HasCasFile increase refcount
 				auto mh = casEntry->mappingHandle;
 				if (mh.IsValid()) // If file was created by helper it will be in the transient mapped memory
 				{
 					u64 mappingOffset = casEntry->mappingOffset;
 					u64 mappingSize = casEntry->mappingSize;
-					casEntryLock.Leave();
 
-					mappedView = m_storage.m_casDataBuffer.MapView(mh, mappingOffset, mappingSize, path.data);
+					mappedView = m_storage.m_casDataBuffer.MapView(mh, mappingOffset, mappingSize, path.data, casEntry->mappingIsIndependent);
 					fileData = mappedView.memory;
 					fileSize = mappedView.size;
 				}
-				casEntryLock.Leave();
 
 				if (!fileData)
 				{
