[UBA]
* Added support for "vfs" in actions yaml file

diff --git a/Engine/Source/Programs/UnrealBuildAccelerator/Common/Private/UbaScheduler.cpp b/Engine/Source/Programs/UnrealBuildAccelerator/Common/Private/UbaScheduler.cpp
--- a/Engine/Source/Programs/UnrealBuildAccelerator/Common/Private/UbaScheduler.cpp
+++ b/Engine/Source/Programs/UnrealBuildAccelerator/Common/Private/UbaScheduler.cpp
@@ -1645,6 +1645,7 @@ namespace uba
 		u32 cacheBucket = 0;
 		u32 memoryGroup = 0;
 		Vector<u32> deps;
+		RootsHandle rootsHandle = 0;
 
 		ProcessStartInfo si;
 
@@ -1654,6 +1655,7 @@ namespace uba
 				si.arguments = arg.c_str();
 				si.workingDir = dir.c_str();
 				si.description = desc.c_str();
+				si.rootsHandle = rootsHandle;
 
 				#if UBA_DEBUG
 				StringBuffer<> logFile;
@@ -1687,6 +1689,7 @@ namespace uba
 				weight = 1.0f;
 				cacheBucket = 0;
 				memoryGroup = 0;
+				rootsHandle = 0;
 			};
 
 		enum InsideArray
@@ -1814,6 +1817,33 @@ namespace uba
 							depStart = it + 1;
 						}
 					}
+					else if (key.Equals(TCV("vfs")))
+					{
+						const tchar* depStart = TStrchr(valueStart, '[');
+						if (!depStart)
+							return false;
+						const tchar* it = depStart + 1;
+						StackBinaryWriter<8*1024> buffer;
+						while (true)
+						{
+							const tchar* vfsStart = it;
+							while (*it != ';')
+								++it;
+							buffer.WriteByte(0);
+							buffer.WriteString(vfsStart, it - vfsStart);
+							++it;
+							const tchar* localStart = it;
+							while (*it != ',' && *it != ']')
+								++it;
+							buffer.WriteString(localStart, it - localStart);
+							if (*it == ']')
+								break;
+							++it;
+							if (*it == ' ')
+								++it;
+						}
+						rootsHandle = m_session.RegisterRoots(buffer.GetData(), buffer.GetPosition());
+					}
 					return true;
 				}
 				}
