UnrealBuildAccellerator: Add some application rules

#rnx

diff --git a/Engine/Source/Programs/UnrealBuildAccelerator/Common/Private/UbaSession.cpp b/Engine/Source/Programs/UnrealBuildAccelerator/Common/Private/UbaSession.cpp
--- a/Engine/Source/Programs/UnrealBuildAccelerator/Common/Private/UbaSession.cpp
+++ b/Engine/Source/Programs/UnrealBuildAccelerator/Common/Private/UbaSession.cpp
@@ -2199,7 +2199,7 @@ namespace uba
 				return rules[i].rules;
 			}
 
-			if (!exeName.Equals(TCV("dotnet.exe")))
+			if (!exeName.Equals(TCV("dotnet.exe")) && !exeName.Equals(TCV("dotnet")))
 				return rules[isDotnet].rules;
 			
 			isDotnet = true;
diff --git a/Engine/Source/Programs/UnrealBuildAccelerator/Core/Private/UbaApplicationRules.cpp b/Engine/Source/Programs/UnrealBuildAccelerator/Core/Private/UbaApplicationRules.cpp
--- a/Engine/Source/Programs/UnrealBuildAccelerator/Core/Private/UbaApplicationRules.cpp
+++ b/Engine/Source/Programs/UnrealBuildAccelerator/Core/Private/UbaApplicationRules.cpp
@@ -782,6 +782,7 @@ namespace uba
 			{ 1752955744,										new ApplicationRulesLdExePlatform1() },
 			{ 238360161,										new ApplicationRulesClangPlusPlusExePlatform2() },
 			{ 2119756440,										new ApplicationRulesLldExePlatform2() },
+			{ 3206998955,										new ApplicationRulesISPCExe() },
 			{ 2898035017,										new ApplicationRulesStageSplitMergeApplication() },
 			{ ApplicationHash(TC("dump_syms.exe")),				new ApplicationRulesDumpSymsExe() },
 			{ ApplicationHash(TC("ld.lld.exe")),				new ApplicationRulesLdLLdExe() },
@@ -805,6 +806,10 @@ namespace uba
 			{ ApplicationHash(TC("ispc")),						new ApplicationRulesISPCExe() },
 			{ ApplicationHash(TC("shadercompileworker")),		new ApplicationRulesShaderCompileWorker() },
 			{ ApplicationHash(TC("ubatestapp")),				new ApplicationRulesUbaTestApp() },
+			{ 851054321,										new ApplicationRulesClangPlusPlusExePlatform2() },
+			{ 138829736,										new ApplicationRulesLldExePlatform2() },
+			{ 286314299,										new ApplicationRulesISPCExe() },
+			{ 3836757465,										new ApplicationRulesStageSplitMergeApplication() },
 #endif
 			{ 0, nullptr }
 		};
