UnrealBuildAccelerator: call SetThreadErrorMode when attaching to a process or thread

#rnx

diff --git a/Engine/Source/Programs/UnrealBuildAccelerator/Detours/Private/Windows/UbaDetoursMainWin.cpp b/Engine/Source/Programs/UnrealBuildAccelerator/Detours/Private/Windows/UbaDetoursMainWin.cpp
--- a/Engine/Source/Programs/UnrealBuildAccelerator/Detours/Private/Windows/UbaDetoursMainWin.cpp
+++ b/Engine/Source/Programs/UnrealBuildAccelerator/Detours/Private/Windows/UbaDetoursMainWin.cpp
@@ -154,6 +154,7 @@ BOOL WINAPI DllMain(HINSTANCE hinst, DWORD dwReason, LPVOID reserved)
 
 	if (dwReason == DLL_PROCESS_ATTACH)
 	{
+		SetThreadErrorMode(GetThreadErrorMode() | SEM_FAILCRITICALERRORS | SEM_NOOPENFILEERRORBOX, nullptr);
 		GetThreadGroupAffinity(GetCurrentThread(), &GroupAffinity);
 
 		if (!DetourRestoreAfterWith())
@@ -178,6 +179,7 @@ BOOL WINAPI DllMain(HINSTANCE hinst, DWORD dwReason, LPVOID reserved)
 	}
 	else if (dwReason == DLL_THREAD_ATTACH)
 	{
+		SetThreadErrorMode(GetThreadErrorMode() | SEM_FAILCRITICALERRORS | SEM_NOOPENFILEERRORBOX, nullptr);
 		SetThreadGroupAffinity(GetCurrentThread(), &GroupAffinity, NULL);
 	}
 	else if (dwReason == DLL_PROCESS_DETACH)
