[UBA]
* Fixed UbaCoordinatorHorde target. It can now compiled and be used together with ubacli/ubahost

diff --git a/Engine/Source/Developer/UbaCoordinatorHorde/Private/UbaCoordinatorHorde.cpp b/Engine/Source/Developer/UbaCoordinatorHorde/Private/UbaCoordinatorHorde.cpp
--- a/Engine/Source/Developer/UbaCoordinatorHorde/Private/UbaCoordinatorHorde.cpp
+++ b/Engine/Source/Developer/UbaCoordinatorHorde/Private/UbaCoordinatorHorde.cpp
@@ -79,8 +79,8 @@ extern "C"
 
 		auto coordinator = new CoordinatorHorde(info.workDir, info.binariesDir);
 		auto& m = coordinator->m_manager;
-		m.SetPool(info.pool);
-		m.SetMaxCoreCount(info.maxCoreCount);
+		//m.SetPool(info.pool);
+		m.SetTargetCoreCount(info.maxCoreCount);
 
 		return coordinator;
 	}
diff --git a/Engine/Source/Programs/UnrealBuildAccelerator/Common/Public/UbaCoordinator.h b/Engine/Source/Programs/UnrealBuildAccelerator/Common/Public/UbaCoordinator.h
--- a/Engine/Source/Programs/UnrealBuildAccelerator/Common/Public/UbaCoordinator.h
+++ b/Engine/Source/Programs/UnrealBuildAccelerator/Common/Public/UbaCoordinator.h
@@ -19,7 +19,7 @@ namespace uba
 	class Coordinator
 	{
 	public:
-		using AddClientCallback = bool(void* userData, const tchar* ip, u16 port);
+		using AddClientCallback = bool(void* userData, const tchar* ip, u16 port, const tchar* crypto);
 		virtual void SetAddClientCallback(AddClientCallback* callback, void* userData) = 0;
 		virtual void SetTargetCoreCount(u32 count) = 0;
 	};
diff --git a/Engine/Source/Programs/UnrealBuildAccelerator/Common/Public/UbaCoordinatorWrapper.h b/Engine/Source/Programs/UnrealBuildAccelerator/Common/Public/UbaCoordinatorWrapper.h
--- a/Engine/Source/Programs/UnrealBuildAccelerator/Common/Public/UbaCoordinatorWrapper.h
+++ b/Engine/Source/Programs/UnrealBuildAccelerator/Common/Public/UbaCoordinatorWrapper.h
@@ -68,7 +68,7 @@ namespace uba
 
 		void ThreadUpdate(u32 maxCoreCount)
 		{
-			m_coordinator->SetAddClientCallback([](void* userData, const tchar* ip, u16 port)
+			m_coordinator->SetAddClientCallback([](void* userData, const tchar* ip, u16 port, const tchar* crypto)
 				{
 					auto& cw = *(CoordinatorWrapper*)userData;
 					return cw.m_networkServer->AddClient(*cw.m_networkBackend, ip, port);
diff --git a/Engine/Source/Programs/UnrealBuildAccelerator/UbaCoordinatorHorde.Target.cs b/Engine/Source/Programs/UnrealBuildAccelerator/UbaCoordinatorHorde.Target.cs
--- a/Engine/Source/Programs/UnrealBuildAccelerator/UbaCoordinatorHorde.Target.cs
+++ b/Engine/Source/Programs/UnrealBuildAccelerator/UbaCoordinatorHorde.Target.cs
@@ -66,7 +66,7 @@ public class UbaCoordinatorHordeTarget : TargetRules
 		GlobalDefinitions.Add("ALLOW_HITCH_DETECTION=0");
 		GlobalDefinitions.Add("UE_USE_MALLOC_FILL_BYTES=0");
 		GlobalDefinitions.Add("DISABLE_CWD_CHANGES=1");
-		GlobalDefinitions.Add("UE_NO_ENGINE_OIDC=1");
+		GlobalDefinitions.Add("UE_NO_ENGINE_OIDC=0");
 
 		GlobalDefinitions.Add("UBA_COORDINATOR_HORDE_DLL");
 	}
