[UBA]
* Added UbaCli.exe -decompress=<file> which will decompress a compressed file in-place

diff --git a/Engine/Source/Programs/UnrealBuildAccelerator/Cli/Private/UbaCli.cpp b/Engine/Source/Programs/UnrealBuildAccelerator/Cli/Private/UbaCli.cpp
--- a/Engine/Source/Programs/UnrealBuildAccelerator/Cli/Private/UbaCli.cpp
+++ b/Engine/Source/Programs/UnrealBuildAccelerator/Cli/Private/UbaCli.cpp
@@ -101,6 +101,7 @@ namespace uba
 		logger.Info(TC("   -populatecache          Populate cache server if connected to one"));
 		logger.Info(TC("   -cachecommand=<cmd>     Send command to cache server. Will output result in log"));
 		logger.Info(TC("   -writecachesummary      Write cache summary file about connected cache server"));
+		logger.Info(TC("   -decompress=<file>      Decompress a compressed file in-place"));
 		logger.Info(TC(""));
 		logger.Info(TC("  CoordinatorOptions (if coordinator set):"));
 		logger.Info(TC("   -uri=<address>          Uri to coordinator"));
@@ -182,11 +183,12 @@ namespace uba
 		bool logToFile = false;
 		bool useHackVfs = false;
 		bool useScheduler = false;
+		bool decompressInPlace = false;
 		TString checkFileTable;
 		TString cacheFilterString;
 		TString cacheCommand;
 		TString testCompress;
-		TString testDecompress;
+		TString fileToDecompress;
 		TString testEncrypt;
 		TString testDecrypt;
 		TString addCas;
@@ -359,7 +361,14 @@ namespace uba
 							return PrintHelp(TC("-testDecompress needs a value"));
 						value.Clear().Append(g_rootDir).EnsureEndsWithSlash().Append(TCV("castemp")).EnsureEndsWithSlash().Append(TCV("TestCompress.tmp"));
 					}
-					testDecompress = value.data;
+					fileToDecompress = value.data;
+				}
+				else if (name.Equals(TCV("-decompress")))
+				{
+					if (value.IsEmpty())
+						return PrintHelp(TC("-decompress needs a filename"));
+					fileToDecompress = value.data;
+					decompressInPlace = true;
 				}
 				else if (name.Equals(TCV("-testencrypt")))
 				{
@@ -756,24 +765,24 @@ namespace uba
 			dest.Append(storage.GetTempPath()).Append(TCV("TestCompress.tmp"));
 			if (!storage.WriteCompressedFile(res, TC("MemoryMap"), InvalidFileHandle, mem, fileSize, dest.data, &header, sizeof(header), 0))
 				return false;
-			if (testDecompress.empty())
+			if (fileToDecompress.empty())
 				return true;
 			logger.Info(TC("Compressing %s successful (Written to %s)"), testCompress.c_str(), dest.data);
 			exit = true;
 		}
 
-		if (!testDecompress.empty())
+		if (!fileToDecompress.empty())
 		{
 			WorkManagerImpl workManager(maxProcessCount, TC("UbaWrk/TstDecm"));
 
-			FileAccessor fa(logger, testDecompress.c_str());
-			if (!fa.OpenMemoryRead())
-				return logger.Error(TC("Failed to open file %s"), testDecompress.c_str());
-			u64 fileSize = fa.GetSize();
-			u8* mem = fa.GetData();
+			FileAccessor faSource(logger, fileToDecompress.c_str());
+			if (!faSource.OpenMemoryRead())
+				return logger.Error(TC("Failed to open file %s"), fileToDecompress.c_str());
+			u64 fileSize = faSource.GetSize();
+			u8* mem = faSource.GetData();
 
 			if (fileSize < 16)
-				return logger.Error(TC("File %s is too small to be compressed. Requires at least 16 bytes"), testDecompress.c_str());
+				return logger.Error(TC("File %s is too small to be compressed. Requires at least 16 bytes"), fileToDecompress.c_str());
 
 			StorageCreateInfo storageInfo(g_rootDir.data, logWriter, workManager);
 			storageInfo.casCapacityBytes = 0;
@@ -785,10 +794,19 @@ namespace uba
 			auto& h = *(CompressedFileHeader*)mem;
 			if (h.IsValid())
 				reader.Skip(sizeof(CompressedFileHeader));
+			else if (decompressInPlace)
+			{
+				logger.Info(TC("File %s has no uba header. Assumed not compressed (use -testCompress to ignore header check)"), fileToDecompress.c_str());
+				return true;
+			}
+
 			u64 decompressedSize = reader.ReadU64();
 
 			if (testCompressOriginalSize && decompressedSize != testCompressOriginalSize)
-				return logger.Error(TC("Compressed file %s has wrong decompressed size. (Is it compressed?)"), testDecompress.c_str());
+				return logger.Error(TC("Compressed file %s has wrong decompressed size. (Is it compressed?)"), fileToDecompress.c_str());
+
+			if (decompressedSize > InvalidValue)
+				return logger.Error(TC("Decompressed size of file %s is way too big (%s). (Is it compressed?)"), fileToDecompress.c_str(), BytesToText(decompressedSize).str);
 
 			StringBuffer<> dest;
 			dest.Append(storage.GetTempPath()).Append(TCV("TestDecompress.tmp"));
@@ -811,7 +829,7 @@ namespace uba
 
 				OO_SINTa decompLen = OodleLZ_Decompress(reader.GetPositionData(), (OO_SINTa)compressedBlockSize, destMem, (OO_SINTa)decompressedBlockSize, OodleLZ_FuzzSafe_Yes, OodleLZ_CheckCRC_No, OodleLZ_Verbosity_None, NULL, 0, NULL, NULL, decoderMem, decoredMemSize);
 				if (decompLen != decompressedBlockSize)
-					return logger.Error(TC("Failed to decompress %s (CompressedSize: %llu DecompressedSize: %llu ReadPos: %llu CompressedBlock: %u DecompressedBlock: %u)"), testDecompress.c_str(), fileSize, decompressedSize, reader.GetPosition(), compressedBlockSize, decompressedBlockSize);
+					return logger.Error(TC("Failed to decompress %s (CompressedSize: %llu DecompressedSize: %llu ReadPos: %llu CompressedBlock: %u DecompressedBlock: %u)"), fileToDecompress.c_str(), fileSize, decompressedSize, reader.GetPosition(), compressedBlockSize, decompressedBlockSize);
 				destMem += decompressedBlockSize;
 				reader.Skip(compressedBlockSize);
 			}
@@ -819,7 +837,17 @@ namespace uba
 			if (!faDest.Close())
 				return false;
 
-			logger.Info(TC("Decompressing %s successful (Written to %s)"), testDecompress.c_str(), dest.data);
+			faSource.Close();
+			
+			if (decompressInPlace)
+			{
+				if (!CopyFileW(dest.data, fileToDecompress.c_str(), false))
+					return logger.Error(TC("Failed to copy %s to %s"), dest.data, fileToDecompress.c_str());
+				DeleteFile(dest.data);
+				logger.Info(TC("Decompressing %s successful"), fileToDecompress.c_str());
+			}
+			else
+				logger.Info(TC("Decompressing %s successful (Written to %s)"), fileToDecompress.c_str(), dest.data);
 			exit = true;
 		}
 
