[UBA]
* Added buffer cache for UbaCacheClient because macos keeps hitting stack overflow
* Fixed so cas ref count can be reduced in DropCasFile when forcing removes (used by virtual files)

diff --git a/Engine/Source/Programs/UnrealBuildAccelerator/Common/Private/UbaCacheClient.cpp b/Engine/Source/Programs/UnrealBuildAccelerator/Common/Private/UbaCacheClient.cpp
--- a/Engine/Source/Programs/UnrealBuildAccelerator/Common/Private/UbaCacheClient.cpp
+++ b/Engine/Source/Programs/UnrealBuildAccelerator/Common/Private/UbaCacheClient.cpp
@@ -35,6 +35,14 @@
 #define UBA_TRACE_FETCH_HINT_SCOPE(text)
 #endif
 
+#if UBA_CACHE_USE_SEND_MAX_BUFFER_CACHE
+#define UBA_CACHE_BINARY_WRITER_MAX_SIZE(name) BinaryWriterUsingCache name(m_sendBufferSlots);
+#define UBA_CACHE_BINARY_READER_MAX_SIZE(name) BinaryReaderUsingCache name(m_sendBufferSlots);
+#else
+#define UBA_CACHE_BINARY_WRITER_MAX_SIZE(name) StackBinaryWriter<SendMaxSize> name;
+#define UBA_CACHE_BINARY_READER_MAX_SIZE(name) StackBinaryReader<SendMaxSize> name;
+#endif
+
 namespace uba
 {
 	////////////////////////////////////////////////////////////////////////////////////////////////////
@@ -149,6 +157,9 @@ namespace uba
 	,	m_storage(info.storage)
 	,	m_client(info.client)
 	,	m_session(info.session)
+#if UBA_CACHE_USE_SEND_MAX_BUFFER_CACHE
+	,	m_sendBufferSlots(SendMaxSize)
+#endif
 	{
 		m_reportCacheKey = info.reportCacheKey;
 		m_reportMissReason = info.reportMissReason;
@@ -581,9 +592,9 @@ namespace uba
 		auto& rootPaths = context.rootPaths;
 		auto& info = context.info;
 
-#if PLATFORM_MAC
-		u8* memory = (u8*)malloc(SendMaxSize);
-		auto g = MakeGuard([memory]() { free(memory); });
+#if UBA_CACHE_USE_SEND_MAX_BUFFER_CACHE
+		u8* memory = m_sendBufferSlots.Pop();
+		auto g = MakeGuard([this, memory]() { m_sendBufferSlots.Push(memory); });
 #else
 		u8 memory[SendMaxSize];
 #endif
@@ -1217,7 +1228,7 @@ namespace uba
 		while (left)
 		{
 			TimerScope ts(stats.sendPathTable);
-			StackBinaryWriter<SendMaxSize> writer;
+			UBA_CACHE_BINARY_WRITER_MAX_SIZE(writer);
 			NetworkMessage msg(m_client, CacheServiceId, CacheMessageType_StorePathTable, writer);
 			writer.Write7BitEncoded(MakeId(bucket.id));
 			u32 toSend = Min(requiredPathTableSize - bucket.pathTableSizeSent, u32(m_client.GetMessageMaxSize() - 32));
@@ -1249,7 +1260,7 @@ namespace uba
 		while (left)
 		{
 			TimerScope ts(stats.sendCasTable);
-			StackBinaryWriter<SendMaxSize> writer;
+			UBA_CACHE_BINARY_WRITER_MAX_SIZE(writer);
 			NetworkMessage msg(m_client, CacheServiceId, CacheMessageType_StoreCasTable, writer);
 			writer.Write7BitEncoded(MakeId(bucket.id));
 			u32 toSend = Min(requiredCasTableSize - bucket.casKeyTableSizeSent, u32(m_client.GetMessageMaxSize() - 32));
@@ -1432,7 +1443,7 @@ namespace uba
 	{
 		TimerScope ts(stats.sendEntry);
 
-		StackBinaryWriter<SendMaxSize> writer;
+		UBA_CACHE_BINARY_WRITER_MAX_SIZE(writer);
 
 		NetworkMessage msg(m_client, CacheServiceId, CacheMessageType_StoreEntry, writer);
 		writer.Write7BitEncoded(MakeId(bucket.id));
@@ -1600,7 +1611,7 @@ namespace uba
 		if (!ownedLogLines)
 			return msg.Send();
 
-		StackBinaryReader<SendMaxSize> reader;
+		UBA_CACHE_BINARY_READER_MAX_SIZE(reader);
 		if (!msg.Send(reader))
 			return false;
 
diff --git a/Engine/Source/Programs/UnrealBuildAccelerator/Common/Private/UbaStorage.cpp b/Engine/Source/Programs/UnrealBuildAccelerator/Common/Private/UbaStorage.cpp
--- a/Engine/Source/Programs/UnrealBuildAccelerator/Common/Private/UbaStorage.cpp
+++ b/Engine/Source/Programs/UnrealBuildAccelerator/Common/Private/UbaStorage.cpp
@@ -29,51 +29,6 @@ namespace uba
 			out.Append('-').AppendBase62(offset);
 	}
 
-	u8* BufferSlots::Pop()
-	{
-		SCOPED_FUTEX(m_slotsLock, lock);
-		if (!m_slots.empty())
-		{
-			auto back = m_slots.back();
-			m_slots.pop_back();
-			return back;
-		}
-#if UBA_CHECK_BUFFER_SLOTS
-		u8* res = new u8[BufferSlotSize + 12];
-		res += 8;
-		*(u32*)(res - 4) = 0xdeadbeef;
-		*(u32*)(res + BufferSlotSize) = 0xdeadbeef;
-		return res;
-#else
-		return new u8[BufferSlotSize];
-#endif
-	}
-
-	void BufferSlots::Push(u8* slot)
-	{
-		if (!slot)
-			return;
-
-#if UBA_CHECK_BUFFER_SLOTS
-		UBA_ASSERT(*(u32*)(slot - 4) == 0xdeadbeef);
-		UBA_ASSERT(*(u32*)(slot + BufferSlotSize) == 0xdeadbeef);
-#endif
-
-		SCOPED_FUTEX(m_slotsLock, lock);
-		m_slots.push_back(slot);
-	}
-
-	BufferSlots::~BufferSlots()
-	{
-#if UBA_CHECK_BUFFER_SLOTS
-		for (u8* slot : m_slots)
-			delete[] (slot-8);
-#else
-		for (u8* slot : m_slots)
-			delete[] slot;
-#endif
-	}
-
 	StorageCreateInfo::StorageCreateInfo(const tchar* rootDir_, LogWriter& w, WorkManager& wm)
 	:	writer(w)
 	,	workManager(wm)
@@ -937,6 +892,7 @@ namespace uba
 	StorageImpl::StorageImpl(const StorageCreateInfo& info, const tchar* logPrefix)
 	:	m_logger(info.writer, logPrefix)
 	,	m_workManager(info.workManager)
+	,	m_bufferSlots(BufferSlotSize)
 	,	m_maxParallelCopyOrLinkBottleneck(info.maxParallelCopyOrLink)
 	,	m_casDataBuffer(m_logger, &info.workManager)
 	{
@@ -2116,7 +2072,7 @@ namespace uba
 		m_casDataBuffer.UnmapView(view, hint);
 	}
 
-	bool StorageImpl::DropCasFile(const CasKey& casKey, bool forceDelete, const tchar* hint)
+	bool StorageImpl::DropCasFile(const CasKey& casKey, bool forceDelete, const tchar* hint, bool reduceRefCount)
 	{
 		SCOPED_READ_LOCK(m_casLookupLock, lookupLock);
 		auto findIt = m_casLookup.find(casKey);
@@ -2144,7 +2100,13 @@ namespace uba
 
 		if (forceDelete)
 		{
-			UBA_ASSERT(casEntry.refCount == 0);
+			StringBuffer<> casFile;
+			if (!StorageImpl::GetCasFileName(casFile, casKey))
+				return false;
+			if (reduceRefCount)
+				--casEntry.refCount;
+
+			UBA_ASSERTF(casEntry.refCount == 0, TC("Force deleting cas %s with refcount %u (%s)"), casFile.data, casEntry.refCount, hint);
 			u64 sizeDeleted = 0;
 
 			if (casEntry.mappingHandle.IsValid())
@@ -2158,10 +2120,6 @@ namespace uba
 			}
 			else
 			{
-				StringBuffer<> casFile;
-				if (!StorageImpl::GetCasFileName(casFile, casKey))
-					return false;
-
 				if (DeleteFileW(casFile.data) == 0)
 				{
 					u32 lastError = GetLastError();
diff --git a/Engine/Source/Programs/UnrealBuildAccelerator/Common/Private/UbaStorageUtils.cpp b/Engine/Source/Programs/UnrealBuildAccelerator/Common/Private/UbaStorageUtils.cpp
--- a/Engine/Source/Programs/UnrealBuildAccelerator/Common/Private/UbaStorageUtils.cpp
+++ b/Engine/Source/Programs/UnrealBuildAccelerator/Common/Private/UbaStorageUtils.cpp
@@ -987,7 +987,7 @@ namespace uba
 
 	////////////////////////////////////////////////////////////////////////////////////////////////////
 
-	CompressWriter::CompressWriter(Logger& l, BufferSlots& bs, WorkManager& wm, StorageStats& s, u8 cc, u8 ccl, bool avof)
+	CompressWriter::CompressWriter(Logger& l, BufferCache& bs, WorkManager& wm, StorageStats& s, u8 cc, u8 ccl, bool avof)
 	:	m_logger(l)
 	,	m_bufferSlots(bs)
 	,	m_workManager(wm)
@@ -1075,7 +1075,7 @@ namespace uba
 			}
 			Logger* logger;
 			LinearWriter* destination;
-			BufferSlots* bufferSlots;
+			BufferCache* bufferSlots;
 			Atomic<u64> refCount;
 			Atomic<u64> compressCounter;
 			Event* events;
diff --git a/Engine/Source/Programs/UnrealBuildAccelerator/Common/Public/UbaCacheClient.h b/Engine/Source/Programs/UnrealBuildAccelerator/Common/Public/UbaCacheClient.h
--- a/Engine/Source/Programs/UnrealBuildAccelerator/Common/Public/UbaCacheClient.h
+++ b/Engine/Source/Programs/UnrealBuildAccelerator/Common/Public/UbaCacheClient.h
@@ -2,12 +2,15 @@
 
 #pragma once
 
+#include "UbaBufferCache.h"
 #include "UbaHash.h"
 #include "UbaLogger.h"
 #include "UbaProcessHandle.h"
 #include "UbaRootPaths.h"
 #include "UbaStats.h"
 
+#define UBA_CACHE_USE_SEND_MAX_BUFFER_CACHE PLATFORM_MAC // To prevent stack overflow
+
 namespace uba
 {
 	class CompactCasKeyTable;
@@ -154,6 +157,10 @@ namespace uba
 		struct PathHash { TString path; CasKey hash; };
 		Vector<PathHash> m_pathHashes;
 
+		#if UBA_CACHE_USE_SEND_MAX_BUFFER_CACHE
+		BufferCache m_sendBufferSlots;
+		#endif
+
 		CacheClient(const CacheClient&) = delete;
 		CacheClient& operator=(const CacheClient&) = delete;
 		struct DowngradedLogger;
diff --git a/Engine/Source/Programs/UnrealBuildAccelerator/Common/Public/UbaStorage.h b/Engine/Source/Programs/UnrealBuildAccelerator/Common/Public/UbaStorage.h
--- a/Engine/Source/Programs/UnrealBuildAccelerator/Common/Public/UbaStorage.h
+++ b/Engine/Source/Programs/UnrealBuildAccelerator/Common/Public/UbaStorage.h
@@ -3,6 +3,7 @@
 #pragma once
 
 #include "UbaBottleneck.h"
+#include "UbaBufferCache.h"
 #include "UbaFile.h"
 #include "UbaFileMapping.h"
 #include "UbaLogger.h"
@@ -50,7 +51,7 @@ namespace uba
 
 		virtual bool StoreCasFile(CasKey& out, const tchar* fileName, const CasKey& casKeyOverride, bool deferCreation) = 0;
 		virtual bool StoreCasFileClient(CasKey& out, StringKey fileNameKey, const tchar* fileName, FileMappingHandle mappingHandle, u64 mappingOffset, u64 fileSize, const tchar* hint, bool keepMappingInMemory = false, bool storeCompressed = true) = 0;
-		virtual bool DropCasFile(const CasKey& casKey, bool forceDelete, const tchar* hint) = 0;
+		virtual bool DropCasFile(const CasKey& casKey, bool forceDelete, const tchar* hint, bool reduceRefCount = false) = 0;
 		virtual bool ReportBadCasFile(const CasKey& casKey) = 0;
 		virtual bool CalculateCasKey(CasKey& out, const tchar* fileName) = 0;
 
@@ -120,17 +121,6 @@ namespace uba
 		u8 casCompressionLevel = 0;
 	};
 
-	struct BufferSlots
-	{
-		u8* Pop();
-		void Push(u8* slot);
-
-		~BufferSlots();
-
-		Futex m_slotsLock;
-		Vector<u8*> m_slots;
-	};
-
 	static constexpr u64 BufferSlotSize = 16*1024*1024;
 	static constexpr u64 BufferSlotHalfSize = BufferSlotSize/2; // This must be three times a msg size or more.
 
@@ -166,7 +156,7 @@ namespace uba
 		virtual bool InvalidateCachedFileInfo(StringKey fileNameKey) override;
 		virtual bool StoreCasFile(CasKey& out, const tchar* fileName, const CasKey& casKeyOverride, bool deferCreation) override;
 		virtual bool StoreCasFileClient(CasKey& out, StringKey fileNameKey, const tchar* fileName, FileMappingHandle mappingHandle, u64 mappingOffset, u64 fileSize, const tchar* hint, bool keepMappingInMemory = false, bool storeCompressed = true) override;
-		virtual bool DropCasFile(const CasKey& casKey, bool forceDelete, const tchar* hint) override;
+		virtual bool DropCasFile(const CasKey& casKey, bool forceDelete, const tchar* hint, bool reduceRefCount = false) override;
 		virtual bool ReportBadCasFile(const CasKey& casKey) override;
 		virtual bool CalculateCasKey(CasKey& out, const tchar* fileName) override;
 		virtual bool CopyOrLink(const CasKey& casKey, const tchar* destination, u32 fileAttributes, bool writeCompressed = false, u32 clientId = 0, const FormattingFunc& formattingFunc = {}, bool isTemp = false, bool allowHardLink = true) override;
@@ -225,7 +215,7 @@ namespace uba
 
 		MutableLogger m_logger;
 		WorkManager& m_workManager;
-		BufferSlots m_bufferSlots;
+		BufferCache m_bufferSlots;
 
 		StringBuffer<> m_rootDir;
 		StringBuffer<> m_tempPath;
diff --git a/Engine/Source/Programs/UnrealBuildAccelerator/Common/Public/UbaStorageUtils.h b/Engine/Source/Programs/UnrealBuildAccelerator/Common/Public/UbaStorageUtils.h
--- a/Engine/Source/Programs/UnrealBuildAccelerator/Common/Public/UbaStorageUtils.h
+++ b/Engine/Source/Programs/UnrealBuildAccelerator/Common/Public/UbaStorageUtils.h
@@ -11,7 +11,7 @@
 namespace uba
 {
 	class FileAccessor;
-	struct BufferSlots;
+	struct BufferCache;
 	struct StorageStats;
 
 	////////////////////////////////////////////////////////////////////////////////////////////////////
@@ -34,7 +34,7 @@ namespace uba
 
 		Logger& m_logger;
 		NetworkClient& m_client;
-		BufferSlots& m_bufferSlots;
+		BufferCache& m_bufferSlots;
 		StorageStats& m_stats;
 		Futex& m_sendOneAtTheTimeLock;
 		u8 m_casCompressor = DefaultCompressor;
@@ -49,7 +49,7 @@ namespace uba
 	{
 		bool RetrieveFile(Logger& logger, NetworkClient& client, const CasKey& casKey, const tchar* destination, bool writeCompressed, MemoryBlock* destinationMem = nullptr, u32 attributes = DefaultAttributes());
 
-		BufferSlots& m_bufferSlots;
+		BufferCache& m_bufferSlots;
 		StorageStats& m_stats;
 		StringBuffer<> m_tempPath;
 		bool m_errorOnFail = true;
@@ -72,7 +72,7 @@ namespace uba
 
 	struct CompressWriter
 	{
-		CompressWriter(Logger& l, BufferSlots& bs, WorkManager& wm, StorageStats& s, u8 cc = DefaultCompressor, u8 ccl = DefaultCompressionLevel, bool avof = false);
+		CompressWriter(Logger& l, BufferCache& bs, WorkManager& wm, StorageStats& s, u8 cc = DefaultCompressor, u8 ccl = DefaultCompressionLevel, bool avof = false);
 
 		bool CompressToFile(u64& outCompressedSize, const tchar* from, FileHandle readHandle, u8* readMem, u64 fileSize, const tchar* toFile, const void* header, u64 headerSize, u64 lastWriteTime, const tchar* tempPath);
 		bool CompressToMapping(FileMappingHandle& outMappingHandle, u64& outMappingSize, u8* readMem, u64 fileSize, const tchar* hint);
@@ -80,7 +80,7 @@ namespace uba
 		bool CompressFromMemOrFile(LinearWriter& destination, const tchar* from, FileHandle readHandle, u8* readMem, u64 fileSize);
 
 		Logger& m_logger;
-		BufferSlots& m_bufferSlots;
+		BufferCache& m_bufferSlots;
 		WorkManager& m_workManager;
 		StorageStats& m_stats;
 		u8 m_casCompressor;
@@ -122,7 +122,7 @@ namespace uba
 		bool DecompressFileToMemory(const tchar* fileName, FileHandle fileHandle, u8* dest, u64 decompressedSize, const tchar* writeHint, u64 fileStartOffset);
 
 		Logger& m_logger;
-		BufferSlots& m_bufferSlots;
+		BufferCache& m_bufferSlots;
 		WorkManager& m_workManager;
 		StorageStats& m_stats;
 	};
