[UBA]
* re-enabled mimalloc overrides. This was disabled as a test a long time ago

diff --git a/Engine/Source/Programs/UnrealBuildAccelerator/Core/Public/UbaCorePch.h b/Engine/Source/Programs/UnrealBuildAccelerator/Core/Public/UbaCorePch.h
--- a/Engine/Source/Programs/UnrealBuildAccelerator/Core/Public/UbaCorePch.h
+++ b/Engine/Source/Programs/UnrealBuildAccelerator/Core/Public/UbaCorePch.h
@@ -4,7 +4,7 @@
 
 #if UBA_USE_MIMALLOC
 #include <mimalloc.h>
-//#include <mimalloc-override.h>
+#include <mimalloc-override.h>
 #endif
 
 #include "UbaPlatform.h"
diff --git a/Engine/Source/Programs/UnrealBuildAccelerator/Core/Public/UbaMemory.h b/Engine/Source/Programs/UnrealBuildAccelerator/Core/Public/UbaMemory.h
--- a/Engine/Source/Programs/UnrealBuildAccelerator/Core/Public/UbaMemory.h
+++ b/Engine/Source/Programs/UnrealBuildAccelerator/Core/Public/UbaMemory.h
@@ -14,7 +14,7 @@
 
 #if UBA_USE_MIMALLOC
 #include <mimalloc.h>
-//#include <mimalloc-override.h>
+#include <mimalloc-override.h>
 #endif
 
 #if PLATFORM_WINDOWS && !defined(aligned_alloc)
