[UBA]
* Fixed compile error in UbaCli on macos/linux

diff --git a/Engine/Source/Programs/UnrealBuildAccelerator/Cli/Private/UbaCli.cpp b/Engine/Source/Programs/UnrealBuildAccelerator/Cli/Private/UbaCli.cpp
--- a/Engine/Source/Programs/UnrealBuildAccelerator/Cli/Private/UbaCli.cpp
+++ b/Engine/Source/Programs/UnrealBuildAccelerator/Cli/Private/UbaCli.cpp
@@ -843,7 +843,7 @@ namespace uba
 			{
 				if (!CopyFileW(dest.data, fileToDecompress.c_str(), false))
 					return logger.Error(TC("Failed to copy %s to %s"), dest.data, fileToDecompress.c_str());
-				DeleteFile(dest.data);
+				DeleteFileW(dest.data);
 				logger.Info(TC("Decompressing %s successful"), fileToDecompress.c_str());
 			}
 			else
