[UBA]
* Fixed linux compile error when mimalloc is enabled

diff --git a/Engine/Source/Programs/UnrealBuildAccelerator/Core/Public/UbaCorePch.h b/Engine/Source/Programs/UnrealBuildAccelerator/Core/Public/UbaCorePch.h
--- a/Engine/Source/Programs/UnrealBuildAccelerator/Core/Public/UbaCorePch.h
+++ b/Engine/Source/Programs/UnrealBuildAccelerator/Core/Public/UbaCorePch.h
@@ -5,6 +5,9 @@
 #if UBA_USE_MIMALLOC
 #include <mimalloc.h>
 #include <mimalloc-override.h>
+#if PLATFORM_LINUX
+namespace std { inline void* mi_realloc(void* p, size_t newsize) { return ::mi_realloc(p, newsize); } }
+#endif
 #endif
 
 #include "UbaPlatform.h"
diff --git a/Engine/Source/Programs/UnrealBuildAccelerator/Core/Public/UbaMemory.h b/Engine/Source/Programs/UnrealBuildAccelerator/Core/Public/UbaMemory.h
--- a/Engine/Source/Programs/UnrealBuildAccelerator/Core/Public/UbaMemory.h
+++ b/Engine/Source/Programs/UnrealBuildAccelerator/Core/Public/UbaMemory.h
@@ -14,7 +14,6 @@
 
 #if UBA_USE_MIMALLOC
 #include <mimalloc.h>
-#include <mimalloc-override.h>
 #endif
 
 #if PLATFORM_WINDOWS && !defined(aligned_alloc)
