[UBA]
* Workaround for win64 arm targets where multiple actions want to write to same SxS manifest

diff --git a/Engine/Source/Programs/UnrealBuildAccelerator/UbaAgent.Target.cs b/Engine/Source/Programs/UnrealBuildAccelerator/UbaAgent.Target.cs
--- a/Engine/Source/Programs/UnrealBuildAccelerator/UbaAgent.Target.cs
+++ b/Engine/Source/Programs/UnrealBuildAccelerator/UbaAgent.Target.cs
@@ -65,6 +65,11 @@ public class UbaAgentTarget : TargetRules
 				"UBA_AGENT_EXECUTABLE=L\"UbaAgent.exe\"",
 				"UBA_DETOURS_LIBRARY_ANSI=\"UbaDetours.dll\"",
 			});
+
+			if (!Target.Architectures.SingleArchitecture.bIsX64)
+			{
+				Rules.WindowsPlatform.ManifestFile = null;
+			}
 		}
 		else
 		{
