[UBA]
* Moved shared events unit tests to separate test

diff --git a/Engine/Source/Programs/UnrealBuildAccelerator/Test/Private/UbaTestAll.h b/Engine/Source/Programs/UnrealBuildAccelerator/Test/Private/UbaTestAll.h
--- a/Engine/Source/Programs/UnrealBuildAccelerator/Test/Private/UbaTestAll.h
+++ b/Engine/Source/Programs/UnrealBuildAccelerator/Test/Private/UbaTestAll.h
@@ -38,6 +38,7 @@ namespace uba
 		UBA_TEST(TestRemoteStorageStore) \
 		UBA_TEST(TestRemoteStorageFetch) \
 		UBA_TEST(TestRemoteStorageStore2) \
+		UBA_TEST(TestSharedEvents) \
 		UBA_TEST(TestDetouredTestApp) \
 		UBA_TEST(TestRemoteDetouredTestApp) \
 		UBA_TEST(TestCompactPathTable) \
diff --git a/Engine/Source/Programs/UnrealBuildAccelerator/Test/Private/UbaTestBasics.cpp b/Engine/Source/Programs/UnrealBuildAccelerator/Test/Private/UbaTestBasics.cpp
--- a/Engine/Source/Programs/UnrealBuildAccelerator/Test/Private/UbaTestBasics.cpp
+++ b/Engine/Source/Programs/UnrealBuildAccelerator/Test/Private/UbaTestBasics.cpp
@@ -132,6 +132,11 @@ namespace uba
 	{
 		if (!TestEventsImpl<Event>(logger, rootDir))
 			return false;
+		return true;
+	}
+
+	bool TestSharedEvents(LoggerWithWriter& logger, const StringBufferBase& rootDir)
+	{
 		#if !PLATFORM_WINDOWS
 		if (!TestEventsImpl<SharedEvent>(logger, rootDir))
 			return false;
