[UBA]
* Changed so shared event unit test is using shared memory

diff --git a/Engine/Source/Programs/UnrealBuildAccelerator/Test/Private/UbaTestBasics.cpp b/Engine/Source/Programs/UnrealBuildAccelerator/Test/Private/UbaTestBasics.cpp
--- a/Engine/Source/Programs/UnrealBuildAccelerator/Test/Private/UbaTestBasics.cpp
+++ b/Engine/Source/Programs/UnrealBuildAccelerator/Test/Private/UbaTestBasics.cpp
@@ -3,6 +3,7 @@
 #include "UbaBinaryParser.h"
 #include "UbaBinaryReaderWriter.h"
 #include "UbaFileAccessor.h"
+#include "UbaFileMapping.h"
 #include "UbaPathUtils.h"
 #include "UbaPlatform.h"
 #include "UbaProcessUtils.h"
@@ -86,9 +87,8 @@ namespace uba
 	}
 
 	template<class EventType>
-	bool TestEventsImpl(LoggerWithWriter& logger, const StringBufferBase& rootDir)
+	bool TestEventsImpl(LoggerWithWriter& logger, const StringBufferBase& rootDir, EventType& ev)
 	{
-		EventType ev;
 		if (!ev.Create(true))
 			return logger.Error(TC("Failed to create event"));
 
@@ -130,7 +130,8 @@ namespace uba
 
 	bool TestEvents(LoggerWithWriter& logger, const StringBufferBase& rootDir)
 	{
-		if (!TestEventsImpl<Event>(logger, rootDir))
+		Event ev;
+		if (!TestEventsImpl(logger, rootDir, ev))
 			return false;
 		return true;
 	}
@@ -138,8 +139,14 @@ namespace uba
 	bool TestSharedEvents(LoggerWithWriter& logger, const StringBufferBase& rootDir)
 	{
 		#if !PLATFORM_WINDOWS
-		if (!TestEventsImpl<SharedEvent>(logger, rootDir))
+		FileMappingBuffer mappingBuffer(logger);
+		mappingBuffer.AddTransient(TC("Temp"), false);
+		MappedView view = mappingBuffer.AllocAndMapView(MappedView_Transient, 1024, 1, TC("Foo"));
+		SharedEvent* ev = new (view.memory) SharedEvent();
+		if (!TestEventsImpl(logger, rootDir, *ev))
 			return false;
+		ev->~SharedEvent();
+		mappingBuffer.UnmapView(view, TC("Foo"));
 		#endif
 		return true;
 	}
