[UBA]
* Fixed bug when using FileMappingBuffer and independent mapping.

diff --git a/Engine/Source/Programs/UnrealBuildAccelerator/Common/Private/UbaFileMapping.cpp b/Engine/Source/Programs/UnrealBuildAccelerator/Common/Private/UbaFileMapping.cpp
--- a/Engine/Source/Programs/UnrealBuildAccelerator/Common/Private/UbaFileMapping.cpp
+++ b/Engine/Source/Programs/UnrealBuildAccelerator/Common/Private/UbaFileMapping.cpp
@@ -525,7 +525,7 @@ namespace uba
 		if (isIndependent)
 		{
 			UBA_ASSERT(!allowShrink);
-			FileMappingHandle mapping = CreateMemoryMappingW(m_logger, PAGE_READWRITE, size, hint, TC("Independent"));
+			FileMappingHandle mapping = CreateMemoryMappingW(m_logger, PAGE_READWRITE, size, nullptr, hint);
 			if (!mapping.IsValid())
 				return {};
 			u8* memory = MapViewOfFile(m_logger, mapping, FILE_MAP_WRITE, 0, size, true);
