[UBA]
* Changed so TStrdup is a inline method instead of a define (to work properly with mimalloc overrides
* Experimental changes for darling

diff --git a/Engine/Source/Programs/UnrealBuildAccelerator/Common/Private/UbaSessionClient.cpp b/Engine/Source/Programs/UnrealBuildAccelerator/Common/Private/UbaSessionClient.cpp
--- a/Engine/Source/Programs/UnrealBuildAccelerator/Common/Private/UbaSessionClient.cpp
+++ b/Engine/Source/Programs/UnrealBuildAccelerator/Common/Private/UbaSessionClient.cpp
@@ -481,7 +481,7 @@ namespace uba
 
 			#if PLATFORM_MAC
 			u32 minOsVersion = reader.ReadU32();
-			if (m_osVersion && m_osVersion < minOsVersion)
+			if (m_osVersion && m_osVersion < minOsVersion && !IsRunningDarling())
 				return m_logger.Error(TC("%s has min os version %u but current os is %u"), moduleFile.data, minOsVersion, m_osVersion);
 			#endif
 
diff --git a/Engine/Source/Programs/UnrealBuildAccelerator/Core/Private/UbaPlatform.cpp b/Engine/Source/Programs/UnrealBuildAccelerator/Core/Private/UbaPlatform.cpp
--- a/Engine/Source/Programs/UnrealBuildAccelerator/Core/Private/UbaPlatform.cpp
+++ b/Engine/Source/Programs/UnrealBuildAccelerator/Core/Private/UbaPlatform.cpp
@@ -87,6 +87,20 @@ namespace uba
 		return GetWineVersion() != nullptr;
 	}
 
+	bool IsRunningDarling()
+	{
+		#if PLATFORM_MAC
+		static bool isDarling = []()
+		{
+			struct stat attr;
+			return stat("/Library/Developer/DarlingCLT", &attr) == 0 && S_ISDIR(attr.st_mode);
+		}();
+		return isDarling;
+		#else
+		return false;
+		#endif
+	}
+
 	bool IsRunningArm()
 	{
 		#if PLATFORM_WINDOWS
@@ -185,7 +199,11 @@ namespace uba
 			return false;
 		u32 bufferLen = TStrlen(buffer) - 1;
 		buffer[bufferLen] = 0;
-		outPretty.Append(TCV("MacOS-")).Append(buffer, bufferLen);
+		if (IsRunningDarling())
+			outPretty.Append(TCV("Linux/Darling-"));
+		else
+			outPretty.Append(TCV("MacOS-"));
+		outPretty.Append(buffer, bufferLen);
 
 		u32 majorVersion = 0;
 		u32 minorVersion = 0;
@@ -668,10 +686,10 @@ namespace uba
 	void GetMappingHandleName(StringBufferBase& out, u64 uid)
 	{
 		#if PLATFORM_MAC
-		out.Append("/tmp/uba_").AppendHex(uid);
-		#else
-		out.Append("/uba_").AppendHex(uid);
+		if (!IsRunningDarling())
+			out.Append(TCV("/tmp"));
 		#endif
+		out.Append(TCV("/uba_")).AppendHex(uid);
 	}
 
 	bool GetPhysicallyInstalledSystemMemory(u64& outKb)
diff --git a/Engine/Source/Programs/UnrealBuildAccelerator/Core/Public/UbaPlatform.h b/Engine/Source/Programs/UnrealBuildAccelerator/Core/Public/UbaPlatform.h
--- a/Engine/Source/Programs/UnrealBuildAccelerator/Core/Public/UbaPlatform.h
+++ b/Engine/Source/Programs/UnrealBuildAccelerator/Core/Public/UbaPlatform.h
@@ -96,6 +96,7 @@ namespace uba
 	bool CreateGuid(Guid& out);
 	const char* GetWineVersion();
 	bool IsRunningWine();
+	bool IsRunningDarling();
 	bool IsRunningArm();
 	bool IsEscapePressed();
 	void Sleep(u32 milliseconds);
@@ -125,6 +126,7 @@ namespace uba
 	inline bool IsDirectory(u32 attributes) { return (attributes & FILE_ATTRIBUTE_DIRECTORY) != 0; }
 	inline bool IsExecutable(u32 attributes) { return false; }
 	inline bool IsUncPath(const tchar* path) { return path[0] == '\\' && path[1] == '\\'; }
+	inline tchar* TStrdup(const tchar* s) { return _wcsdup(s); }
 	#define TStrlen(s) u32(wcslen(s))
 	#define TStrchr(a, b) wcschr(a, b)
 	#define TStrrchr(a, b) wcsrchr(a, b)
@@ -134,7 +136,6 @@ namespace uba
 	#define Tvsprintf_s(...) vswprintf_s(__VA_ARGS__)
 	#define TStrcpy_s wcscpy_s
 	#define TStrcat_s wcscat_s
-	#define TStrdup _wcsdup
 	#define TFputs fputws
 	#define UBA_NOINLINE __declspec(noinline)
 	struct FileMappingHandle
@@ -159,7 +160,6 @@ namespace uba
 	#define Tvsprintf_s(a, b, c, d) vsnprintf(a, b, c, d)
 	#define TStrcpy_s(a, b, c) strcpy_s(a, b, c)
 	#define TStrcat_s(a, b, c) strcat_s(a, b, c)
-	#define TStrdup strdup
 	#define TFputs fputs
 	#define localtime_s(a,b) localtime_r(b,a)
 	void strcpy_s(tchar* dest, u64 destCapacity, const tchar* source);
@@ -171,6 +171,7 @@ namespace uba
 	inline bool IsDirectory(u32 attributes) { return S_ISDIR(attributes); }
 	inline bool IsExecutable(u32 attributes) { return (attributes & (S_IXUSR | S_IXGRP | S_IXOTH)) != 0; }
 	inline bool IsUncPath(const tchar* path) { return false; }
+	inline tchar* TStrdup(const tchar* s) { return strdup(s); }
 	#if PLATFORM_LINUX
 	inline constexpr bool CaseInsensitiveFs = false;
 	#define st_mtimespec st_mtim
