[UBA]
* Fixed compile error in UbaExports for TStrdup

diff --git a/Engine/Source/Programs/UnrealBuildAccelerator/Host/Private/UbaExports.cpp b/Engine/Source/Programs/UnrealBuildAccelerator/Host/Private/UbaExports.cpp
--- a/Engine/Source/Programs/UnrealBuildAccelerator/Host/Private/UbaExports.cpp
+++ b/Engine/Source/Programs/UnrealBuildAccelerator/Host/Private/UbaExports.cpp
@@ -488,8 +488,8 @@ extern "C"
 	{
 		auto info = new uba::SessionServerCreateInfo(storage, client, writer);
 		info->Apply(uba::GetConfig());
-		info->rootDir = TStrdup(rootDir);
-		info->traceOutputFile = TStrdup(traceOutputFile);
+		info->rootDir = uba::TStrdup(rootDir);
+		info->traceOutputFile = uba::TStrdup(traceOutputFile);
 		info->allowCustomAllocator = !disableCustomAllocator;
 		info->launchVisualizer = launchVisualizer;
 		info->resetCas = resetCas;
