[UBA]
* Implemented semaphore based shared events for posix (default off)

diff --git a/Engine/Source/Programs/UnrealBuildAccelerator/Common/Private/UbaProcess.cpp b/Engine/Source/Programs/UnrealBuildAccelerator/Common/Private/UbaProcess.cpp
--- a/Engine/Source/Programs/UnrealBuildAccelerator/Common/Private/UbaProcess.cpp
+++ b/Engine/Source/Programs/UnrealBuildAccelerator/Common/Private/UbaProcess.cpp
@@ -135,25 +135,51 @@ namespace uba
 		return fa.Close();
 	}
 
+	#if !PLATFORM_WINDOWS
+	SharedEvent& CreateSharedEvent(u8* mem, u32 index)
+	{
+		if (UseSemaphoreForSharedEvents())
+			return *new SharedEvent();
+		else
+			return *(mem ? (new (mem + sizeof(SharedEvent)*index) SharedEvent) : (SharedEvent*)nullptr);
+	}
+	#endif
+
 	ProcessImpl::ProcessImpl(Session& session, u32 id, ProcessImpl* parent, bool detourEnabled)
 	:	m_session(session)
 	,	m_parentProcess(parent)
 	,	m_id(id)
 	,	m_comMemory(detourEnabled ? m_session.m_processCommunicationAllocator.Alloc(TC("")) : FileMappingAllocator::Allocation())
 		#if !PLATFORM_WINDOWS
-	,	m_cancelEvent(*(m_comMemory.memory ? (new (m_comMemory.memory) SharedEvent) : (SharedEvent*)nullptr))
-	,	m_writeEvent(*(m_comMemory.memory ? (new (m_comMemory.memory + sizeof(SharedEvent)) SharedEvent) : (SharedEvent*)nullptr))
-	,	m_readEvent(*(m_comMemory.memory ? (new (m_comMemory.memory + sizeof(SharedEvent)*2) SharedEvent) : (SharedEvent*)nullptr))
+	,	m_cancelEvent(CreateSharedEvent(m_comMemory.memory, 0))
+	,	m_writeEvent(CreateSharedEvent(m_comMemory.memory, 1))
+	,	m_readEvent(CreateSharedEvent(m_comMemory.memory, 2))
 		#endif
 	,	m_detourEnabled(detourEnabled)
 	,	m_shared(parent ? parent->m_shared : *new Shared)
 	{
-		if (m_comMemory.memory)
+		if (!m_comMemory.memory)
+			return;
+
+		#if !PLATFORM_WINDOWS
+		if (UseSemaphoreForSharedEvents())
 		{
-			m_cancelEvent.Create(true);
-			m_writeEvent.Create(false);
-			m_readEvent.Create(false);
+			StringBuffer<256> name;
+			name.Append(TCV("uba")).AppendValue(m_comMemory.handle.uid).Append('_').AppendValue(m_comMemory.offset);
+			u64 len = name.count;
+			name.Append("_cancel");
+			m_cancelEvent.Create(name.data, true);
+			name.Resize(len).Append("_read");
+			m_writeEvent.Create(name.data, true);
+			name.Resize(len).Append("_write");
+			m_readEvent.Create(name.data, true);
+			return;
 		}
+		#endif
+
+		m_cancelEvent.Create(true);
+		m_writeEvent.Create(false);
+		m_readEvent.Create(false);
 	}
 
 	ProcessImpl::~ProcessImpl()
@@ -172,11 +198,7 @@ namespace uba
 
 		if (m_comMemory.memory)
 		{
-			#if !PLATFORM_WINDOWS
-			m_cancelEvent.~SharedEvent();
-			m_writeEvent.~SharedEvent();
-			m_readEvent.~SharedEvent();
-			#endif
+			InternalDestroyEvents();
 			m_session.m_processCommunicationAllocator.Free(m_comMemory);
 		}
 
@@ -539,8 +561,11 @@ namespace uba
 		u8* comMemory = m_comMemory.memory;
 		u64 comMemorySize = CommunicationMemSize;
 		#if !PLATFORM_WINDOWS
-		comMemory += sizeof(SharedEvent) * 3;
-		comMemorySize -= sizeof(SharedEvent) * 3;
+		if (!UseSemaphoreForSharedEvents())
+		{
+			comMemory += sizeof(SharedEvent) * 3;
+			comMemorySize -= sizeof(SharedEvent) * 3;
+		}
 		#endif
 
 		u32 retryCount = 0; // Do not allow retry
@@ -688,9 +713,7 @@ namespace uba
 		{
 			#if !PLATFORM_WINDOWS
 			SCOPED_FUTEX(m_comMemoryLock, lock);
-			m_cancelEvent.~SharedEvent();
-			m_writeEvent.~SharedEvent();
-			m_readEvent.~SharedEvent();
+			InternalDestroyEvents();
 			#endif
 			m_session.m_processCommunicationAllocator.Free(m_comMemory);
 			m_comMemory = {};
@@ -2309,6 +2332,24 @@ namespace uba
 		LogLine(printInSession, std::move(line), logType);
 	}
 
+	void ProcessImpl::InternalDestroyEvents()
+	{
+		#if !PLATFORM_WINDOWS
+		if (UseSemaphoreForSharedEvents())
+		{
+			delete &m_cancelEvent;
+			delete &m_writeEvent;
+			delete &m_readEvent;
+		}
+		else
+		{
+			m_cancelEvent.~SharedEvent();
+			m_writeEvent.~SharedEvent();
+			m_readEvent.~SharedEvent();
+		}
+		#endif
+	}
+
 #if !PLATFORM_WINDOWS
 	bool ProcessImpl::PollStdPipes(PipeReader& outReader, PipeReader& errReader, int timeoutMs)
 	{
diff --git a/Engine/Source/Programs/UnrealBuildAccelerator/Common/Public/UbaProcess.h b/Engine/Source/Programs/UnrealBuildAccelerator/Common/Public/UbaProcess.h
--- a/Engine/Source/Programs/UnrealBuildAccelerator/Common/Public/UbaProcess.h
+++ b/Engine/Source/Programs/UnrealBuildAccelerator/Common/Public/UbaProcess.h
@@ -104,6 +104,7 @@ namespace uba
 		u32 InternalCreateProcess(void* environment, FileMappingHandle communicationHandle, u64 communicationOffset);
 		u32	InternalExitProcess(bool cancel);
 		void InternalLogLine(bool printInSession, TString&& line, LogEntryType logType);
+		void InternalDestroyEvents();
 		void WaitForParent();
 		void WaitForChildrenExit();
 		bool CancelWithError(bool dummy = true);
diff --git a/Engine/Source/Programs/UnrealBuildAccelerator/Core/Private/UbaEvent.cpp b/Engine/Source/Programs/UnrealBuildAccelerator/Core/Private/UbaEvent.cpp
--- a/Engine/Source/Programs/UnrealBuildAccelerator/Core/Private/UbaEvent.cpp
+++ b/Engine/Source/Programs/UnrealBuildAccelerator/Core/Private/UbaEvent.cpp
@@ -8,6 +8,7 @@
 #if !PLATFORM_WINDOWS
 #include <sys/time.h>
 #include <new>
+#include <semaphore.h>
 #endif
 
 #define UBA_TEST_WAIT_QUALITY 0
@@ -18,6 +19,19 @@ namespace uba
 
 	u64 GetMonoticTimeNs();
 
+	inline u64 ToNanoSeconds(const struct timespec& ts)
+	{
+		return u64(ts.tv_sec)*1'000'000'000 + ts.tv_nsec;
+	}
+
+	inline struct timespec ToTimeSpecFromNano(u64 nanoSeconds)
+	{
+		struct timespec ts;
+		ts.tv_sec = nanoSeconds / 1'000'000'000;
+		ts.tv_nsec = nanoSeconds % 1'000'000'000;
+		return ts;
+	}
+
 	struct EventImpl
 	{
 		EventImpl() = default;
@@ -152,19 +166,6 @@ namespace uba
 			UnlockEventMutex();
 		}
 
-		inline u64 ToNanoSeconds(const struct timespec& ts)
-		{
-			return u64(ts.tv_sec)*1'000'000'000 + ts.tv_nsec;
-		}
-
-		inline struct timespec ToTimeSpec(u64 nanoSeconds)
-		{
-			struct timespec ts;
-			ts.tv_sec = nanoSeconds / 1'000'000'000;
-			ts.tv_nsec = nanoSeconds % 1'000'000'000;
-			return ts;
-		}
-
 		bool IsCreated()
 		{
 			return m_initialized;
@@ -227,11 +228,11 @@ namespace uba
 					else  // timed wait.
 					{
 #if PLATFORM_MAC
-						struct timespec timeout = ToTimeSpec(timeoutNs);
+						struct timespec timeout = ToTimeSpecFromNano(timeoutNs);
 						int rc = pthread_cond_timedwait_relative_np(&m_condition, &m_mutex, &timeout); // unlocks Mutex while blocking...
 						UBA_ASSERTF((rc == 0) || (rc == ETIMEDOUT), "pthread_cond_timedwait_relative_np failed"); (void)rc;
 #else
-						struct timespec timeout = ToTimeSpec(startTimeNs + timeoutNs);
+						struct timespec timeout = ToTimeSpecFromNano(startTimeNs + timeoutNs);
 						int rc = pthread_cond_timedwait(&m_condition, &m_mutex, &timeout); // unlocks Mutex while blocking...
 						UBA_ASSERTF((rc == 0) || (rc == ETIMEDOUT), "pthread_cond_timedwait failed"); (void)rc;
 #endif
@@ -260,7 +261,7 @@ namespace uba
 		bool GetManualReset() { return m_manualReset; }
 		void ResetManualReset(bool mr) {}
 
-
+		// Must be first!
 		EventImpl* next;
 
 	private:
@@ -286,6 +287,114 @@ namespace uba
 		}
 	};
 
+	class SemaphoreEvent
+	{
+	public:
+		~SemaphoreEvent()
+		{
+			if (m_snd)
+				Destroy();
+		}
+
+		bool Create(const tchar* name, bool isOwner)
+		{
+			UBA_ASSERT(!m_snd);
+
+			StringBuffer<> n1(name);
+			StringBuffer<> n2(name);
+			n1.Append(TCV("_cl"));
+			n2.Append(TCV("_sv"));
+		
+			const tchar* sndName = n1.data;
+			const tchar* rvcName = n2.data;
+			int flags = 0;
+			mode_t mode = 0600;
+
+			if (isOwner)
+			{
+				flags |= O_CREAT;
+				sndName = n2.data;
+				rvcName = n1.data;
+				m_name = strdup(name);
+				m_isOwner = true;
+			};
+
+			m_snd = sem_open(sndName, flags, mode, 0);
+			m_rsv = sem_open(rvcName, flags, mode, 0);
+
+			UBA_ASSERTF(m_snd != SEM_FAILED, "sem_open failed");
+			UBA_ASSERTF(m_rsv != SEM_FAILED, "sem_open failed");
+
+			return true;
+		}
+
+		void Destroy()
+		{
+			sem_close((sem_t*)m_snd);
+			m_snd = nullptr;
+
+			sem_close((sem_t*)m_rsv);
+			m_rsv = nullptr;
+
+			if (m_isOwner)
+			{
+				StringBuffer<> n1(m_name);
+				StringBuffer<> n2(m_name);
+				n1.Append(TCV("_cl"));
+				n2.Append(TCV("_sv"));
+				sem_unlink(n1.data);
+				sem_unlink(n2.data);
+				free(m_name);
+				m_isOwner = false;
+			}
+		}
+
+		void Set()
+		{
+			if (!sem_post((sem_t*)m_snd))
+				return;
+			UBA_ASSERTF(false, "sem_post failed");
+		}
+
+		bool IsSet(u32 timeoutMs = ~0u)
+		{
+			if (timeoutMs == ~0u)
+			{
+				while (sem_wait((sem_t*)m_rsv) == -1)
+				{
+					if (errno == EINTR)
+						continue;
+					UBA_ASSERTF(false, "sem_wait failed");
+				}
+				return errno == 0;
+			}
+
+			u64 timeoutNs = u64(timeoutMs) * 1'000'000;
+
+			struct timespec now;
+			clock_gettime(CLOCK_REALTIME, &now);
+
+			struct timespec absTime = ToTimeSpecFromNano(ToNanoSeconds(now) + timeoutNs);
+
+			while (sem_timedwait((sem_t*)m_rsv, &absTime) == -1)
+			{
+				if (errno == EINTR)
+					continue;
+				if (errno == ETIMEDOUT)
+					return false;
+				UBA_ASSERTF(false, "sem_wait failed");
+			}
+			return true;
+		}
+
+	private:
+		u8 m_dummy; // This is to map into SharedEvent under the hood
+		void* m_snd = nullptr;
+		void* m_rsv = nullptr;
+		tchar* m_name = nullptr;
+		bool m_isOwner = false;
+	};
+
 #else
 	struct EventImpl
 	{
@@ -397,56 +506,76 @@ namespace uba
 #if !PLATFORM_WINDOWS
 	SharedEvent::SharedEvent()
 	{
+		*((u8*)m_data) = 0;
 		static_assert(sizeof(m_data) >= sizeof(EventImpl));
-		new (m_data) EventImpl();
-	}
-
-	SharedEvent::SharedEvent(bool manualReset)
-	{
-		new (m_data) EventImpl();
-		Create(manualReset);
+		static_assert(sizeof(m_data) >= sizeof(SemaphoreEvent));
 	}
 
 	SharedEvent::~SharedEvent()
 	{
 		Destroy();
-		((EventImpl&)m_data).~EventImpl();
 	}
 
 	bool SharedEvent::Create(bool manualReset)//, bool shared)
 	{
+		UBA_ASSERT(*((u8*)m_data) == 0);
+		new (m_data) EventImpl();
+		UBA_ASSERT(*((u8*)m_data) == 0);
+		*((u8*)m_data) = 1;
 		return ((EventImpl&)m_data).Create(manualReset, true);
 	}
 
-	void SharedEvent::Destroy()
+	bool SharedEvent::Create(const tchar* name, bool isOwner)
 	{
-		((EventImpl&)m_data).Destroy();
+		UBA_ASSERT(*((u8*)m_data) == 0);
+		new (m_data) SemaphoreEvent();
+		*((u8*)m_data) = 2;
+		return ((SemaphoreEvent&)m_data).Create(name, isOwner);
 	}
 
-	void SharedEvent::Set()
+	void SharedEvent::Destroy()
 	{
-		((EventImpl&)m_data).Set();
-	}
+		if (*((u8*)m_data) == 0)
+			return;
 
-	void SharedEvent::Reset()
-	{
-		((EventImpl&)m_data).Reset();
+		if (IsSemaphore())
+		{
+			((SemaphoreEvent&)m_data).Destroy();
+			((SemaphoreEvent&)m_data).~SemaphoreEvent();
+		}
+		else
+		{
+			((EventImpl&)m_data).Destroy();
+			((EventImpl&)m_data).~EventImpl();
+		}
+		*((u8*)m_data) = 0;
 	}
 
-	bool SharedEvent::IsCreated()
+	void SharedEvent::Set()
 	{
-		return ((EventImpl&)m_data).IsCreated();
+		if (*((u8*)m_data) == 0)
+			return;
+
+		if (IsSemaphore())
+			((SemaphoreEvent&)m_data).Set();
+		else
+			((EventImpl&)m_data).Set();
 	}
 
 	bool SharedEvent::IsSet(u32 timeOutMs)
 	{
-		return ((EventImpl&)m_data).IsSet(timeOutMs);
+		if (*((u8*)m_data) == 0)
+			return false;
+
+		if (IsSemaphore())
+			return ((SemaphoreEvent&)m_data).IsSet(timeOutMs);
+		else
+			return ((EventImpl&)m_data).IsSet(timeOutMs);
 	}
 
-	void* SharedEvent::GetHandle()
+	bool SharedEvent::IsSemaphore()
 	{
-		UBA_ASSERTF(false, "Event::GetHandle not available");
-		return 0;
+		return *((u8*)m_data) == 2;
 	}
 #endif
 }
diff --git a/Engine/Source/Programs/UnrealBuildAccelerator/Core/Private/UbaPlatform.cpp b/Engine/Source/Programs/UnrealBuildAccelerator/Core/Private/UbaPlatform.cpp
--- a/Engine/Source/Programs/UnrealBuildAccelerator/Core/Private/UbaPlatform.cpp
+++ b/Engine/Source/Programs/UnrealBuildAccelerator/Core/Private/UbaPlatform.cpp
@@ -692,6 +692,11 @@ namespace uba
 		out.Append(TCV("/uba_")).AppendHex(uid);
 	}
 
+	bool UseSemaphoreForSharedEvents()
+	{
+		return false;
+	}
+
 	bool GetPhysicallyInstalledSystemMemory(u64& outKb)
 	{
 	#if PLATFORM_MAC
diff --git a/Engine/Source/Programs/UnrealBuildAccelerator/Core/Public/UbaEvent.h b/Engine/Source/Programs/UnrealBuildAccelerator/Core/Public/UbaEvent.h
--- a/Engine/Source/Programs/UnrealBuildAccelerator/Core/Public/UbaEvent.h
+++ b/Engine/Source/Programs/UnrealBuildAccelerator/Core/Public/UbaEvent.h
@@ -30,15 +30,13 @@ namespace uba
 	{
 	public:
 		SharedEvent();
-		SharedEvent(bool manualReset);
 		~SharedEvent();
 		bool Create(bool manualReset);
+		bool Create(const tchar* name, bool isOwner);
 		void Destroy();
 		void Set();
-		void Reset();
-		bool IsCreated();
 		bool IsSet(u32 timeOutMs = ~0u);
-		void* GetHandle();
+		bool IsSemaphore();
 	private:
 		#if PLATFORM_MAC || defined(__aarch64__)
 		u64 m_data[16];
diff --git a/Engine/Source/Programs/UnrealBuildAccelerator/Core/Public/UbaPlatform.h b/Engine/Source/Programs/UnrealBuildAccelerator/Core/Public/UbaPlatform.h
--- a/Engine/Source/Programs/UnrealBuildAccelerator/Core/Public/UbaPlatform.h
+++ b/Engine/Source/Programs/UnrealBuildAccelerator/Core/Public/UbaPlatform.h
@@ -165,6 +165,7 @@ namespace uba
 	void strcpy_s(tchar* dest, u64 destCapacity, const tchar* source);
 	void strcat_s(tchar* dest, u64 destCapacity, const tchar* source);
 	void GetMappingHandleName(StringBufferBase& out, u64 uid);
+	bool UseSemaphoreForSharedEvents();
 	bool GetPhysicallyInstalledSystemMemory(u64& outKb);
 	inline u64 FromTimeSpec(const timespec& ts) { return u64(ts.tv_sec) * 10'000'000ull + u64(ts.tv_nsec/100); }
 	inline timespec ToTimeSpec(u64 time) { timespec ts; ts.tv_sec = time / 10'000'000ull; ts.tv_nsec = (time - (u64(ts.tv_sec) * 10'000'000ull)) * 100; return ts; }
diff --git a/Engine/Source/Programs/UnrealBuildAccelerator/Detours/Private/PosixOS/UbaDetoursMainPosix.cpp b/Engine/Source/Programs/UnrealBuildAccelerator/Detours/Private/PosixOS/UbaDetoursMainPosix.cpp
--- a/Engine/Source/Programs/UnrealBuildAccelerator/Detours/Private/PosixOS/UbaDetoursMainPosix.cpp
+++ b/Engine/Source/Programs/UnrealBuildAccelerator/Detours/Private/PosixOS/UbaDetoursMainPosix.cpp
@@ -77,10 +77,30 @@ static void __attribute__((constructor(102))) PreInitCtor()
 		return;
 	}
 
-	g_cancelEvent = ((SharedEvent*)rptr);
-	g_readEvent = ((SharedEvent*)rptr) + 1;
-	g_writeEvent = ((SharedEvent*)rptr) + 2;
-	g_messageMappingMem = rptr + sizeof(SharedEvent) * 3;
+	g_messageMappingMem = rptr;
+
+	if (UseSemaphoreForSharedEvents())
+	{
+		StringBuffer<256> name;
+		name.Append(TCV("uba")).AppendValue(comIdUid).Append('_').AppendValue(comIdOffset);
+		u64 len = name.count;
+		g_cancelEvent = new SharedEvent();
+		name.Append("_cancel");
+		g_cancelEvent->Create(name.data, false);
+		g_readEvent = new SharedEvent();
+		name.Resize(len).Append("_read");
+		g_readEvent->Create(name.data, false);
+		g_writeEvent = new SharedEvent();
+		name.Resize(len).Append("_write");
+		g_writeEvent->Create(name.data, false);
+	}
+	else
+	{
+		g_cancelEvent = ((SharedEvent*)rptr);
+		g_readEvent = ((SharedEvent*)rptr) + 1;
+		g_writeEvent = ((SharedEvent*)rptr) + 2;
+		g_messageMappingMem += sizeof(SharedEvent) * 3;
+	}
 
 	AddExceptionHandler();
 	InitMemory();
@@ -103,7 +123,17 @@ void CloseCom()
 {
 	if (g_comFd == -1)
 		return;
-	void* mem = g_cancelEvent;
+	void* mem = g_messageMappingMem;
+	if (UseSemaphoreForSharedEvents())
+	{
+		delete g_cancelEvent;
+		delete g_readEvent;
+		delete g_writeEvent;
+	}
+	else
+	{
+		mem = g_cancelEvent;
+	}
 	g_cancelEvent = nullptr;
 	g_readEvent = nullptr;
 	g_writeEvent = nullptr;
diff --git a/Engine/Source/Programs/UnrealBuildAccelerator/Test/Private/UbaTestAll.h b/Engine/Source/Programs/UnrealBuildAccelerator/Test/Private/UbaTestAll.h
--- a/Engine/Source/Programs/UnrealBuildAccelerator/Test/Private/UbaTestAll.h
+++ b/Engine/Source/Programs/UnrealBuildAccelerator/Test/Private/UbaTestAll.h
@@ -40,6 +40,7 @@ namespace uba
 		UBA_TEST(TestRemoteStorageStore2) \
 		UBA_TEST(TestFileMappingBuffer) \
 		UBA_TEST(TestSharedEvents) \
+		UBA_TEST(TestSharedEvents2) \
 		UBA_TEST(TestDetouredTestApp) \
 		UBA_TEST(TestRemoteDetouredTestApp) \
 		UBA_TEST(TestCompactPathTable) \
diff --git a/Engine/Source/Programs/UnrealBuildAccelerator/Test/Private/UbaTestBasics.cpp b/Engine/Source/Programs/UnrealBuildAccelerator/Test/Private/UbaTestBasics.cpp
--- a/Engine/Source/Programs/UnrealBuildAccelerator/Test/Private/UbaTestBasics.cpp
+++ b/Engine/Source/Programs/UnrealBuildAccelerator/Test/Private/UbaTestBasics.cpp
@@ -87,16 +87,13 @@ namespace uba
 	}
 
 	template<class EventType>
-	bool TestEventsImpl(LoggerWithWriter& logger, const StringBufferBase& rootDir, EventType& ev)
+	bool TestEventsImpl(LoggerWithWriter& logger, const StringBufferBase& rootDir, EventType& ev, EventType& ev2)
 	{
-		if (!ev.Create(true))
-			return logger.Error(TC("Failed to create event"));
-
 		Thread t([&]()
 			{
 				Sleep(500);
 				//logger.Info(TC("Setting event"));
-				ev.Set();
+				ev2.Set();
 				Sleep(500);
 				return true;
 			});
@@ -120,8 +117,8 @@ namespace uba
 
 		#if 0 // Long time test... disabled by default
 		u64 time = GetTime();
-		EventType ev2(true);
-		ev2.IsSet(10 * 60 * 1000);
+		EventType longEv(true);
+		longEv.IsSet(10 * 60 * 1000);
 		if (TimeToMs((GetTime() - time) < 9 * 60 * 1000))
 			return logger.Error(TC("Event timeout was way too fast"));
 		#endif
@@ -131,7 +128,9 @@ namespace uba
 	bool TestEvents(LoggerWithWriter& logger, const StringBufferBase& rootDir)
 	{
 		Event ev;
-		if (!TestEventsImpl(logger, rootDir, ev))
+		if (!ev.Create(true))
+			return logger.Error(TC("Failed to create event"));
+		if (!TestEventsImpl(logger, rootDir, ev, ev))
 			return false;
 		return true;
 	}
@@ -143,7 +142,9 @@ namespace uba
 		mappingBuffer.AddTransient(TC("Temp"), false);
 		MappedView view = mappingBuffer.AllocAndMapView(MappedView_Transient, 1024, 1, TC("Foo"));
 		SharedEvent* ev = new (view.memory) SharedEvent();
-		if (!TestEventsImpl(logger, rootDir, *ev))
+		if (!ev->Create(true))
+			return logger.Error(TC("Failed to create event"));
+		if (!TestEventsImpl(logger, rootDir, *ev, *ev))
 			return false;
 		ev->~SharedEvent();
 		mappingBuffer.UnmapView(view, TC("Foo"));
@@ -151,6 +152,19 @@ namespace uba
 		return true;
 	}
 
+	bool TestSharedEvents2(LoggerWithWriter& logger, const StringBufferBase& rootDir)
+	{
+		#if !PLATFORM_WINDOWS
+		SharedEvent ev;
+		ev.Create("uba_tse", true);
+		SharedEvent ev2;
+		ev2.Create("uba_tse", false);
+		if (!TestEventsImpl(logger, rootDir, ev, ev2))
+			return false;
+		#endif
+		return true;
+	}
+
 	bool TestPaths(LoggerWithWriter& logger, const StringBufferBase& rootDir)
 	{
 		const tchar* workingDir = IsWindows ? TC("e:\\dev\\") : TC("/dev/bar/");
