[UBA]
* Fixed assert triggering when using uba debug binaries in UBT on non-windows

diff --git a/Engine/Source/Programs/UnrealBuildAccelerator/Common/Private/UbaTrace.cpp b/Engine/Source/Programs/UnrealBuildAccelerator/Common/Private/UbaTrace.cpp
--- a/Engine/Source/Programs/UnrealBuildAccelerator/Common/Private/UbaTrace.cpp
+++ b/Engine/Source/Programs/UnrealBuildAccelerator/Common/Private/UbaTrace.cpp
@@ -106,7 +106,11 @@ namespace uba
 	bool Trace::StartWrite(const tchar* namedTrace, u64 traceMemCapacity, const tchar* channelName)
 	{
 		m_memoryCapacity = traceMemCapacity;
-		m_memoryHandle = uba::CreateMemoryMappingW(m_logger, PAGE_READWRITE|SEC_RESERVE, m_memoryCapacity, namedTrace, TC("Trace"));
+		const tchar* mappingName = namedTrace;
+		#if !PLATFORM_WINDOWS
+		mappingName = nullptr;
+		#endif
+		m_memoryHandle = uba::CreateMemoryMappingW(m_logger, PAGE_READWRITE|SEC_RESERVE, m_memoryCapacity, mappingName, TC("Trace"));
 		if (!m_memoryHandle.IsValid())
 			return false;
 		if (GetLastError() != ERROR_ALREADY_EXISTS)
