[UBA]
* Fixed macos compile error (does not support sem_waittimed)

diff --git a/Engine/Source/Programs/UnrealBuildAccelerator/Core/Private/UbaEvent.cpp b/Engine/Source/Programs/UnrealBuildAccelerator/Core/Private/UbaEvent.cpp
--- a/Engine/Source/Programs/UnrealBuildAccelerator/Core/Private/UbaEvent.cpp
+++ b/Engine/Source/Programs/UnrealBuildAccelerator/Core/Private/UbaEvent.cpp
@@ -374,8 +374,10 @@ namespace uba
 			struct timespec now;
 			clock_gettime(CLOCK_REALTIME, &now);
 
-			struct timespec absTime = ToTimeSpecFromNano(ToNanoSeconds(now) + timeoutNs);
+			u64 absTimeNano = ToNanoSeconds(now) + timeoutNs;
 
+			#if PLATFORM_LINUX
+			struct timespec absTime = ToTimeSpecFromNano(absTimeNano);
 			while (sem_timedwait((sem_t*)m_rsv, &absTime) == -1)
 			{
 				if (errno == EINTR)
@@ -384,6 +386,17 @@ namespace uba
 					return false;
 				UBA_ASSERTF(false, "sem_wait failed");
 			}
+			#else
+			while (true)
+			{
+				if (sem_trywait((sem_t*)m_rsv) == 0)
+					return true;
+				Sleep(1);
+				clock_gettime(CLOCK_REALTIME, &now);
+				if (ToNanoSeconds(now) > absTimeNano)
+					return false;
+			}
+			#endif
 			return true;
 		}
 
