[UBA]
* Fixed warning message in unit tests

diff --git a/Engine/Source/Programs/UnrealBuildAccelerator/Test/Private/UbaTestSession.cpp b/Engine/Source/Programs/UnrealBuildAccelerator/Test/Private/UbaTestSession.cpp
--- a/Engine/Source/Programs/UnrealBuildAccelerator/Test/Private/UbaTestSession.cpp
+++ b/Engine/Source/Programs/UnrealBuildAccelerator/Test/Private/UbaTestSession.cpp
@@ -319,7 +319,7 @@ namespace uba
 		FILE* fpCommand = popen(command, "r");
 		if (fpCommand == nullptr || fgets(commandOutput.data, commandOutput.capacity, fpCommand) == nullptr || pclose(fpCommand) != 0)
 		{
-			logger.Warning("Failed to run '%s' or get a response");
+			logger.Warning("Failed to run command '%s' or get a response", command);
 			return false;
 		}
 
