[UBA]
* Enabled semaphore usage instead of shared process mutex/cond on darling
* Fixed unit tests not working on darling

diff --git a/Engine/Source/Programs/UnrealBuildAccelerator/Core/Private/UbaPlatform.cpp b/Engine/Source/Programs/UnrealBuildAccelerator/Core/Private/UbaPlatform.cpp
--- a/Engine/Source/Programs/UnrealBuildAccelerator/Core/Private/UbaPlatform.cpp
+++ b/Engine/Source/Programs/UnrealBuildAccelerator/Core/Private/UbaPlatform.cpp
@@ -694,7 +694,7 @@ namespace uba
 
 	bool UseSemaphoreForSharedEvents()
 	{
-		return false;
+		return IsRunningDarling();
 	}
 
 	bool GetPhysicallyInstalledSystemMemory(u64& outKb)
diff --git a/Engine/Source/Programs/UnrealBuildAccelerator/Test/Private/UbaTestBasics.cpp b/Engine/Source/Programs/UnrealBuildAccelerator/Test/Private/UbaTestBasics.cpp
--- a/Engine/Source/Programs/UnrealBuildAccelerator/Test/Private/UbaTestBasics.cpp
+++ b/Engine/Source/Programs/UnrealBuildAccelerator/Test/Private/UbaTestBasics.cpp
@@ -138,6 +138,8 @@ namespace uba
 	bool TestSharedEvents(LoggerWithWriter& logger, const StringBufferBase& rootDir)
 	{
 		#if !PLATFORM_WINDOWS
+		if (IsRunningDarling()) // bug in darling, this does not work
+			return true;
 		FileMappingBuffer mappingBuffer(logger);
 		mappingBuffer.AddTransient(TC("Temp"), false);
 		MappedView view = mappingBuffer.AllocAndMapView(MappedView_Transient, 1024, 1, TC("Foo"));
diff --git a/Engine/Source/Programs/UnrealBuildAccelerator/Test/Private/UbaTestSession.cpp b/Engine/Source/Programs/UnrealBuildAccelerator/Test/Private/UbaTestSession.cpp
--- a/Engine/Source/Programs/UnrealBuildAccelerator/Test/Private/UbaTestSession.cpp
+++ b/Engine/Source/Programs/UnrealBuildAccelerator/Test/Private/UbaTestSession.cpp
@@ -350,9 +350,18 @@ namespace uba
 		const tchar* clangPath = TC("c:\\Program Files\\Microsoft Visual Studio\\2022\\Professional\\VC\\Tools\\Llvm\\bin\\clang-cl.exe");
 #elif PLATFORM_MAC
 		StringBuffer<MaxPath> xcodePath;
-		if (!ExecuteCommand(logger, "/usr/bin/xcrun --find clang++", xcodePath))
-			return true;
-		const tchar* clangPath = xcodePath.data;
+		const tchar* clangPath;
+		if (IsRunningDarling())
+		{
+			// xcrun does not work in darling
+			clangPath = "/Volumes/SystemRoot/usr/bin/clang++";
+		}
+		else
+		{
+			if (!ExecuteCommand(logger, "/usr/bin/xcrun --find clang++", xcodePath))
+				return true;
+			clangPath = xcodePath.data;
+		}
 #else
 		const tchar* clangPath = TC("/usr/bin/clang++");
 #endif
@@ -369,10 +378,13 @@ namespace uba
 		args.Append("/Brepro ");
 #elif PLATFORM_MAC
 		StringBuffer<MaxPath> xcodeSDKPath;
-		if (!ExecuteCommand(logger, "xcrun --show-sdk-path", xcodeSDKPath))
-			return true;
-		args.Append("-isysroot ");
-		args.Append(xcodeSDKPath.data).Append(' ');
+		if (!IsRunningDarling())
+		{
+			if (!ExecuteCommand(logger, "xcrun --show-sdk-path", xcodeSDKPath))
+				return true;
+			args.Append("-isysroot ");
+			args.Append(xcodeSDKPath.data).Append(' ');
+		}
 #endif
 		args.Append(TCV("-o code Code.cpp"));
 
