[UBA]
* Changed CopyFileW implementation for macos to attempt clonefile first, and then fallback to copyfile or hand rolled read/write to copy

diff --git a/Engine/Source/Programs/UnrealBuildAccelerator/Common/Private/UbaFile.cpp b/Engine/Source/Programs/UnrealBuildAccelerator/Common/Private/UbaFile.cpp
--- a/Engine/Source/Programs/UnrealBuildAccelerator/Common/Private/UbaFile.cpp
+++ b/Engine/Source/Programs/UnrealBuildAccelerator/Common/Private/UbaFile.cpp
@@ -2,6 +2,7 @@
 
 #include "UbaFile.h"
 #include "UbaDirectoryIterator.h"
+#include "UbaEnvironment.h"
 #include "UbaEvent.h"
 #include "UbaPathUtils.h"
 #include "UbaProcessStats.h"
@@ -24,6 +25,7 @@ extern "C" NTSTATUS NTAPI NtQueryFullAttributesFile(const OBJECT_ATTRIBUTES *att
 
 #if PLATFORM_MAC
 #include <mach-o/dyld.h>
+#include <sys/clonefile.h>
 #include <copyfile.h>
 #endif
 
@@ -681,15 +683,64 @@ namespace uba
 #endif
 	}
 
-	bool CopyFileW(const tchar* existingFileName, const tchar* newFileName, bool bFailIfExists)
+	bool CopyFileW(const tchar* existingFileName, const tchar* newFileName, bool bFailIfExists, u8* scratch, u64 scratchSize)
 	{
+		ExtendedTimerScope ts(KernelStats::GetCurrent().copyFile);
 #if PLATFORM_WINDOWS
 		MAKE_LONG_FILENAME(existingFileName);
 		MAKE_LONG_FILENAME(newFileName);
 		return ::CopyFileW(existingFileName, newFileName, bFailIfExists);
 #elif PLATFORM_MAC
-		if (copyfile(existingFileName, newFileName, 0, COPYFILE_ALL) == 0)
+		if (clonefile(existingFileName, newFileName, 0) == 0)
 			return true;
+
+		if (scratch)
+		{
+			int sourceFd = open(existingFileName, O_RDONLY | O_CLOEXEC, 0);
+			if (sourceFd == -1)
+				return false;
+			auto closeSourceFile = MakeGuard([&]() { close(sourceFd); });
+			struct stat sourceStat;
+			if (fstat(sourceFd, &sourceStat) != 0)
+				return false;
+			if (!S_ISREG(sourceStat.st_mode))
+				return false;
+			u64 fileSize = sourceStat.st_size;
+			int targetFd = open(newFileName, O_CREAT | O_TRUNC | O_WRONLY | O_CLOEXEC, 0600);
+			auto closeTargetFile = MakeGuard([&]() { close(targetFd); });
+			if (targetFd == -1)
+				return false;//m_logger.Error(TC("Failed to create file %s for move from temporary file %s (%s)"), m_fileName, realFileName, strerror(errno));
+			u64 left = fileSize;
+			while (left)
+			{
+				u64 toCopy = Min(left, scratchSize);
+				auto bytesRead = read(sourceFd, scratch, toCopy);
+				if (bytesRead <= 0)
+					return false;//m_logger.Error(TC("Failed to do read from file %s (%s)"), realFileName, strerror(errno));
+				u8* readPos = scratch;
+				u64 toWrite = bytesRead;
+				while (toWrite)
+				{
+					auto written = write(targetFd, readPos, toWrite);
+					if (written == -1)
+						return false;// m_logger.Error(TC("Failed to write to file %s (%s)"), m_fileName, strerror(errno));
+					toWrite -= written;
+					readPos += written;
+				}
+				left -= bytesRead;
+			}
+			closeSourceFile.Execute();
+			if (fchmod(targetFd, sourceStat.st_mode & 07777) != 0)
+				return false;
+			fchown(targetFd, sourceStat.st_uid, sourceStat.st_gid);
+
+			return true;
+		}
+		else
+		{
+			if (copyfile(existingFileName, newFileName, 0, COPYFILE_ALL) == 0)
+				return true;
+		}
 		UBA_ASSERTF(false, TC("CopyFileW failed on %s - Error handling not implemented (%s)"), existingFileName, strerror(errno));
 		return false;
 #else
diff --git a/Engine/Source/Programs/UnrealBuildAccelerator/Common/Private/UbaStorage.cpp b/Engine/Source/Programs/UnrealBuildAccelerator/Common/Private/UbaStorage.cpp
--- a/Engine/Source/Programs/UnrealBuildAccelerator/Common/Private/UbaStorage.cpp
+++ b/Engine/Source/Programs/UnrealBuildAccelerator/Common/Private/UbaStorage.cpp
@@ -2402,7 +2402,13 @@ namespace uba
 				#endif
 
 				if (!success)
-					success = uba::CopyFileW(casFile.data, destination, true) != 0;
+				{
+					u8* mem = nullptr;
+					if (IsRunningDarling()) // Should also be used in WSL2
+						mem = m_bufferSlots.Pop(); 
+					success = uba::CopyFileW(casFile.data, destination, true, mem, BufferSlotSize);
+					m_bufferSlots.Push(mem);
+				}
 
 				if (success)
 				{
diff --git a/Engine/Source/Programs/UnrealBuildAccelerator/Common/Public/UbaFile.h b/Engine/Source/Programs/UnrealBuildAccelerator/Common/Public/UbaFile.h
--- a/Engine/Source/Programs/UnrealBuildAccelerator/Common/Public/UbaFile.h
+++ b/Engine/Source/Programs/UnrealBuildAccelerator/Common/Public/UbaFile.h
@@ -78,7 +78,7 @@ namespace uba
 	bool CreateDirectoryW(const tchar* pathName);
 	bool RemoveDirectoryW(const tchar* pathName);
 	bool DeleteFileW(const tchar* fileName);
-	bool CopyFileW(const tchar* existingFileName, const tchar* newFileName, bool bFailIfExists);
+	bool CopyFileW(const tchar* existingFileName, const tchar* newFileName, bool bFailIfExists, u8* scratch = nullptr, u64 scratchSize = 0);
 	u32 GetLongPathNameW(const tchar* lpszShortPath, tchar* lpszLongPath, u32 cchBuffer);
 	bool GetFileLastWriteTime(u64& outTime, FileHandle hFile);
 	bool SetFileLastWriteTime(FileHandle fileHandle, u64 writeTime);
