[UBA]
* Implemented support for native cross process macos mach semaphores

diff --git a/Engine/Source/Programs/UnrealBuildAccelerator/Common/Private/UbaProcess.cpp b/Engine/Source/Programs/UnrealBuildAccelerator/Common/Private/UbaProcess.cpp
--- a/Engine/Source/Programs/UnrealBuildAccelerator/Common/Private/UbaProcess.cpp
+++ b/Engine/Source/Programs/UnrealBuildAccelerator/Common/Private/UbaProcess.cpp
@@ -168,11 +168,11 @@ namespace uba
 			name.Append(TCV("uba")).AppendValue(m_comMemory.handle.uid).Append('_').AppendValue(m_comMemory.offset);
 			u64 len = name.count;
 			name.Append("_cancel");
-			m_cancelEvent.Create(name.data, true);
+			m_cancelEvent.Create(name, true);
 			name.Resize(len).Append("_read");
-			m_writeEvent.Create(name.data, true);
+			m_writeEvent.Create(name, true);
 			name.Resize(len).Append("_write");
-			m_readEvent.Create(name.data, true);
+			m_readEvent.Create(name, true);
 			return;
 		}
 		#endif
@@ -1260,6 +1260,11 @@ namespace uba
 		writer.WriteString(m_realWorkingDir);
 
 		writer.WriteString(process.m_realApplication);
+
+		#if UBA_USE_NATIVE_MAC_SEMAPHORES
+		writer.WriteU32(process.m_id);
+		#endif
+
 		#if !PLATFORM_WINDOWS
 		writer.WriteU64(process.m_comMemory.handle.uid);
 		writer.WriteU32(process.m_comMemory.offset);
@@ -2033,6 +2038,11 @@ namespace uba
 			StringBuffer<512> detoursVar;
 			StringBuffer<32> processVar;
 
+			#if UBA_USE_NATIVE_MAC_SEMAPHORES
+			StringBuffer<32> sessionId;
+			StringBuffer<32> pidVar;
+			#endif
+
 			Vector<const char*> envvars;
 
 			const char* it = (const char*)environment;
@@ -2078,6 +2088,12 @@ namespace uba
 
 
 				comIdVar.Append("UBA_COMID=").AppendValue(communicationHandle.uid).Append('+').AppendValue(communicationOffset);
+
+				#if UBA_USE_NATIVE_MAC_SEMAPHORES
+				sessionId.Append("UBA_SESSION_ID=").AppendValue(m_session.GetSessionId());
+				pidVar.Append("UBA_PID=").AppendValue(m_id);
+				#endif
+
 				workingDir.Append("UBA_CWD=").Append(m_realWorkingDir);
 				rulesStr.Append("UBA_RULES=").AppendValue(m_startInfo.rules->index);
 
@@ -2098,6 +2114,12 @@ namespace uba
 				envvars.push_back(comIdVar.data);
 				envvars.push_back(workingDir.data);
 				envvars.push_back(rulesStr.data);
+
+				#if PLATFORM_MAC
+				envvars.push_back(sessionId.data);
+				envvars.push_back(pidVar.data);
+				#endif
+
 				if (m_runningRemote)
 					envvars.push_back("UBA_REMOTE=1");
 				if (!logFile.IsEmpty())
diff --git a/Engine/Source/Programs/UnrealBuildAccelerator/Common/Private/UbaSession.cpp b/Engine/Source/Programs/UnrealBuildAccelerator/Common/Private/UbaSession.cpp
--- a/Engine/Source/Programs/UnrealBuildAccelerator/Common/Private/UbaSession.cpp
+++ b/Engine/Source/Programs/UnrealBuildAccelerator/Common/Private/UbaSession.cpp
@@ -12,7 +12,6 @@
 #include "UbaDirectoryIterator.h"
 #include "UbaApplicationRules.h"
 #include "UbaPathUtils.h"
-#include "UbaProtocol.h"
 #include "UbaStorageUtils.h"
 #include "UbaWorkManager.h"
 #include <span>
@@ -21,6 +20,12 @@
 extern char **environ;
 #endif
 
+#if UBA_USE_NATIVE_MAC_SEMAPHORES
+#include <mach/mach.h>
+#include <mach/semaphore.h>
+#include <servers/bootstrap.h>
+#endif
+
 #define UBA_DEBUG_TRACK_DIR 0 // UBA_DEBUG_LOGGER
 
 //////////////////////////////////////////////////////////////////////////////
@@ -1821,6 +1826,10 @@ namespace uba
 				return true;
 			});
 
+		#if UBA_USE_NATIVE_MAC_SEMAPHORES
+		StartSemaphoreService();
+		#endif
+
 		return true;
 	}
 
@@ -1835,6 +1844,10 @@ namespace uba
 		CancelAllProcessesAndWait();
 		FlushDeadProcesses();
 
+		#if UBA_USE_NATIVE_MAC_SEMAPHORES
+		StopSemaphoreService();
+		#endif
+
 		delete &m_processMessageThreads;
 
 		for (auto& kv : m_virtualSourceFiles)
@@ -4019,6 +4032,112 @@ namespace uba
 		m_dependencyCrawler.Add(at + 1, startInfo.workingDir, CreateFileFunc, DevirtualizePathFunc, startInfo.application, crawlerType, startInfo.rules->index);
 	}
 
+#if UBA_USE_NATIVE_MAC_SEMAPHORES
+	struct Session::SemaphoreServiceInfo { mach_port_t port = MACH_PORT_NULL; };
+
+	u32 Session::GetSessionId()
+	{
+		return m_sessionId;
+	}
+
+	void Session::StartSemaphoreService()
+	{
+		static Atomic<u32> id;
+		m_sessionId = id++;
+		mach_port_t port = MACH_PORT_NULL;
+		kern_return_t kr;
+		kr = mach_port_allocate(mach_task_self(), MACH_PORT_RIGHT_RECEIVE, &port);
+		UBA_ASSERT(kr == KERN_SUCCESS);
+
+		// Give ourselves a send right so the child can contact us if needed.
+		// (Optional if you only expect to receive.)
+		kr = mach_port_insert_right(mach_task_self(), port, port, MACH_MSG_TYPE_MAKE_SEND);
+		UBA_ASSERT(kr == KERN_SUCCESS);
+
+		#pragma clang diagnostic push
+		#pragma clang diagnostic ignored "-Wdeprecated-declarations"
+		StringBuffer<64> service;
+		service.Append("UbaSemaphoreService").AppendValue(m_sessionId);
+		//name_t service = "com.epicgames.uba";
+		kr = bootstrap_register(bootstrap_port, service.data, port);
+		UBA_ASSERT(kr == KERN_SUCCESS);
+		#pragma clang diagnostic pop
+
+		m_semaphoreServiceInfo = new SemaphoreServiceInfo { port };
+		m_semaphoreServiceThread.Start([this]() { ThreadSemaphoreService(); return 0; }, TC("UbaSemSrv"));
+	}
+	
+	void Session::StopSemaphoreService()
+	{
+		mach_port_destroy(mach_task_self(), m_semaphoreServiceInfo->port);
+		delete m_semaphoreServiceInfo;
+		m_semaphoreServiceInfo = nullptr;
+		m_semaphoreServiceThread.Wait();
+	}
+
+	void Session::ThreadSemaphoreService()
+	{
+		typedef struct {
+			mach_msg_header_t header;
+			u32 pid;
+			uint8_t           extra[512];
+		} request_msg_t;
+
+		typedef struct {
+			mach_msg_header_t        header;
+			mach_msg_body_t          body;
+			mach_msg_port_descriptor_t sem_desc[3]; // carries the semaphore right
+		} reply_msg_t;
+
+		kern_return_t kr;
+		mach_port_t port = m_semaphoreServiceInfo->port;
+
+		while (m_semaphoreServiceInfo)
+		{
+			request_msg_t req;
+			memset(&req, 0, sizeof(req));
+			kr = mach_msg(&req.header, MACH_RCV_MSG, 0, sizeof(req), port, MACH_MSG_TIMEOUT_NONE, MACH_PORT_NULL);
+			if (kr != KERN_SUCCESS)
+				break;
+			UBA_ASSERTF(kr == KERN_SUCCESS, "ERR: %u", kr);
+			
+			reply_msg_t rep;
+			memset(&rep, 0, sizeof(rep));
+			rep.header.msgh_bits = MACH_MSGH_BITS(MACH_MSG_TYPE_MOVE_SEND_ONCE, 0) | MACH_MSGH_BITS_COMPLEX;
+			rep.header.msgh_size = sizeof(rep);
+			rep.header.msgh_remote_port = req.header.msgh_remote_port; // destination is child's reply port
+			rep.header.msgh_local_port  = MACH_PORT_NULL;
+			rep.body.msgh_descriptor_count = 3;
+
+			//printf("GOT REQUEST FOR PID %u\n", req.pid);
+
+			SCOPED_FUTEX_READ(m_processesLock, plock);
+			auto findIt = m_processes.find(req.pid);
+			if (findIt == m_processes.end())
+				continue;
+			ProcessImpl& proc = *(ProcessImpl*)findIt->second.m_process;
+
+			rep.sem_desc[0].name        = proc.m_cancelEvent.GetHandle();
+			rep.sem_desc[0].disposition = MACH_MSG_TYPE_COPY_SEND;
+			rep.sem_desc[0].type        = MACH_MSG_PORT_DESCRIPTOR;
+			rep.sem_desc[1].name        = proc.m_writeEvent.GetHandle();
+			rep.sem_desc[1].disposition = MACH_MSG_TYPE_COPY_SEND;
+			rep.sem_desc[1].type        = MACH_MSG_PORT_DESCRIPTOR;
+			rep.sem_desc[2].name        = proc.m_readEvent.GetHandle();
+			rep.sem_desc[2].disposition = MACH_MSG_TYPE_COPY_SEND;
+			rep.sem_desc[2].type        = MACH_MSG_PORT_DESCRIPTOR;
+
+			//printf("SENDING HANDLES FOR PID %u\n", req.pid);
+
+			kr = mach_msg(&rep.header, MACH_SEND_MSG, rep.header.msgh_size, 0, MACH_PORT_NULL, MACH_MSG_TIMEOUT_NONE, MACH_PORT_NULL);
+			if (kr != KERN_SUCCESS)
+				break;
+			UBA_ASSERTF(kr == KERN_SUCCESS, "ERR: %u", kr);
+		}
+		//printf("EXITED SEMA SERVICE\n");
+	}
+#endif
+
 	void GenerateNameForProcess(StringBufferBase& out, const tchar* arguments, u32 counterSuffix)
 	{
 		const tchar* start = arguments;
diff --git a/Engine/Source/Programs/UnrealBuildAccelerator/Common/Public/UbaSession.h b/Engine/Source/Programs/UnrealBuildAccelerator/Common/Public/UbaSession.h
--- a/Engine/Source/Programs/UnrealBuildAccelerator/Common/Public/UbaSession.h
+++ b/Engine/Source/Programs/UnrealBuildAccelerator/Common/Public/UbaSession.h
@@ -9,6 +9,7 @@
 #include "UbaFileMapping.h"
 #include "UbaProcessHandle.h"
 #include "UbaProcessStartInfo.h"
+#include "UbaProtocol.h"
 #include "UbaRootPaths.h"
 #include "UbaSessionCreateInfo.h"
 #include "UbaStats.h"
@@ -358,6 +359,17 @@ namespace uba
 		UnorderedMap<TString, bool> m_isX64Application;
 		#endif
 
+		#if UBA_USE_NATIVE_MAC_SEMAPHORES
+		u32 GetSessionId();
+		void StartSemaphoreService();
+		void StopSemaphoreService();
+		void ThreadSemaphoreService();
+		u32 m_sessionId;
+		struct SemaphoreServiceInfo;
+		Thread m_semaphoreServiceThread;
+		SemaphoreServiceInfo* m_semaphoreServiceInfo;
+		#endif
+
 		VolumeCache m_volumeCache;
 
 		DependencyCrawler m_dependencyCrawler;
diff --git a/Engine/Source/Programs/UnrealBuildAccelerator/Core/Private/UbaEvent.cpp b/Engine/Source/Programs/UnrealBuildAccelerator/Core/Private/UbaEvent.cpp
--- a/Engine/Source/Programs/UnrealBuildAccelerator/Core/Private/UbaEvent.cpp
+++ b/Engine/Source/Programs/UnrealBuildAccelerator/Core/Private/UbaEvent.cpp
@@ -2,6 +2,7 @@
 
 #include "UbaEvent.h"
 #include "UbaPlatform.h"
+#include "UbaProtocol.h"
 #include "UbaStringBuffer.h"
 #include "UbaSynchronization.h"
 
@@ -11,6 +12,11 @@
 #include <semaphore.h>
 #endif
 
+#if UBA_USE_NATIVE_MAC_SEMAPHORES
+#include <mach/mach.h>
+#include <mach/semaphore.h>
+#endif
+
 #define UBA_TEST_WAIT_QUALITY 0
 
 namespace uba
@@ -19,14 +25,16 @@ namespace uba
 
 	u64 GetMonoticTimeNs();
 
-	inline u64 ToNanoSeconds(const struct timespec& ts)
+	template<typename Timespec>
+	inline u64 ToNanoSeconds(const Timespec& ts)
 	{
 		return u64(ts.tv_sec)*1'000'000'000 + ts.tv_nsec;
 	}
 
-	inline struct timespec ToTimeSpecFromNano(u64 nanoSeconds)
+	template<typename Timespec>
+	inline Timespec ToTimeSpecFromNano(u64 nanoSeconds)
 	{
-		struct timespec ts;
+		Timespec ts;
 		ts.tv_sec = nanoSeconds / 1'000'000'000;
 		ts.tv_nsec = nanoSeconds % 1'000'000'000;
 		return ts;
@@ -227,15 +235,15 @@ namespace uba
 					}
 					else  // timed wait.
 					{
-#if PLATFORM_MAC
-						struct timespec timeout = ToTimeSpecFromNano(timeoutNs);
+						#if PLATFORM_MAC
+						auto timeout = ToTimeSpecFromNano<timespec>(timeoutNs);
 						int rc = pthread_cond_timedwait_relative_np(&m_condition, &m_mutex, &timeout); // unlocks Mutex while blocking...
 						UBA_ASSERTF((rc == 0) || (rc == ETIMEDOUT), "pthread_cond_timedwait_relative_np failed"); (void)rc;
-#else
-						struct timespec timeout = ToTimeSpecFromNano(startTimeNs + timeoutNs);
+						#else
+						auto timeout = ToTimeSpecFromNano<timespec>(startTimeNs + timeoutNs);
 						int rc = pthread_cond_timedwait(&m_condition, &m_mutex, &timeout); // unlocks Mutex while blocking...
 						UBA_ASSERTF((rc == 0) || (rc == ETIMEDOUT), "pthread_cond_timedwait failed"); (void)rc;
-#endif
+						#endif
 						#if UBA_TEST_WAIT_QUALITY
 						++loop;
 						#endif
@@ -292,14 +300,31 @@ namespace uba
 	public:
 		~SemaphoreEvent()
 		{
-			if (m_snd)
-				Destroy();
+			Destroy();
 		}
 
-		bool Create(const tchar* name, bool isOwner)
+		bool Create(StringView name, bool isOwner)
 		{
-			UBA_ASSERT(!m_snd);
-
+#if UBA_USE_NATIVE_MAC_SEMAPHORES
+			UBA_ASSERT(!m_sem && !m_isOwner);
+			if (isOwner)
+			{
+				kern_return_t kr;
+				kr = semaphore_create(mach_task_self(), &m_sem, SYNC_POLICY_FIFO, 0);
+				UBA_ASSERT(kr == KERN_SUCCESS);
+				mach_port_insert_right(mach_task_self(), m_sem, m_sem, MACH_MSG_TYPE_MAKE_SEND);
+				m_isOwner = true;
+				//printf("CREATE OWNER WITH %u\n", m_sem);
+			}
+			else
+			{
+				u64 sem;
+				Parse(sem, name.data, name.count);
+				m_sem = (semaphore_t)sem;
+				//printf("CREATE CLIENT WITH %u\n", m_sem);
+			}
+#else
+			UBA_ASSERT(!m_snd && !m_isOwner);
 			StringBuffer<> n1(name);
 			StringBuffer<> n2(name);
 			n1.Append(TCV("_cl"));
@@ -315,26 +340,34 @@ namespace uba
 				flags |= O_CREAT;
 				sndName = n2.data;
 				rvcName = n1.data;
-				m_name = strdup(name);
+				m_name = TStrdup(name.data);
 				m_isOwner = true;
 			};
 
 			m_snd = sem_open(sndName, flags, mode, 0);
-			m_rsv = sem_open(rvcName, flags, mode, 0);
+			m_rcv = sem_open(rvcName, flags, mode, 0);
 
 			UBA_ASSERTF(m_snd != SEM_FAILED, "sem_open failed");
-			UBA_ASSERTF(m_rsv != SEM_FAILED, "sem_open failed");
-
+			UBA_ASSERTF(m_rcv != SEM_FAILED, "sem_open failed");
+#endif
 			return true;
 		}
 
 		void Destroy()
 		{
-			sem_close((sem_t*)m_snd);
+#if UBA_USE_NATIVE_MAC_SEMAPHORES
+			if (!m_sem)
+				return;
+			if (m_isOwner)
+				semaphore_destroy(mach_task_self(), m_sem);
+#else
+			if (!m_snd)
+				return;
+			sem_close(m_snd);
 			m_snd = nullptr;
 
-			sem_close((sem_t*)m_rsv);
-			m_rsv = nullptr;
+			sem_close(m_rcv);
+			m_rcv = nullptr;
 
 			if (m_isOwner)
 			{
@@ -347,20 +380,32 @@ namespace uba
 				free(m_name);
 				m_isOwner = false;
 			}
+#endif
 		}
 
 		void Set()
 		{
-			if (!sem_post((sem_t*)m_snd))
+			#if UBA_USE_NATIVE_MAC_SEMAPHORES
+			semaphore_signal(m_sem);
+			#else
+			if (!sem_post(m_snd))
 				return;
 			UBA_ASSERTF(false, "sem_post failed");
+			#endif
 		}
 
 		bool IsSet(u32 timeoutMs = ~0u)
 		{
+#if UBA_USE_NATIVE_MAC_SEMAPHORES
+			if (timeoutMs == ~0u)
+				return semaphore_wait(m_sem) == KERN_SUCCESS;
+			mach_timespec_t ts = ToTimeSpecFromNano<mach_timespec_t>(u64(timeoutMs) * 1'000'000);
+			kern_return_t kr = semaphore_timedwait(m_sem, ts);
+			return kr == KERN_SUCCESS;
+#else
 			if (timeoutMs == ~0u)
 			{
-				while (sem_wait((sem_t*)m_rsv) == -1)
+				while (sem_wait(m_rcv) == -1)
 				{
 					if (errno == EINTR)
 						continue;
@@ -377,8 +422,8 @@ namespace uba
 			u64 absTimeNano = ToNanoSeconds(now) + timeoutNs;
 
 			#if PLATFORM_LINUX
-			struct timespec absTime = ToTimeSpecFromNano(absTimeNano);
-			while (sem_timedwait((sem_t*)m_rsv, &absTime) == -1)
+			auto absTime = ToTimeSpecFromNano<timespec>(absTimeNano);
+			while (sem_timedwait(m_rcv, &absTime) == -1)
 			{
 				if (errno == EINTR)
 					continue;
@@ -389,23 +434,38 @@ namespace uba
 			#else
 			while (true)
 			{
-				if (sem_trywait((sem_t*)m_rsv) == 0)
+				if (sem_trywait(m_rcv) == 0)
 					return true;
-				Sleep(1);
-				clock_gettime(CLOCK_REALTIME, &now);
 				if (ToNanoSeconds(now) > absTimeNano)
 					return false;
+				usleep(1);
+				clock_gettime(CLOCK_REALTIME, &now);
 			}
 			#endif
 			return true;
+#endif
+		}
+
+		void ToString(StringBufferBase& out)
+		{
+#if UBA_USE_NATIVE_MAC_SEMAPHORES
+			out.AppendValue(m_sem);
+#else
+			out.Append(m_name);
+#endif
 		}
 
 	private:
 		u8 m_dummy; // This is to map into SharedEvent under the hood
-		void* m_snd = nullptr;
-		void* m_rsv = nullptr;
-		tchar* m_name = nullptr;
 		bool m_isOwner = false;
+
+#if UBA_USE_NATIVE_MAC_SEMAPHORES
+		semaphore_t m_sem = MACH_PORT_NULL;
+#else
+		sem_t* m_snd = nullptr;
+		sem_t* m_rcv = nullptr;
+		tchar* m_name = nullptr;
+#endif
 	};
 
 #else
@@ -538,7 +598,7 @@ namespace uba
 		return ((EventImpl&)m_data).Create(manualReset, true);
 	}
 
-	bool SharedEvent::Create(const tchar* name, bool isOwner)
+	bool SharedEvent::Create(StringView name, bool isOwner)
 	{
 		UBA_ASSERT(*((u8*)m_data) == 0);
 		new (m_data) SemaphoreEvent();
@@ -590,5 +650,12 @@ namespace uba
 	{
 		return *((u8*)m_data) == 2;
 	}
+
+	void SharedEvent::ToString(StringBufferBase& out)
+	{
+		UBA_ASSERT(IsSemaphore());
+		((SemaphoreEvent&)m_data).ToString(out);
+	}
+
 #endif
 }
diff --git a/Engine/Source/Programs/UnrealBuildAccelerator/Core/Public/UbaEvent.h b/Engine/Source/Programs/UnrealBuildAccelerator/Core/Public/UbaEvent.h
--- a/Engine/Source/Programs/UnrealBuildAccelerator/Core/Public/UbaEvent.h
+++ b/Engine/Source/Programs/UnrealBuildAccelerator/Core/Public/UbaEvent.h
@@ -6,6 +6,9 @@
 
 namespace uba
 {
+	class StringBufferBase;
+	struct StringView;
+
 	class Event
 	{
 	public:
@@ -32,11 +35,12 @@ namespace uba
 		SharedEvent();
 		~SharedEvent();
 		bool Create(bool manualReset);
-		bool Create(const tchar* name, bool isOwner);
+		bool Create(StringView name, bool isOwner);
 		void Destroy();
 		void Set();
 		bool IsSet(u32 timeOutMs = ~0u);
 		bool IsSemaphore();
+		void ToString(StringBufferBase& out);
 	private:
 		#if PLATFORM_MAC || defined(__aarch64__)
 		u64 m_data[16];
diff --git a/Engine/Source/Programs/UnrealBuildAccelerator/Core/Public/UbaProtocol.h b/Engine/Source/Programs/UnrealBuildAccelerator/Core/Public/UbaProtocol.h
--- a/Engine/Source/Programs/UnrealBuildAccelerator/Core/Public/UbaProtocol.h
+++ b/Engine/Source/Programs/UnrealBuildAccelerator/Core/Public/UbaProtocol.h
@@ -59,6 +59,8 @@ namespace uba
 
 #define UBA_ENABLE_ON_DISK_FILE_MAPPINGS 0
 
+#define UBA_USE_NATIVE_MAC_SEMAPHORES 0
+
 // Currently only used for detoured process
 #if UBA_DEBUG
 #define UBA_DEBUG_LOG_ENABLED 1
diff --git a/Engine/Source/Programs/UnrealBuildAccelerator/Detours/Private/PosixOS/UbaDetoursFunctionsPosix.cpp b/Engine/Source/Programs/UnrealBuildAccelerator/Detours/Private/PosixOS/UbaDetoursFunctionsPosix.cpp
--- a/Engine/Source/Programs/UnrealBuildAccelerator/Detours/Private/PosixOS/UbaDetoursFunctionsPosix.cpp
+++ b/Engine/Source/Programs/UnrealBuildAccelerator/Detours/Private/PosixOS/UbaDetoursFunctionsPosix.cpp
@@ -2168,6 +2168,10 @@ int shared_posix_spawn(pid_t* pid, const char* path, const posix_spawn_file_acti
 	StringBuffer<32> rulesStr;
 	StringBuffer<512> logFile;
 
+	#if UBA_USE_NATIVE_MAC_SEMAPHORES
+	StringBuffer<64> pidVar;
+	#endif
+	
 	{
 		TimerScope ts(g_stats.createProcess);
 		SCOPED_WRITE_LOCK(g_communicationLock, pcs);
@@ -2211,6 +2215,10 @@ int shared_posix_spawn(pid_t* pid, const char* path, const posix_spawn_file_acti
 
 		realApplication = reader.ReadString();
 
+		#if UBA_USE_NATIVE_MAC_SEMAPHORES
+		pidVar.Append("UBA_PID=").AppendValue(reader.ReadU32());
+		#endif
+
 		comIdVar.Append("UBA_COMID=").AppendValue(reader.ReadU64()).Append('+').AppendValue(reader.ReadU32());
 
 		logFile.Append("UBA_LOGFILE=");
@@ -2229,6 +2237,11 @@ int shared_posix_spawn(pid_t* pid, const char* path, const posix_spawn_file_acti
 		envvars.push_back(rulesStr.data);
 		envvars.push_back(logFile.data);
 		//envvars.push_back("LD_DEBUG=bindings");
+
+		#if UBA_USE_NATIVE_MAC_SEMAPHORES
+		envvars.push_back(pidVar.data);
+		#endif
+		
 		envvars.push_back(nullptr);
 	}
 
diff --git a/Engine/Source/Programs/UnrealBuildAccelerator/Detours/Private/PosixOS/UbaDetoursMainPosix.cpp b/Engine/Source/Programs/UnrealBuildAccelerator/Detours/Private/PosixOS/UbaDetoursMainPosix.cpp
--- a/Engine/Source/Programs/UnrealBuildAccelerator/Detours/Private/PosixOS/UbaDetoursMainPosix.cpp
+++ b/Engine/Source/Programs/UnrealBuildAccelerator/Detours/Private/PosixOS/UbaDetoursMainPosix.cpp
@@ -8,6 +8,9 @@
 #include <dlfcn.h>
 #if PLATFORM_LINUX
 #include <sys/prctl.h>
+#elif UBA_USE_NATIVE_MAC_SEMAPHORES
+#include <mach/mach.h>
+#include <servers/bootstrap.h>
 #endif
 using namespace uba;
 
@@ -81,18 +84,68 @@ static void __attribute__((constructor(102))) PreInitCtor()
 
 	if (UseSemaphoreForSharedEvents())
 	{
+		g_readEvent = new SharedEvent();
+		g_writeEvent = new SharedEvent();
+		g_cancelEvent = new SharedEvent();
+
+		#if UBA_USE_NATIVE_MAC_SEMAPHORES
+
+		//int tty = (FileHandle)open("/dev/tty", O_WRONLY);
+		//auto log = [tty](const char* str, int len) { write(tty, str, len); };
+		//auto logLine = [&](const char* str, bool newline = true) { log(str, TStrlen(str)); if (newline) log("\n", 1); };
+
+		const char* pidStr = getenv("UBA_PID");
+		UBA_ASSERT(pidStr);
+		u64 pid;
+		Parse(pid, pidStr, TStrlen(pidStr));
+		unsetenv("UBA_PID");
+
+		StringBuffer<> service;
+		service.Append("UbaSemaphoreService");
+		service.Append(getenv("UBA_SESSION_ID"));
+		mach_port_t server_send = MACH_PORT_NULL;
+		kern_return_t kr = bootstrap_look_up(bootstrap_port, service.data, &server_send);
+		
+		mach_port_t reply_port = MACH_PORT_NULL;
+		kr = mach_port_allocate(mach_task_self(), MACH_PORT_RIGHT_RECEIVE, &reply_port);
+		UBA_ASSERT(kr == KERN_SUCCESS);
+
+		struct { mach_msg_header_t header; u32 pid; } req;
+		memset(&req, 0, sizeof(req));
+		req.header.msgh_bits = MACH_MSGH_BITS(MACH_MSG_TYPE_COPY_SEND, MACH_MSG_TYPE_MAKE_SEND_ONCE);
+		req.header.msgh_size = sizeof(req);
+		req.header.msgh_remote_port = server_send;
+		req.header.msgh_local_port  = reply_port;
+		req.pid = u32(pid);
+		kr = mach_msg(&req.header, MACH_SEND_MSG, req.header.msgh_size, 0, MACH_PORT_NULL, MACH_MSG_TIMEOUT_NONE, MACH_PORT_NULL);
+		UBA_ASSERT(kr == KERN_SUCCESS);
+
+		struct { mach_msg_header_t header; mach_msg_body_t body; mach_msg_port_descriptor_t sem_desc[3]; uint8_t extra[512]; } rep;
+		memset(&rep, 0, sizeof(rep));
+		kr = mach_msg(&rep.header, MACH_RCV_MSG, 0, sizeof(rep), reply_port, MACH_MSG_TIMEOUT_NONE, MACH_PORT_NULL);
+		UBA_ASSERT(kr == KERN_SUCCESS);
+
+		StringBuffer<32> b;
+		b.AppendValue(rep.sem_desc[0].name);
+		g_cancelEvent->Create(b, false);
+		b.Clear().AppendValue(rep.sem_desc[1].name);
+		g_readEvent->Create(b, false);
+		b.Clear().AppendValue(rep.sem_desc[2].name);
+		g_writeEvent->Create(b, false);
+
+		#else
+
 		StringBuffer<256> name;
 		name.Append(TCV("uba")).AppendValue(comIdUid).Append('_').AppendValue(comIdOffset);
 		u64 len = name.count;
-		g_cancelEvent = new SharedEvent();
 		name.Append("_cancel");
-		g_cancelEvent->Create(name.data, false);
-		g_readEvent = new SharedEvent();
+		g_cancelEvent->Create(name, false);
 		name.Resize(len).Append("_read");
-		g_readEvent->Create(name.data, false);
-		g_writeEvent = new SharedEvent();
+		g_readEvent->Create(name, false);
 		name.Resize(len).Append("_write");
-		g_writeEvent->Create(name.data, false);
+		g_writeEvent->Create(name, false);
+
+		#endif
 	}
 	else
 	{
diff --git a/Engine/Source/Programs/UnrealBuildAccelerator/Test/Private/UbaTestBasics.cpp b/Engine/Source/Programs/UnrealBuildAccelerator/Test/Private/UbaTestBasics.cpp
--- a/Engine/Source/Programs/UnrealBuildAccelerator/Test/Private/UbaTestBasics.cpp
+++ b/Engine/Source/Programs/UnrealBuildAccelerator/Test/Private/UbaTestBasics.cpp
@@ -158,9 +158,11 @@ namespace uba
 	{
 		#if !PLATFORM_WINDOWS
 		SharedEvent ev;
-		ev.Create("uba_tse", true);
+		ev.Create(TCV("uba_tse"), true);
+		StringBuffer<64> tmp;
+		ev.ToString(tmp);
 		SharedEvent ev2;
-		ev2.Create("uba_tse", false);
+		ev2.Create(tmp, false);
 		if (!TestEventsImpl(logger, rootDir, ev, ev2))
 			return false;
 		#endif
