[UBA]
* Fixed compile errors when native mach semaphores is enabled

diff --git a/Engine/Source/Programs/UnrealBuildAccelerator/Common/Private/UbaProcess.cpp b/Engine/Source/Programs/UnrealBuildAccelerator/Common/Private/UbaProcess.cpp
--- a/Engine/Source/Programs/UnrealBuildAccelerator/Common/Private/UbaProcess.cpp
+++ b/Engine/Source/Programs/UnrealBuildAccelerator/Common/Private/UbaProcess.cpp
@@ -2115,7 +2115,7 @@ namespace uba
 				envvars.push_back(workingDir.data);
 				envvars.push_back(rulesStr.data);
 
-				#if PLATFORM_MAC
+				#if UBA_USE_NATIVE_MAC_SEMAPHORES
 				envvars.push_back(sessionId.data);
 				envvars.push_back(pidVar.data);
 				#endif
diff --git a/Engine/Source/Programs/UnrealBuildAccelerator/Common/Private/UbaSession.cpp b/Engine/Source/Programs/UnrealBuildAccelerator/Common/Private/UbaSession.cpp
--- a/Engine/Source/Programs/UnrealBuildAccelerator/Common/Private/UbaSession.cpp
+++ b/Engine/Source/Programs/UnrealBuildAccelerator/Common/Private/UbaSession.cpp
@@ -4117,15 +4117,19 @@ namespace uba
 				continue;
 			ProcessImpl& proc = *(ProcessImpl*)findIt->second.m_process;
 
-			rep.sem_desc[0].name        = proc.m_cancelEvent.GetHandle();
-			rep.sem_desc[0].disposition = MACH_MSG_TYPE_COPY_SEND;
-			rep.sem_desc[0].type        = MACH_MSG_PORT_DESCRIPTOR;
-			rep.sem_desc[1].name        = proc.m_writeEvent.GetHandle();
-			rep.sem_desc[1].disposition = MACH_MSG_TYPE_COPY_SEND;
-			rep.sem_desc[1].type        = MACH_MSG_PORT_DESCRIPTOR;
-			rep.sem_desc[2].name        = proc.m_readEvent.GetHandle();
-			rep.sem_desc[2].disposition = MACH_MSG_TYPE_COPY_SEND;
-			rep.sem_desc[2].type        = MACH_MSG_PORT_DESCRIPTOR;
+			auto setSem = [&](int index, SharedEvent& event)
+			{
+				StringBuffer<32> str;
+				event.ToString(str);
+				u64 handle;
+				str.Parse(handle);
+				rep.sem_desc[index].name        = handle;
+				rep.sem_desc[index].disposition = MACH_MSG_TYPE_COPY_SEND;
+				rep.sem_desc[index].type        = MACH_MSG_PORT_DESCRIPTOR;
+			};
+			setSem(0, proc.m_cancelEvent);
+			setSem(1, proc.m_writeEvent);
+			setSem(2, proc.m_readEvent);
 
 			//printf("SENDING HANDLES FOR PID %u\n", req.pid);
 
