[UBA]
* Added so pthread cancellation is disabled on uba threads
* Moved definitions for features that are not toggleable to UbaDefinitions.h

diff --git a/Engine/Source/Programs/UnrealBuildAccelerator/Common/Private/UbaProcess.cpp b/Engine/Source/Programs/UnrealBuildAccelerator/Common/Private/UbaProcess.cpp
--- a/Engine/Source/Programs/UnrealBuildAccelerator/Common/Private/UbaProcess.cpp
+++ b/Engine/Source/Programs/UnrealBuildAccelerator/Common/Private/UbaProcess.cpp
@@ -2,6 +2,7 @@
 
 #include "UbaProcess.h"
 #include "UbaApplicationRules.h"
+#include "UbaDefinitions.h"
 #include "UbaFileAccessor.h"
 #include "UbaProtocol.h"
 #include "UbaProcessStats.h"
diff --git a/Engine/Source/Programs/UnrealBuildAccelerator/Common/Private/UbaThread.cpp b/Engine/Source/Programs/UnrealBuildAccelerator/Common/Private/UbaThread.cpp
--- a/Engine/Source/Programs/UnrealBuildAccelerator/Common/Private/UbaThread.cpp
+++ b/Engine/Source/Programs/UnrealBuildAccelerator/Common/Private/UbaThread.cpp
@@ -1,6 +1,7 @@
 // Copyright Epic Games, Inc. All Rights Reserved.
 
 #include "UbaThread.h"
+#include "UbaDefinitions.h"
 #include "UbaEnvironment.h"
 #include "UbaPlatform.h"
 
@@ -104,16 +105,24 @@ namespace uba
 		err = pthread_create(&pth, &tattr, [](void* p) -> void*
 			{
 				// Ignore sigint. Uba is designed for someone else to cancel it
-				#if !PLATFORM_WINDOWS
 				sigset_t set;
 				sigemptyset(&set);
 				sigaddset(&set, SIGINT);
 				pthread_sigmask(SIG_BLOCK, &set, NULL);		
+
+				#if UBA_DISABLE_PTHREAD_CANCELLATIONS
+				int old;
+				pthread_setcancelstate(PTHREAD_CANCEL_DISABLE, &old);
 				#endif
 
 				auto& t = *(Thread*)p;
 				int res = t.m_func();
 				t.m_finished.Set();
+				
+				#if UBA_DISABLE_PTHREAD_CANCELLATIONS
+				pthread_setcancelstate(old, NULL);
+				#endif
+
 				return (void*)(uintptr_t)res;
 			}, this);
 		UBA_ASSERT(err == 0); (void)err;
diff --git a/Engine/Source/Programs/UnrealBuildAccelerator/Common/Public/UbaSession.h b/Engine/Source/Programs/UnrealBuildAccelerator/Common/Public/UbaSession.h
--- a/Engine/Source/Programs/UnrealBuildAccelerator/Common/Public/UbaSession.h
+++ b/Engine/Source/Programs/UnrealBuildAccelerator/Common/Public/UbaSession.h
@@ -5,11 +5,11 @@
 #include "UbaBottleneck.h"
 #include "UbaDependencyCrawler.h"
 #include "UbaFile.h"
+#include "UbaDefinitions.h"
 #include "UbaDirectoryTable.h"
 #include "UbaFileMapping.h"
 #include "UbaProcessHandle.h"
 #include "UbaProcessStartInfo.h"
-#include "UbaProtocol.h"
 #include "UbaRootPaths.h"
 #include "UbaSessionCreateInfo.h"
 #include "UbaStats.h"
diff --git a/Engine/Source/Programs/UnrealBuildAccelerator/Core/Private/UbaEvent.cpp b/Engine/Source/Programs/UnrealBuildAccelerator/Core/Private/UbaEvent.cpp
--- a/Engine/Source/Programs/UnrealBuildAccelerator/Core/Private/UbaEvent.cpp
+++ b/Engine/Source/Programs/UnrealBuildAccelerator/Core/Private/UbaEvent.cpp
@@ -1,8 +1,8 @@
 // Copyright Epic Games, Inc. All Rights Reserved.
 
 #include "UbaEvent.h"
+#include "UbaDefinitions.h"
 #include "UbaPlatform.h"
-#include "UbaProtocol.h"
 #include "UbaStringBuffer.h"
 #include "UbaSynchronization.h"
 
@@ -432,14 +432,33 @@ namespace uba
 				UBA_ASSERTF(false, "sem_wait failed");
 			}
 			#else
+
+			u64 sleepNs = 50 * 1000; // 50 µs
+			const u64 sleepCapNs = 10 * 1000 * 1000; // 10 ms
+
 			while (true)
 			{
 				if (sem_trywait(m_rcv) == 0)
 					return true;
-				if (ToNanoSeconds(now) > absTimeNano)
+				if (errno != EAGAIN)
 					return false;
-				usleep(1);
+
 				clock_gettime(CLOCK_REALTIME, &now);
+				u64 nowNano = ToNanoSeconds(now);
+				if (nowNano > absTimeNano)
+					return false;
+
+				u64 remaining = absTimeNano - nowNano;
+				u64 nap = sleepNs < remaining ? sleepNs : remaining;
+				auto deltaNs = ToTimeSpecFromNano<timespec>(nap);
+				nanosleep(&deltaNs, nullptr);
+
+				if (sleepNs < sleepCapNs)
+				{
+					sleepNs <<= 1;
+					if (sleepNs > sleepCapNs)
+						sleepNs = sleepCapNs;
+				}
 			}
 			#endif
 			return true;
diff --git a/Engine/Source/Programs/UnrealBuildAccelerator/Core/Public/UbaDefinitions.h b/Engine/Source/Programs/UnrealBuildAccelerator/Core/Public/UbaDefinitions.h
new file mode 100644
--- /dev/null
+++ b/Engine/Source/Programs/UnrealBuildAccelerator/Core/Public/UbaDefinitions.h
@@ -0,0 +1,13 @@
+// Copyright Epic Games, Inc. All Rights Reserved.
+
+#pragma once
+
+// Use native mach semaphores for IPC.
+#define UBA_USE_NATIVE_MAC_SEMAPHORES PLATFORM_MAC
+
+// Experimental. Testing to use file mappings that are backed by files to reduce page file pressure
+#define UBA_ENABLE_ON_DISK_FILE_MAPPINGS 0
+
+// Disable pthread cancellation to reduce work when waking up from event.
+// No uba code should rely on cancelling threads and shutdown is gracefully so this can be disabled
+#define UBA_DISABLE_PTHREAD_CANCELLATIONS 1
\ No newline at end of file
diff --git a/Engine/Source/Programs/UnrealBuildAccelerator/Core/Public/UbaProtocol.h b/Engine/Source/Programs/UnrealBuildAccelerator/Core/Public/UbaProtocol.h
--- a/Engine/Source/Programs/UnrealBuildAccelerator/Core/Public/UbaProtocol.h
+++ b/Engine/Source/Programs/UnrealBuildAccelerator/Core/Public/UbaProtocol.h
@@ -57,10 +57,6 @@ namespace uba
 	inline constexpr u32 EmergencyShutdownExitCode = 9886;
 }
 
-#define UBA_ENABLE_ON_DISK_FILE_MAPPINGS 0
-
-#define UBA_USE_NATIVE_MAC_SEMAPHORES 0
-
 // Currently only used for detoured process
 #if UBA_DEBUG
 #define UBA_DEBUG_LOG_ENABLED 1
diff --git a/Engine/Source/Programs/UnrealBuildAccelerator/Detours/Private/PosixOS/UbaDetoursFunctionsPosix.cpp b/Engine/Source/Programs/UnrealBuildAccelerator/Detours/Private/PosixOS/UbaDetoursFunctionsPosix.cpp
--- a/Engine/Source/Programs/UnrealBuildAccelerator/Detours/Private/PosixOS/UbaDetoursFunctionsPosix.cpp
+++ b/Engine/Source/Programs/UnrealBuildAccelerator/Detours/Private/PosixOS/UbaDetoursFunctionsPosix.cpp
@@ -4,6 +4,7 @@
 
 #include "UbaDetoursFileMappingTable.h"
 #include "UbaBinaryReaderWriter.h"
+#include "UbaDefinitions.h""
 #include "UbaDirectoryTable.h"
 #include "UbaPlatform.h"
 #include "UbaProcessStats.h"
diff --git a/Engine/Source/Programs/UnrealBuildAccelerator/Detours/Private/PosixOS/UbaDetoursMainPosix.cpp b/Engine/Source/Programs/UnrealBuildAccelerator/Detours/Private/PosixOS/UbaDetoursMainPosix.cpp
--- a/Engine/Source/Programs/UnrealBuildAccelerator/Detours/Private/PosixOS/UbaDetoursMainPosix.cpp
+++ b/Engine/Source/Programs/UnrealBuildAccelerator/Detours/Private/PosixOS/UbaDetoursMainPosix.cpp
@@ -169,6 +169,10 @@ static void __attribute__((constructor(102))) PreInitCtor()
 	chdir(realCwd);
 	unsetenv("UBA_CWD");
 
+	// TODO: Don't know if we want this
+	//int old;
+	//pthread_setcancelstate(PTHREAD_CANCEL_DISABLE, &old);
+
 	Init();
 }
 
diff --git a/Engine/Source/Programs/UnrealBuildAccelerator/Detours/Private/Windows/UbaDetoursFunctionsWin.cpp b/Engine/Source/Programs/UnrealBuildAccelerator/Detours/Private/Windows/UbaDetoursFunctionsWin.cpp
--- a/Engine/Source/Programs/UnrealBuildAccelerator/Detours/Private/Windows/UbaDetoursFunctionsWin.cpp
+++ b/Engine/Source/Programs/UnrealBuildAccelerator/Detours/Private/Windows/UbaDetoursFunctionsWin.cpp
@@ -43,6 +43,7 @@ DETOURED_FUNCTIONS_MEMORY_NON_WINE
 #endif
 
 #include "UbaBinaryParser.h"
+#include "UbaDefinitions.h"
 #include "UbaDirectoryTable.h"
 #include "UbaProcessStats.h"
 #include "UbaProtocol.h"
