[UBA]
* Moved the various event implementation to separate files

diff --git a/Engine/Source/Programs/UnrealBuildAccelerator/Core/Private/UbaEvent.cpp b/Engine/Source/Programs/UnrealBuildAccelerator/Core/Private/UbaEvent.cpp
--- a/Engine/Source/Programs/UnrealBuildAccelerator/Core/Private/UbaEvent.cpp
+++ b/Engine/Source/Programs/UnrealBuildAccelerator/Core/Private/UbaEvent.cpp
@@ -6,510 +6,24 @@
 #include "UbaStringBuffer.h"
 #include "UbaSynchronization.h"
 
-#if !PLATFORM_WINDOWS
-#include <sys/time.h>
-#include <new>
-#include <semaphore.h>
-#endif
-
-#if UBA_USE_NATIVE_MAC_SEMAPHORES
-#include <mach/mach.h>
-#include <mach/semaphore.h>
-#endif
-
 #define UBA_TEST_WAIT_QUALITY 0
 
+#include "UbaEventPThread.inl"
+#include "UbaEventMachSemaphore.inl"
+#include "UbaEventPosixSemaphore.inl"
+#include "UbaEventWin32.inl"
+
 namespace uba
 {
-#if !PLATFORM_WINDOWS
-
-	u64 GetMonoticTimeNs();
-
-	template<typename Timespec>
-	inline u64 ToNanoSeconds(const Timespec& ts)
-	{
-		return u64(ts.tv_sec)*1'000'000'000 + ts.tv_nsec;
-	}
-
-	template<typename Timespec>
-	inline Timespec ToTimeSpecFromNano(u64 nanoSeconds)
-	{
-		Timespec ts;
-		ts.tv_sec = nanoSeconds / 1'000'000'000;
-		ts.tv_nsec = nanoSeconds % 1'000'000'000;
-		return ts;
-	}
-
-	struct EventImpl
-	{
-		EventImpl() = default;
-		~EventImpl() { Destroy(); }
-
-		bool Create(bool mr, bool shared = false)
-		{
-			static_assert(Atomic<u32>::is_always_lock_free, "atomic<T> not lock free, can't work in shared mem");
-			static_assert(Atomic<TriggerType>::is_always_lock_free, "atomic<T> not lock free, can't work in shared mem");
-			static_assert(Atomic<bool>::is_always_lock_free, "atomic<T> not lock free, can't work in shared mem");
-
-			UBA_ASSERTF(!m_initialized, "Can't create already created Event");
-			m_manualReset = mr;
-
-			pthread_mutexattr_t attrmutex;
-			if (pthread_mutexattr_init(&attrmutex) != 0)
-			{
-				UBA_ASSERTF(false, "pthread_mutexattr_init failed");
-				return false;
-			}
-
-			if (shared)
-			{
-				if (pthread_mutexattr_setpshared(&attrmutex, PTHREAD_PROCESS_SHARED) != 0)
-				{
-					UBA_ASSERTF(false, "pthread_mutexattr_setpshared failed");
-					return false;
-				}
-				#if PLATFORM_LINUX
-				if (pthread_mutexattr_setrobust(&attrmutex, PTHREAD_MUTEX_ROBUST) != 0)
-				{
-					UBA_ASSERTF(false, "pthread_mutexattr_setrobust failed");
-					return false;
-				}
-				#endif
-			}
-
-			if (pthread_mutex_init(&m_mutex, &attrmutex) != 0)
-			{
-				UBA_ASSERTF(false, "pthread_mutex_init failed");
-				return false;
-			}
-			pthread_mutexattr_destroy(&attrmutex);
-
-			pthread_condattr_t attrcond;
-			if (pthread_condattr_init(&attrcond) != 0)
-			{
-				UBA_ASSERTF(false, "pthread_condattr_init failed");
-				return false;
-			}
-
-			#if PLATFORM_LINUX
-			if (pthread_condattr_setclock(&attrcond, CLOCK_MONOTONIC) != 0)
-			{
-				UBA_ASSERTF(false, "pthread_condattr_setclock failed");
-				return false;
-			}
-			#endif
-
-			if (shared)
-			{
-				if (pthread_condattr_setpshared(&attrcond, PTHREAD_PROCESS_SHARED) != 0)
-				{
-					UBA_ASSERTF(false, "pthread_condattr_setpshared failed");
-					return false;
-				}
-			}
-
-			if (pthread_cond_init(&m_condition, &attrcond) != 0)
-			{
-				UBA_ASSERTF(false, "pthread_cond_init failed");
-				pthread_mutex_destroy(&m_mutex);
-				return false;
-			}
-			pthread_condattr_destroy(&attrcond);
-
-			m_initialized = true;
-			return true;
-		}
-
-		void Destroy()
-		{
-			if (!m_initialized)
-				return;
-
-			LockEventMutex();
-			m_manualReset = true;
-			UnlockEventMutex();
-			Set();
-
-			LockEventMutex();
-			m_initialized = false;
-			while (m_waitingThreads)
-			{
-				UnlockEventMutex();
-				LockEventMutex();
-			}
-
-			pthread_cond_destroy(&m_condition);
-			UnlockEventMutex();
-			pthread_mutex_destroy(&m_mutex);
-		}
-
-		void Set()
-		{
-			if (!m_initialized)
-				return;
-			LockEventMutex();
-
-			if (m_manualReset)
-			{
-				m_triggered = TriggerType_All;
-				if (pthread_cond_broadcast(&m_condition) != 0)
-					UBA_ASSERTF(false, "pthread_cond_broadcast failed");
-			}
-			else
-			{
-				m_triggered = TriggerType_One;
-				if (pthread_cond_signal(&m_condition) != 0)  // may release multiple threads anyhow!
-					UBA_ASSERTF(false, "pthread_cond_signal failed");
-			}
-
-			UnlockEventMutex();
-		}
-
-		void Reset()
-		{
-			if (!m_initialized)
-				return;
-			LockEventMutex();
-			m_triggered = TriggerType_None;
-			UnlockEventMutex();
-		}
-
-		bool IsCreated()
-		{
-			return m_initialized;
-		}
-
-		bool IsSet(u32 timeoutMs = ~0u)
-		{
-			if (!m_initialized)
-				return false;
-
-			u64 startTimeNs = 0;
-
-			// We need to know the start time if we're going to do a timed wait.
-			if ((timeoutMs > 0) && (timeoutMs != ~0u))  // not polling and not infinite wait.
-				startTimeNs = GetMonoticTimeNs();
-
-			u64 timeoutNs = u64(timeoutMs) * 1'000'000;
-
-			#if UBA_TEST_WAIT_QUALITY
-			u32 loop = 0;
-			bool timedOut = false;
-			u64 initialStartTimeNs = startTimeNs;
-			u64 initialTimeoutNs = timeoutNs;
-			auto qg = MakeGuard([&]()
-				{
-					if (startTimeNs && timedOut)
-					{
-						u64 overtimeNs = GetMonoticTimeNs() - initialStartTimeNs - initialTimeoutNs;
-						printf("Loops: %u TimeOut: %ums Over-wait: %lluus\n", loop, timeoutMs, overtimeNs/1000);
-					}
-				});
-			#endif
-
-			LockEventMutex();
-			auto unlock = MakeGuard([this]() { UnlockEventMutex(); });
-
-			// loop in case we fall through the Condition signal but someone else claims the event.
-			do
-			{
-				if (m_triggered == TriggerType_One)
-				{
-					m_triggered = TriggerType_None;
-					return true;
-				}
-
-				if (m_triggered == TriggerType_All)
-				{
-					return true;
-				}
-
-				// No event signalled yet.
-				if (timeoutNs != 0)  // not just polling, wait on the condition variable.
-				{
-					++m_waitingThreads;
-					if (timeoutMs == ~0u) // infinite wait?
-					{
-						if (pthread_cond_wait(&m_condition, &m_mutex) != 0)  // unlocks Mutex while blocking...
-							UBA_ASSERTF(false, "pthread_cond_wait failed");
-					}
-					else  // timed wait.
-					{
-						#if PLATFORM_MAC
-						auto timeout = ToTimeSpecFromNano<timespec>(timeoutNs);
-						int rc = pthread_cond_timedwait_relative_np(&m_condition, &m_mutex, &timeout); // unlocks Mutex while blocking...
-						UBA_ASSERTF((rc == 0) || (rc == ETIMEDOUT), "pthread_cond_timedwait_relative_np failed"); (void)rc;
-						#else
-						auto timeout = ToTimeSpecFromNano<timespec>(startTimeNs + timeoutNs);
-						int rc = pthread_cond_timedwait(&m_condition, &m_mutex, &timeout); // unlocks Mutex while blocking...
-						UBA_ASSERTF((rc == 0) || (rc == ETIMEDOUT), "pthread_cond_timedwait failed"); (void)rc;
-						#endif
-						#if UBA_TEST_WAIT_QUALITY
-						++loop;
-						#endif
-
-						u64 nowNs = GetMonoticTimeNs();
-						u64 diffNs = nowNs - startTimeNs;
-
-						timeoutNs = diffNs >= timeoutNs ? 0 : (timeoutNs - diffNs);
-						startTimeNs = nowNs;
-					}
-					--m_waitingThreads;
-					UBA_ASSERTF(m_waitingThreads >= 0, "m_waitingThreads less than 0");
-				}
-
-			} while (timeoutNs != 0);
-
-			#if UBA_TEST_WAIT_QUALITY
-			timedOut = true;
-			#endif
-			return false;
-		}
-
-		bool GetManualReset() { return m_manualReset; }
-		void ResetManualReset(bool mr) {}
-
-		// Must be first!
-		EventImpl* next;
-
-	private:
-		enum TriggerType : u8 { TriggerType_None, TriggerType_One, TriggerType_All };
-
-		Atomic<bool> m_initialized;
-		Atomic<bool> m_manualReset;
-		Atomic<TriggerType> m_triggered;
-		Atomic<u32> m_waitingThreads;
-		pthread_mutex_t m_mutex;
-		pthread_cond_t m_condition;
-
-		void LockEventMutex()
-		{
-			int res = pthread_mutex_lock(&m_mutex);(void)res;
-			UBA_ASSERTF(res == 0, "pthread_mutex_lock failed (error code %i)", res);
-		}
-
-		void UnlockEventMutex()
-		{
-			int res = pthread_mutex_unlock(&m_mutex);
-			UBA_ASSERTF(res == 0, "pthread_mutex_unlock failed (error code %i)", res);
-		}
-	};
-
-	class SemaphoreEvent
-	{
-	public:
-		~SemaphoreEvent()
-		{
-			Destroy();
-		}
-
-		bool Create(StringView name, bool isOwner)
-		{
-#if UBA_USE_NATIVE_MAC_SEMAPHORES
-			UBA_ASSERT(!m_sem && !m_isOwner);
-			if (isOwner)
-			{
-				kern_return_t kr;
-				kr = semaphore_create(mach_task_self(), &m_sem, SYNC_POLICY_FIFO, 0);
-				UBA_ASSERT(kr == KERN_SUCCESS);
-				mach_port_insert_right(mach_task_self(), m_sem, m_sem, MACH_MSG_TYPE_MAKE_SEND);
-				m_isOwner = true;
-				//printf("CREATE OWNER WITH %u\n", m_sem);
-			}
-			else
-			{
-				u64 sem;
-				Parse(sem, name.data, name.count);
-				m_sem = (semaphore_t)sem;
-				//printf("CREATE CLIENT WITH %u\n", m_sem);
-			}
-#else
-			UBA_ASSERT(!m_snd && !m_isOwner);
-			StringBuffer<> n1(name);
-			StringBuffer<> n2(name);
-			n1.Append(TCV("_cl"));
-			n2.Append(TCV("_sv"));
-		
-			const tchar* sndName = n1.data;
-			const tchar* rvcName = n2.data;
-			int flags = 0;
-			mode_t mode = 0600;
-
-			if (isOwner)
-			{
-				flags |= O_CREAT;
-				sndName = n2.data;
-				rvcName = n1.data;
-				m_name = TStrdup(name.data);
-				m_isOwner = true;
-			};
-
-			m_snd = sem_open(sndName, flags, mode, 0);
-			m_rcv = sem_open(rvcName, flags, mode, 0);
-
-			UBA_ASSERTF(m_snd != SEM_FAILED, "sem_open failed");
-			UBA_ASSERTF(m_rcv != SEM_FAILED, "sem_open failed");
-#endif
-			return true;
-		}
-
-		void Destroy()
-		{
-#if UBA_USE_NATIVE_MAC_SEMAPHORES
-			if (!m_sem)
-				return;
-			if (m_isOwner)
-				semaphore_destroy(mach_task_self(), m_sem);
-#else
-			if (!m_snd)
-				return;
-			sem_close(m_snd);
-			m_snd = nullptr;
-
-			sem_close(m_rcv);
-			m_rcv = nullptr;
-
-			if (m_isOwner)
-			{
-				StringBuffer<> n1(m_name);
-				StringBuffer<> n2(m_name);
-				n1.Append(TCV("_cl"));
-				n2.Append(TCV("_sv"));
-				sem_unlink(n1.data);
-				sem_unlink(n2.data);
-				free(m_name);
-				m_isOwner = false;
-			}
-#endif
-		}
-
-		void Set()
-		{
-			#if UBA_USE_NATIVE_MAC_SEMAPHORES
-			semaphore_signal(m_sem);
-			#else
-			if (!sem_post(m_snd))
-				return;
-			UBA_ASSERTF(false, "sem_post failed");
-			#endif
-		}
-
-		bool IsSet(u32 timeoutMs = ~0u)
-		{
-#if UBA_USE_NATIVE_MAC_SEMAPHORES
-			if (timeoutMs == ~0u)
-				return semaphore_wait(m_sem) == KERN_SUCCESS;
-			mach_timespec_t ts = ToTimeSpecFromNano<mach_timespec_t>(u64(timeoutMs) * 1'000'000);
-			kern_return_t kr = semaphore_timedwait(m_sem, ts);
-			return kr == KERN_SUCCESS;
-#else
-			if (timeoutMs == ~0u)
-			{
-				while (sem_wait(m_rcv) == -1)
-				{
-					if (errno == EINTR)
-						continue;
-					UBA_ASSERTF(false, "sem_wait failed");
-				}
-				return errno == 0;
-			}
-
-			u64 timeoutNs = u64(timeoutMs) * 1'000'000;
-
-			struct timespec now;
-			clock_gettime(CLOCK_REALTIME, &now);
-
-			u64 absTimeNano = ToNanoSeconds(now) + timeoutNs;
-
-			#if PLATFORM_LINUX
-			auto absTime = ToTimeSpecFromNano<timespec>(absTimeNano);
-			while (sem_timedwait(m_rcv, &absTime) == -1)
-			{
-				if (errno == EINTR)
-					continue;
-				if (errno == ETIMEDOUT)
-					return false;
-				UBA_ASSERTF(false, "sem_wait failed");
-			}
-			#else
-
-			u64 sleepNs = 50 * 1000; // 50 µs
-			const u64 sleepCapNs = 10 * 1000 * 1000; // 10 ms
-
-			while (true)
-			{
-				if (sem_trywait(m_rcv) == 0)
-					return true;
-				if (errno != EAGAIN)
-					return false;
-
-				clock_gettime(CLOCK_REALTIME, &now);
-				u64 nowNano = ToNanoSeconds(now);
-				if (nowNano > absTimeNano)
-					return false;
-
-				u64 remaining = absTimeNano - nowNano;
-				u64 nap = sleepNs < remaining ? sleepNs : remaining;
-				auto deltaNs = ToTimeSpecFromNano<timespec>(nap);
-				nanosleep(&deltaNs, nullptr);
-
-				if (sleepNs < sleepCapNs)
-				{
-					sleepNs <<= 1;
-					if (sleepNs > sleepCapNs)
-						sleepNs = sleepCapNs;
-				}
-			}
-			#endif
-			return true;
-#endif
-		}
-
-		void ToString(StringBufferBase& out)
-		{
-#if UBA_USE_NATIVE_MAC_SEMAPHORES
-			out.AppendValue(m_sem);
-#else
-			out.Append(m_name);
-#endif
-		}
-
-	private:
-		u8 m_dummy; // This is to map into SharedEvent under the hood
-		bool m_isOwner = false;
-
-#if UBA_USE_NATIVE_MAC_SEMAPHORES
-		semaphore_t m_sem = MACH_PORT_NULL;
-#else
-		sem_t* m_snd = nullptr;
-		sem_t* m_rcv = nullptr;
-		tchar* m_name = nullptr;
-#endif
-	};
-
-#else
-	struct EventImpl
-	{
-		HANDLE handle;
-		union
-		{
-			EventImpl* next;
-			bool manualReset;
-		};
-
-		bool Create(bool mr) { handle = CreateEvent(nullptr, mr, false, NULL); return handle != 0; }
-		void Set() { SetEvent(handle); }
-		void Reset() { ResetEvent(handle); }
-		bool IsSet(u32 timeOutMs) { return WaitForSingleObject(handle, timeOutMs) == 0;} 
-		bool GetManualReset() { return manualReset; }
-		void ResetManualReset(bool mr) { manualReset = mr; }
-	};
-#endif
+	#if !PLATFORM_WINDOWS
+	using EventImpl = EventPThread;
+	#else
+	using EventImpl = EventWin32;
+	#endif
 
 	Futex g_firstEventLock[2];
 	EventImpl* g_firstEvent[2];
 
-
 	Event::Event(bool manualReset)
 	{
 		Create(manualReset);
@@ -549,7 +63,7 @@ namespace uba
 	{
 		if (!m_impl)
 			return;
-		auto impl = m_impl;
+		auto impl = (EventImpl*)m_impl;
 		m_impl = nullptr;
 		bool manualReset = impl->GetManualReset();
 		auto*& first = g_firstEvent[manualReset];
@@ -561,13 +75,13 @@ namespace uba
 	void Event::Set()
 	{
 		if (m_impl)
-			m_impl->Set();
+			((EventImpl*)m_impl)->Set();
 	}
 
 	void Event::Reset()
 	{
 		if (m_impl)
-			m_impl->Reset();
+			((EventImpl*)m_impl)->Reset();
 	}
 
 	bool Event::IsCreated()
@@ -578,7 +92,7 @@ namespace uba
 	bool Event::IsSet(u32 timeOutMs)
 	{
 		if (m_impl)
-			return m_impl->IsSet(timeOutMs);
+			return ((EventImpl*)m_impl)->IsSet(timeOutMs);
 		return false;
 	}
 
@@ -587,7 +101,7 @@ namespace uba
 		#if PLATFORM_WINDOWS
 		if (!m_impl)
 			return nullptr;
-		return m_impl->handle;
+		return ((EventImpl*)m_impl)->handle;
 		#else
 		UBA_ASSERTF(false, "Event::GetHandle not available");
 		return 0;
@@ -596,11 +110,18 @@ namespace uba
 
 
 #if !PLATFORM_WINDOWS
+
+	#if UBA_USE_NATIVE_MAC_SEMAPHORES
+	using SharedEventImpl = EventMacSemaphore;
+	#else
+	using SharedEventImpl = EventPosixSemaphore;
+	#endif
+
 	SharedEvent::SharedEvent()
 	{
 		*((u8*)m_data) = 0;
 		static_assert(sizeof(m_data) >= sizeof(EventImpl));
-		static_assert(sizeof(m_data) >= sizeof(SemaphoreEvent));
+		static_assert(sizeof(m_data) >= sizeof(SharedEventImpl));
 	}
 
 	SharedEvent::~SharedEvent()
@@ -620,9 +141,9 @@ namespace uba
 	bool SharedEvent::Create(StringView name, bool isOwner)
 	{
 		UBA_ASSERT(*((u8*)m_data) == 0);
-		new (m_data) SemaphoreEvent();
+		new (m_data) SharedEventImpl();
 		*((u8*)m_data) = 2;
-		return ((SemaphoreEvent&)m_data).Create(name, isOwner);
+		return ((SharedEventImpl&)m_data).Create(name, isOwner);
 	}
 
 	void SharedEvent::Destroy()
@@ -632,8 +153,8 @@ namespace uba
 
 		if (IsSemaphore())
 		{
-			((SemaphoreEvent&)m_data).Destroy();
-			((SemaphoreEvent&)m_data).~SemaphoreEvent();
+			((SharedEventImpl&)m_data).Destroy();
+			((SharedEventImpl&)m_data).~SharedEventImpl();
 		}
 		else
 		{
@@ -649,7 +170,7 @@ namespace uba
 			return;
 
 		if (IsSemaphore())
-			((SemaphoreEvent&)m_data).Set();
+			((SharedEventImpl&)m_data).Set();
 		else
 			((EventImpl&)m_data).Set();
 	}
@@ -660,7 +181,7 @@ namespace uba
 			return false;
 
 		if (IsSemaphore())
-			return ((SemaphoreEvent&)m_data).IsSet(timeOutMs);
+			return ((SharedEventImpl&)m_data).IsSet(timeOutMs);
 		else
 			return ((EventImpl&)m_data).IsSet(timeOutMs);
 	}
@@ -673,7 +194,7 @@ namespace uba
 	void SharedEvent::ToString(StringBufferBase& out)
 	{
 		UBA_ASSERT(IsSemaphore());
-		((SemaphoreEvent&)m_data).ToString(out);
+		((SharedEventImpl&)m_data).ToString(out);
 	}
 
 #endif
diff --git a/Engine/Source/Programs/UnrealBuildAccelerator/Core/Private/UbaEventMachSemaphore.inl b/Engine/Source/Programs/UnrealBuildAccelerator/Core/Private/UbaEventMachSemaphore.inl
new file mode 100644
--- /dev/null
+++ b/Engine/Source/Programs/UnrealBuildAccelerator/Core/Private/UbaEventMachSemaphore.inl
@@ -0,0 +1,75 @@
+// Copyright Epic Games, Inc. All Rights Reserved.
+
+#include "UbaDefinitions.h"
+
+#if UBA_USE_NATIVE_MAC_SEMAPHORES
+#include <mach/mach.h>
+#include <mach/semaphore.h>
+
+namespace uba
+{
+	class EventMacSemaphore
+	{
+	public:
+		~EventMacSemaphore()
+		{
+			Destroy();
+		}
+
+		bool Create(StringView name, bool isOwner)
+		{
+			UBA_ASSERT(!m_sem && !m_isOwner);
+			if (isOwner)
+			{
+				kern_return_t kr;
+				kr = semaphore_create(mach_task_self(), &m_sem, SYNC_POLICY_FIFO, 0);
+				UBA_ASSERT(kr == KERN_SUCCESS);
+				mach_port_insert_right(mach_task_self(), m_sem, m_sem, MACH_MSG_TYPE_MAKE_SEND);
+				m_isOwner = true;
+				//printf("CREATE OWNER WITH %u\n", m_sem);
+			}
+			else
+			{
+				u64 sem;
+				Parse(sem, name.data, name.count);
+				m_sem = (semaphore_t)sem;
+				//printf("CREATE CLIENT WITH %u\n", m_sem);
+			}
+			return true;
+		}
+
+		void Destroy()
+		{
+			if (!m_sem)
+				return;
+			if (m_isOwner)
+				semaphore_destroy(mach_task_self(), m_sem);
+		}
+
+		void Set()
+		{
+			semaphore_signal(m_sem);
+		}
+
+		bool IsSet(u32 timeoutMs = ~0u)
+		{
+			if (timeoutMs == ~0u)
+				return semaphore_wait(m_sem) == KERN_SUCCESS;
+			mach_timespec_t ts = ToTimeSpecFromNano<mach_timespec_t>(u64(timeoutMs) * 1'000'000);
+			kern_return_t kr = semaphore_timedwait(m_sem, ts);
+			return kr == KERN_SUCCESS;
+		}
+
+		void ToString(StringBufferBase& out)
+		{
+			out.AppendValue(m_sem);
+		}
+
+	private:
+		u8 m_dummy; // This is to map into SharedEvent under the hood
+		bool m_isOwner = false;
+		semaphore_t m_sem = MACH_PORT_NULL;
+	};
+}
+
+#endif
diff --git a/Engine/Source/Programs/UnrealBuildAccelerator/Core/Private/UbaEventPThread.inl b/Engine/Source/Programs/UnrealBuildAccelerator/Core/Private/UbaEventPThread.inl
new file mode 100644
--- /dev/null
+++ b/Engine/Source/Programs/UnrealBuildAccelerator/Core/Private/UbaEventPThread.inl
@@ -0,0 +1,280 @@
+// Copyright Epic Games, Inc. All Rights Reserved.
+
+#if !PLATFORM_WINDOWS
+
+namespace uba
+{
+	u64 GetMonoticTimeNs();
+
+	template<typename Timespec>
+	inline u64 ToNanoSeconds(const Timespec& ts)
+	{
+		return u64(ts.tv_sec)*1'000'000'000 + ts.tv_nsec;
+	}
+
+	template<typename Timespec>
+	inline Timespec ToTimeSpecFromNano(u64 nanoSeconds)
+	{
+		Timespec ts;
+		ts.tv_sec = nanoSeconds / 1'000'000'000;
+		ts.tv_nsec = nanoSeconds % 1'000'000'000;
+		return ts;
+	}
+
+	struct EventPThread
+	{
+		EventPThread() = default;
+		~EventPThread() { Destroy(); }
+
+		bool Create(bool mr, bool shared = false)
+		{
+			static_assert(Atomic<u32>::is_always_lock_free, "atomic<T> not lock free, can't work in shared mem");
+			static_assert(Atomic<TriggerType>::is_always_lock_free, "atomic<T> not lock free, can't work in shared mem");
+			static_assert(Atomic<bool>::is_always_lock_free, "atomic<T> not lock free, can't work in shared mem");
+
+			UBA_ASSERTF(!m_initialized, "Can't create already created Event");
+			m_manualReset = mr;
+
+			pthread_mutexattr_t attrmutex;
+			if (pthread_mutexattr_init(&attrmutex) != 0)
+			{
+				UBA_ASSERTF(false, "pthread_mutexattr_init failed");
+				return false;
+			}
+
+			if (shared)
+			{
+				if (pthread_mutexattr_setpshared(&attrmutex, PTHREAD_PROCESS_SHARED) != 0)
+				{
+					UBA_ASSERTF(false, "pthread_mutexattr_setpshared failed");
+					return false;
+				}
+				#if PLATFORM_LINUX
+				if (pthread_mutexattr_setrobust(&attrmutex, PTHREAD_MUTEX_ROBUST) != 0)
+				{
+					UBA_ASSERTF(false, "pthread_mutexattr_setrobust failed");
+					return false;
+				}
+				#endif
+			}
+
+			if (pthread_mutex_init(&m_mutex, &attrmutex) != 0)
+			{
+				UBA_ASSERTF(false, "pthread_mutex_init failed");
+				return false;
+			}
+			pthread_mutexattr_destroy(&attrmutex);
+
+			pthread_condattr_t attrcond;
+			if (pthread_condattr_init(&attrcond) != 0)
+			{
+				UBA_ASSERTF(false, "pthread_condattr_init failed");
+				return false;
+			}
+
+			#if PLATFORM_LINUX
+			if (pthread_condattr_setclock(&attrcond, CLOCK_MONOTONIC) != 0)
+			{
+				UBA_ASSERTF(false, "pthread_condattr_setclock failed");
+				return false;
+			}
+			#endif
+
+			if (shared)
+			{
+				if (pthread_condattr_setpshared(&attrcond, PTHREAD_PROCESS_SHARED) != 0)
+				{
+					UBA_ASSERTF(false, "pthread_condattr_setpshared failed");
+					return false;
+				}
+			}
+
+			if (pthread_cond_init(&m_condition, &attrcond) != 0)
+			{
+				UBA_ASSERTF(false, "pthread_cond_init failed");
+				pthread_mutex_destroy(&m_mutex);
+				return false;
+			}
+			pthread_condattr_destroy(&attrcond);
+
+			m_initialized = true;
+			return true;
+		}
+
+		void Destroy()
+		{
+			if (!m_initialized)
+				return;
+
+			LockEventMutex();
+			m_manualReset = true;
+			UnlockEventMutex();
+			Set();
+
+			LockEventMutex();
+			m_initialized = false;
+			while (m_waitingThreads)
+			{
+				UnlockEventMutex();
+				LockEventMutex();
+			}
+
+			pthread_cond_destroy(&m_condition);
+			UnlockEventMutex();
+			pthread_mutex_destroy(&m_mutex);
+		}
+
+		void Set()
+		{
+			if (!m_initialized)
+				return;
+			LockEventMutex();
+
+			if (m_manualReset)
+			{
+				m_triggered = TriggerType_All;
+				if (pthread_cond_broadcast(&m_condition) != 0)
+					UBA_ASSERTF(false, "pthread_cond_broadcast failed");
+			}
+			else
+			{
+				m_triggered = TriggerType_One;
+				if (pthread_cond_signal(&m_condition) != 0)  // may release multiple threads anyhow!
+					UBA_ASSERTF(false, "pthread_cond_signal failed");
+			}
+
+			UnlockEventMutex();
+		}
+
+		void Reset()
+		{
+			if (!m_initialized)
+				return;
+			LockEventMutex();
+			m_triggered = TriggerType_None;
+			UnlockEventMutex();
+		}
+
+		bool IsCreated()
+		{
+			return m_initialized;
+		}
+
+		bool IsSet(u32 timeoutMs = ~0u)
+		{
+			if (!m_initialized)
+				return false;
+
+			u64 startTimeNs = 0;
+
+			// We need to know the start time if we're going to do a timed wait.
+			if ((timeoutMs > 0) && (timeoutMs != ~0u))  // not polling and not infinite wait.
+				startTimeNs = GetMonoticTimeNs();
+
+			u64 timeoutNs = u64(timeoutMs) * 1'000'000;
+
+			#if UBA_TEST_WAIT_QUALITY
+			u32 loop = 0;
+			bool timedOut = false;
+			u64 initialStartTimeNs = startTimeNs;
+			u64 initialTimeoutNs = timeoutNs;
+			auto qg = MakeGuard([&]()
+				{
+					if (startTimeNs && timedOut)
+					{
+						u64 overtimeNs = GetMonoticTimeNs() - initialStartTimeNs - initialTimeoutNs;
+						printf("Loops: %u TimeOut: %ums Over-wait: %lluus\n", loop, timeoutMs, overtimeNs/1000);
+					}
+				});
+			#endif
+
+			LockEventMutex();
+			auto unlock = MakeGuard([this]() { UnlockEventMutex(); });
+
+			// loop in case we fall through the Condition signal but someone else claims the event.
+			do
+			{
+				if (m_triggered == TriggerType_One)
+				{
+					m_triggered = TriggerType_None;
+					return true;
+				}
+
+				if (m_triggered == TriggerType_All)
+				{
+					return true;
+				}
+
+				// No event signalled yet.
+				if (timeoutNs != 0)  // not just polling, wait on the condition variable.
+				{
+					++m_waitingThreads;
+					if (timeoutMs == ~0u) // infinite wait?
+					{
+						if (pthread_cond_wait(&m_condition, &m_mutex) != 0)  // unlocks Mutex while blocking...
+							UBA_ASSERTF(false, "pthread_cond_wait failed");
+					}
+					else  // timed wait.
+					{
+						#if PLATFORM_MAC
+						auto timeout = ToTimeSpecFromNano<timespec>(timeoutNs);
+						int rc = pthread_cond_timedwait_relative_np(&m_condition, &m_mutex, &timeout); // unlocks Mutex while blocking...
+						UBA_ASSERTF((rc == 0) || (rc == ETIMEDOUT), "pthread_cond_timedwait_relative_np failed"); (void)rc;
+						#else
+						auto timeout = ToTimeSpecFromNano<timespec>(startTimeNs + timeoutNs);
+						int rc = pthread_cond_timedwait(&m_condition, &m_mutex, &timeout); // unlocks Mutex while blocking...
+						UBA_ASSERTF((rc == 0) || (rc == ETIMEDOUT), "pthread_cond_timedwait failed"); (void)rc;
+						#endif
+						#if UBA_TEST_WAIT_QUALITY
+						++loop;
+						#endif
+
+						u64 nowNs = GetMonoticTimeNs();
+						u64 diffNs = nowNs - startTimeNs;
+
+						timeoutNs = diffNs >= timeoutNs ? 0 : (timeoutNs - diffNs);
+						startTimeNs = nowNs;
+					}
+					--m_waitingThreads;
+					UBA_ASSERTF(m_waitingThreads >= 0, "m_waitingThreads less than 0");
+				}
+
+			} while (timeoutNs != 0);
+
+			#if UBA_TEST_WAIT_QUALITY
+			timedOut = true;
+			#endif
+			return false;
+		}
+
+		bool GetManualReset() { return m_manualReset; }
+		void ResetManualReset(bool mr) {}
+
+		// Must be first!
+		EventPThread* next;
+
+	private:
+		enum TriggerType : u8 { TriggerType_None, TriggerType_One, TriggerType_All };
+
+		Atomic<bool> m_initialized;
+		Atomic<bool> m_manualReset;
+		Atomic<TriggerType> m_triggered;
+		Atomic<u32> m_waitingThreads;
+		pthread_mutex_t m_mutex;
+		pthread_cond_t m_condition;
+
+		void LockEventMutex()
+		{
+			int res = pthread_mutex_lock(&m_mutex);(void)res;
+			UBA_ASSERTF(res == 0, "pthread_mutex_lock failed (error code %i)", res);
+		}
+
+		void UnlockEventMutex()
+		{
+			int res = pthread_mutex_unlock(&m_mutex);
+			UBA_ASSERTF(res == 0, "pthread_mutex_unlock failed (error code %i)", res);
+		}
+	};
+}
+
+#endif
\ No newline at end of file
diff --git a/Engine/Source/Programs/UnrealBuildAccelerator/Core/Private/UbaEventPosixSemaphore.inl b/Engine/Source/Programs/UnrealBuildAccelerator/Core/Private/UbaEventPosixSemaphore.inl
new file mode 100644
--- /dev/null
+++ b/Engine/Source/Programs/UnrealBuildAccelerator/Core/Private/UbaEventPosixSemaphore.inl
@@ -0,0 +1,159 @@
+// Copyright Epic Games, Inc. All Rights Reserved.
+
+#if !PLATFORM_WINDOWS
+
+#if !PLATFORM_WINDOWS
+#include <sys/time.h>
+#include <new>
+#include <semaphore.h>
+#endif
+
+namespace uba
+{
+	class EventPosixSemaphore
+	{
+	public:
+		~EventPosixSemaphore()
+		{
+			Destroy();
+		}
+
+		bool Create(StringView name, bool isOwner)
+		{
+			UBA_ASSERT(!m_snd && !m_isOwner);
+			StringBuffer<> n1(name);
+			StringBuffer<> n2(name);
+			n1.Append(TCV("_cl"));
+			n2.Append(TCV("_sv"));
+		
+			const tchar* sndName = n1.data;
+			const tchar* rvcName = n2.data;
+			int flags = 0;
+			mode_t mode = 0600;
+
+			if (isOwner)
+			{
+				flags |= O_CREAT;
+				sndName = n2.data;
+				rvcName = n1.data;
+				m_name = TStrdup(name.data);
+				m_isOwner = true;
+			};
+
+			m_snd = sem_open(sndName, flags, mode, 0);
+			m_rcv = sem_open(rvcName, flags, mode, 0);
+
+			UBA_ASSERTF(m_snd != SEM_FAILED, "sem_open failed");
+			UBA_ASSERTF(m_rcv != SEM_FAILED, "sem_open failed");
+			return true;
+		}
+
+		void Destroy()
+		{
+			if (!m_snd)
+				return;
+			sem_close(m_snd);
+			m_snd = nullptr;
+
+			sem_close(m_rcv);
+			m_rcv = nullptr;
+
+			if (m_isOwner)
+			{
+				StringBuffer<> n1(m_name);
+				StringBuffer<> n2(m_name);
+				n1.Append(TCV("_cl"));
+				n2.Append(TCV("_sv"));
+				sem_unlink(n1.data);
+				sem_unlink(n2.data);
+				free(m_name);
+				m_isOwner = false;
+			}
+		}
+
+		void Set()
+		{
+			if (!sem_post(m_snd))
+				return;
+			UBA_ASSERTF(false, "sem_post failed");
+		}
+
+		bool IsSet(u32 timeoutMs = ~0u)
+		{
+			if (timeoutMs == ~0u)
+			{
+				while (sem_wait(m_rcv) == -1)
+				{
+					if (errno == EINTR)
+						continue;
+					UBA_ASSERTF(false, "sem_wait failed");
+				}
+				return errno == 0;
+			}
+
+			u64 timeoutNs = u64(timeoutMs) * 1'000'000;
+
+			struct timespec now;
+			clock_gettime(CLOCK_REALTIME, &now);
+
+			u64 absTimeNano = ToNanoSeconds(now) + timeoutNs;
+
+			#if PLATFORM_LINUX
+			auto absTime = ToTimeSpecFromNano<timespec>(absTimeNano);
+			while (sem_timedwait(m_rcv, &absTime) == -1)
+			{
+				if (errno == EINTR)
+					continue;
+				if (errno == ETIMEDOUT)
+					return false;
+				UBA_ASSERTF(false, "sem_wait failed");
+			}
+			#else
+
+			u64 sleepNs = 50 * 1000; // 50 µs
+			const u64 sleepCapNs = 10 * 1000 * 1000; // 10 ms
+
+			while (true)
+			{
+				if (sem_trywait(m_rcv) == 0)
+					return true;
+				if (errno != EAGAIN)
+					return false;
+
+				clock_gettime(CLOCK_REALTIME, &now);
+				u64 nowNano = ToNanoSeconds(now);
+				if (nowNano > absTimeNano)
+					return false;
+
+				u64 remaining = absTimeNano - nowNano;
+				u64 nap = sleepNs < remaining ? sleepNs : remaining;
+				auto deltaNs = ToTimeSpecFromNano<timespec>(nap);
+				nanosleep(&deltaNs, nullptr);
+
+				if (sleepNs < sleepCapNs)
+				{
+					sleepNs <<= 1;
+					if (sleepNs > sleepCapNs)
+						sleepNs = sleepCapNs;
+				}
+			}
+			#endif
+
+			return true;
+		}
+
+		void ToString(StringBufferBase& out)
+		{
+			out.Append(m_name);
+		}
+
+	private:
+		u8 m_dummy; // This is to map into SharedEvent under the hood
+		bool m_isOwner = false;
+		sem_t* m_snd = nullptr;
+		sem_t* m_rcv = nullptr;
+		tchar* m_name = nullptr;
+	};
+}
+
+#endif
\ No newline at end of file
diff --git a/Engine/Source/Programs/UnrealBuildAccelerator/Core/Private/UbaEventWin32.inl b/Engine/Source/Programs/UnrealBuildAccelerator/Core/Private/UbaEventWin32.inl
new file mode 100644
--- /dev/null
+++ b/Engine/Source/Programs/UnrealBuildAccelerator/Core/Private/UbaEventWin32.inl
@@ -0,0 +1,26 @@
+// Copyright Epic Games, Inc. All Rights Reserved.
+
+#if PLATFORM_WINDOWS
+
+namespace uba
+{
+	struct EventWin32
+	{
+		HANDLE handle;
+		union
+		{
+			EventWin32* next;
+			bool manualReset;
+		};
+
+		bool Create(bool mr) { handle = CreateEvent(nullptr, mr, false, NULL); return handle != 0; }
+		void Set() { SetEvent(handle); }
+		void Reset() { ResetEvent(handle); }
+		bool IsSet(u32 timeOutMs) { return WaitForSingleObject(handle, timeOutMs) == 0;} 
+		bool GetManualReset() { return manualReset; }
+		void ResetManualReset(bool mr) { manualReset = mr; }
+	};
+
+}
+
+#endif
\ No newline at end of file
diff --git a/Engine/Source/Programs/UnrealBuildAccelerator/Core/Public/UbaEvent.h b/Engine/Source/Programs/UnrealBuildAccelerator/Core/Public/UbaEvent.h
--- a/Engine/Source/Programs/UnrealBuildAccelerator/Core/Public/UbaEvent.h
+++ b/Engine/Source/Programs/UnrealBuildAccelerator/Core/Public/UbaEvent.h
@@ -23,7 +23,7 @@ namespace uba
 		bool IsSet(u32 timeOutMs = ~0u);
 		void* GetHandle();
 	private:
-		struct EventImpl* m_impl = nullptr;
+		void* m_impl = nullptr;
 		Event(const Event&) = delete;
 		void operator=(const Event&) = delete;
 	};
