[UBA]
* Fixed compile errors for mac in detours code

diff --git a/Engine/Source/Programs/UnrealBuildAccelerator/Detours/Private/PosixOS/UbaDetoursFunctionsPosix.cpp b/Engine/Source/Programs/UnrealBuildAccelerator/Detours/Private/PosixOS/UbaDetoursFunctionsPosix.cpp
--- a/Engine/Source/Programs/UnrealBuildAccelerator/Detours/Private/PosixOS/UbaDetoursFunctionsPosix.cpp
+++ b/Engine/Source/Programs/UnrealBuildAccelerator/Detours/Private/PosixOS/UbaDetoursFunctionsPosix.cpp
@@ -4,7 +4,7 @@
 
 #include "UbaDetoursFileMappingTable.h"
 #include "UbaBinaryReaderWriter.h"
-#include "UbaDefinitions.h""
+#include "UbaDefinitions.h"
 #include "UbaDirectoryTable.h"
 #include "UbaPlatform.h"
 #include "UbaProcessStats.h"
diff --git a/Engine/Source/Programs/UnrealBuildAccelerator/Detours/Private/PosixOS/UbaDetoursMainPosix.cpp b/Engine/Source/Programs/UnrealBuildAccelerator/Detours/Private/PosixOS/UbaDetoursMainPosix.cpp
--- a/Engine/Source/Programs/UnrealBuildAccelerator/Detours/Private/PosixOS/UbaDetoursMainPosix.cpp
+++ b/Engine/Source/Programs/UnrealBuildAccelerator/Detours/Private/PosixOS/UbaDetoursMainPosix.cpp
@@ -1,9 +1,9 @@
 // Copyright Epic Games, Inc. All Rights Reserved.
 
 #include "UbaBinaryReaderWriter.h"
+#include "UbaDefinitions.h"
 #include "UbaDetoursShared.h"
 #include "UbaEvent.h"
-#include "UbaProtocol.h"
 #include "UbaProcessStats.h"
 #include <dlfcn.h>
 #if PLATFORM_LINUX
