[UBA]
* Changed bool manualReset in Event to be an enum instead to make code easier to read

diff --git a/Engine/Source/Programs/UnrealBuildAccelerator/Agent/Private/UbaAgent.cpp b/Engine/Source/Programs/UnrealBuildAccelerator/Agent/Private/UbaAgent.cpp
--- a/Engine/Source/Programs/UnrealBuildAccelerator/Agent/Private/UbaAgent.cpp
+++ b/Engine/Source/Programs/UnrealBuildAccelerator/Agent/Private/UbaAgent.cpp
@@ -998,7 +998,7 @@ namespace uba
 		Vector<ProcessLogLine> logLines[2];
 		u32 logLinesIndex = 0;
 		Futex logLinesLock;
-		Event logLinesAvailable(false);
+		Event logLinesAvailable(EventResetType_Auto);
 
 		auto processFinished = [&](const ProcessHandle& process)
 		{
@@ -1202,7 +1202,7 @@ namespace uba
 			if (!client->FetchConfig(config))
 				continue;
 
-			Event wakeupSessionWait(false);
+			Event wakeupSessionWait(EventResetType_Auto);
 			Atomic<u32> targetConnectionCount = minConnectionCount;
 			StorageClient* storageClient = nullptr;
 
diff --git a/Engine/Source/Programs/UnrealBuildAccelerator/CacheService/Private/UbaCacheService.cpp b/Engine/Source/Programs/UnrealBuildAccelerator/CacheService/Private/UbaCacheService.cpp
--- a/Engine/Source/Programs/UnrealBuildAccelerator/CacheService/Private/UbaCacheService.cpp
+++ b/Engine/Source/Programs/UnrealBuildAccelerator/CacheService/Private/UbaCacheService.cpp
@@ -527,7 +527,7 @@ namespace uba
 			trace.SessionAdded(0, 0, TCV("CacheServer"), {});
 			trace.ProcessAdded(0, 0, TCV("Running"), {});
 			networkServer.SetWorkTracker(&trace);
-			traceLoop.Create(true);
+			traceLoop.Create(EventResetType_Manual);
 			traceThread.Start([&]()
 				{
 					while (!traceLoop.IsSet(500))
diff --git a/Engine/Source/Programs/UnrealBuildAccelerator/Cli/Private/UbaCli.cpp b/Engine/Source/Programs/UnrealBuildAccelerator/Cli/Private/UbaCli.cpp
--- a/Engine/Source/Programs/UnrealBuildAccelerator/Cli/Private/UbaCli.cpp
+++ b/Engine/Source/Programs/UnrealBuildAccelerator/Cli/Private/UbaCli.cpp
@@ -1391,7 +1391,7 @@ namespace uba
 
 			bool success = true;
 			Atomic<u32> counter;
-			static Event finished(true);
+			static Event finished(EventResetType_Manual);
 
 			scheduler.SetProcessFinishedCallback([&](const ProcessHandle& ph)
 				{
diff --git a/Engine/Source/Programs/UnrealBuildAccelerator/Common/Private/UbaCacheClient.cpp b/Engine/Source/Programs/UnrealBuildAccelerator/Common/Private/UbaCacheClient.cpp
--- a/Engine/Source/Programs/UnrealBuildAccelerator/Common/Private/UbaCacheClient.cpp
+++ b/Engine/Source/Programs/UnrealBuildAccelerator/Common/Private/UbaCacheClient.cpp
@@ -1522,7 +1522,7 @@ namespace uba
 
 		struct Entry
 		{
-			Entry(u8* slot, u32 i, u32 messageMaxSize) : reader(slot + i * messageMaxSize, 0, SendMaxSize), done(true) {}
+			Entry(u8* slot, u32 i, u32 messageMaxSize) : reader(slot + i * messageMaxSize, 0, SendMaxSize), done(EventResetType_Manual) {}
 			NetworkMessage message;
 			BinaryReader reader;
 			Event done;
diff --git a/Engine/Source/Programs/UnrealBuildAccelerator/Common/Private/UbaCacheServer.cpp b/Engine/Source/Programs/UnrealBuildAccelerator/Common/Private/UbaCacheServer.cpp
--- a/Engine/Source/Programs/UnrealBuildAccelerator/Common/Private/UbaCacheServer.cpp
+++ b/Engine/Source/Programs/UnrealBuildAccelerator/Common/Private/UbaCacheServer.cpp
@@ -1204,8 +1204,8 @@ namespace uba
 
 		bool shouldSave = allowSave && (entriesAdded || deletedCasCount || deleteEntryCount || forceAllSteps);
 
-		Event saveCasEvent(true);
-		Event saveDbEvent(true);
+		Event saveCasEvent(EventResetType_Manual);
+		Event saveDbEvent(EventResetType_Manual);
 		if (shouldSave)
 		{
 			m_server.AddWork([&](const WorkContext&) { m_storage.SaveCasTable(false, false); saveCasEvent.Set(); }, 1, TC("SaveCas"));
diff --git a/Engine/Source/Programs/UnrealBuildAccelerator/Common/Private/UbaFileAccessor.cpp b/Engine/Source/Programs/UnrealBuildAccelerator/Common/Private/UbaFileAccessor.cpp
--- a/Engine/Source/Programs/UnrealBuildAccelerator/Common/Private/UbaFileAccessor.cpp
+++ b/Engine/Source/Programs/UnrealBuildAccelerator/Common/Private/UbaFileAccessor.cpp
@@ -298,7 +298,7 @@ namespace uba
 				u64 index = i % BlockCount;
 
 				if (i < BlockCount)
-					ev[i].Create(true);
+					ev[i].Create(EventResetType_Manual);
 				else
 				{
 					if (!WaitAndCheckError(index))
diff --git a/Engine/Source/Programs/UnrealBuildAccelerator/Common/Private/UbaNetworkBackendTcp.cpp b/Engine/Source/Programs/UnrealBuildAccelerator/Common/Private/UbaNetworkBackendTcp.cpp
--- a/Engine/Source/Programs/UnrealBuildAccelerator/Common/Private/UbaNetworkBackendTcp.cpp
+++ b/Engine/Source/Programs/UnrealBuildAccelerator/Common/Private/UbaNetworkBackendTcp.cpp
@@ -97,7 +97,7 @@ namespace uba
 
 	struct NetworkBackendTcp::Connection
 	{
-		Connection(Logger& l, SOCKET s) : logger(l), socket(s), ready(true) { CreateGuid(uid); }
+		Connection(Logger& l, SOCKET s) : logger(l), socket(s), ready(EventResetType_Manual) { CreateGuid(uid); }
 
 		Logger& logger;
 		Atomic<SOCKET> socket;
@@ -264,7 +264,7 @@ namespace uba
 
 		if (info.statusUpdateSeconds)
 		{
-			m_tcpStatusLoop.Create(true);
+			m_tcpStatusLoop.Create(EventResetType_Manual);
 			m_tcpStatusThread.Start([this, sus = info.statusUpdateSeconds]() { ThreadStatus(sus); return 0; }, TC("UbaTcpStat"));
 		}
 	}
@@ -429,7 +429,7 @@ namespace uba
 		{
 			if (skipCount-- > 0)
 				continue;
-			e.listening.Create(true);
+			e.listening.Create(EventResetType_Manual);
 			e.thread.Start([this, &logger, &e]
 				{
 					ThreadListen(logger, e);
@@ -1044,7 +1044,7 @@ namespace uba
 #if UBA_USE_OVERLAPPED_SEND
 		if (m_useOverlappedSend)
 		{
-			Event ev(true);
+			Event ev(EventResetType_Manual);
 
 			// Right now we are experimenting with if we can ignore taking a lock around the entire thing and only the WSASend.
 			// Documentation is slightly unclear and network forums claim the order of data is correct even though multiple
diff --git a/Engine/Source/Programs/UnrealBuildAccelerator/Common/Private/UbaNetworkClient.cpp b/Engine/Source/Programs/UnrealBuildAccelerator/Common/Private/UbaNetworkClient.cpp
--- a/Engine/Source/Programs/UnrealBuildAccelerator/Common/Private/UbaNetworkClient.cpp
+++ b/Engine/Source/Programs/UnrealBuildAccelerator/Common/Private/UbaNetworkClient.cpp
@@ -24,7 +24,7 @@ namespace uba
 
 	struct NetworkClient::Connection
 	{
-		Connection(NetworkClient& o) : owner(o), disconnectedEvent(true) {}
+		Connection(NetworkClient& o) : owner(o), disconnectedEvent(EventResetType_Manual) {}
 		NetworkClient& owner;
 		void* backendConnection = nullptr;
 		Atomic<u32> connected;
@@ -43,8 +43,8 @@ namespace uba
 	: WorkManagerImpl(info.workerCount == ~0u ? GetLogicalProcessorCount() : info.workerCount, TC("UbaWrk/NetClnt"))
 	,	m_logWriter(info.logWriter)
 	,	m_logger(info.logWriter, SetGetPrefix(name))
-	,	m_isConnected(true)
-	,	m_isOrWasConnected(true)
+	,	m_isConnected(EventResetType_Manual)
+	,	m_isOrWasConnected(EventResetType_Manual)
 	{
 		outCtorSuccess = true;
 
@@ -92,7 +92,7 @@ namespace uba
 	{
 		struct RecvContext
 		{
-			RecvContext(NetworkClient& c, NetworkBackend& b, void* bc) : client(c), backend(b), backendConnection(bc), recvEvent(true), exitScopeEvent(true)
+			RecvContext(NetworkClient& c, NetworkBackend& b, void* bc) : client(c), backend(b), backendConnection(bc), recvEvent(EventResetType_Manual), exitScopeEvent(EventResetType_Manual)
 			{
 				error = 255;
 			}
@@ -736,7 +736,7 @@ namespace uba
 		{
 			if (!async)
 			{
-				if (!gotResponse.Create(true))
+				if (!gotResponse.Create(EventResetType_Manual))
 				{
 					m_logger.Error(TC("Failed to create event, this should not happen?!?"));
 					message.m_error = 13;
diff --git a/Engine/Source/Programs/UnrealBuildAccelerator/Common/Private/UbaNetworkServer.cpp b/Engine/Source/Programs/UnrealBuildAccelerator/Common/Private/UbaNetworkServer.cpp
--- a/Engine/Source/Programs/UnrealBuildAccelerator/Common/Private/UbaNetworkServer.cpp
+++ b/Engine/Source/Programs/UnrealBuildAccelerator/Common/Private/UbaNetworkServer.cpp
@@ -41,7 +41,7 @@ namespace uba
 
 	struct NetworkServer::WorkerContext
 	{
-		WorkerContext(NetworkServer& s) : server(s), workAvailable(false)
+		WorkerContext(NetworkServer& s) : server(s), workAvailable(EventResetType_Auto)
 		{
 			writeMemSize = server.m_sendSize;
 			writeMem = new u8[writeMemSize];
@@ -126,7 +126,7 @@ namespace uba
 		,	m_backend(backend)
 		,	m_remoteSockAddr(remoteSockAddr)
 		,	m_cryptoKey(cryptoKey)
-		,	m_disconnectCallbackCalled(true)
+		,	m_disconnectCallbackCalled(EventResetType_Manual)
 		,	m_id(id)
 		,	m_backendConnection(backendConnection)
 		{
@@ -929,7 +929,7 @@ namespace uba
 				return false;
 		}
 
-		Event done(true);
+		Event done(EventResetType_Manual);
 		bool success = false;
 
 		m_addConnections.emplace_back([this, &success , &done, &backend, ip2 = TString(ip), port, cryptoKey]()
@@ -1310,7 +1310,7 @@ namespace uba
 				return worker;
 			
 			PopWorkerRequest req;
-			req.ev.Create(true);
+			req.ev.Create(EventResetType_Manual);
 
 			if (!m_firstRequest)
 				m_firstRequest = &req;
diff --git a/Engine/Source/Programs/UnrealBuildAccelerator/Common/Private/UbaProcess.cpp b/Engine/Source/Programs/UnrealBuildAccelerator/Common/Private/UbaProcess.cpp
--- a/Engine/Source/Programs/UnrealBuildAccelerator/Common/Private/UbaProcess.cpp
+++ b/Engine/Source/Programs/UnrealBuildAccelerator/Common/Private/UbaProcess.cpp
@@ -178,9 +178,9 @@ namespace uba
 		}
 		#endif
 
-		m_cancelEvent.Create(true);
-		m_writeEvent.Create(false);
-		m_readEvent.Create(false);
+		m_cancelEvent.Create(EventResetType_Manual);
+		m_writeEvent.Create(EventResetType_Auto);
+		m_readEvent.Create(EventResetType_Auto);
 	}
 
 	ProcessImpl::~ProcessImpl()
@@ -253,7 +253,7 @@ namespace uba
 			m_realWorkingDir = m_startInfo.workingDir;
 
 		if (m_parentProcess)
-			m_waitForParent.Create(true);
+			m_waitForParent.Create(EventResetType_Manual);
 
 		UBA_ASSERT(m_startInfo.rules);
 
@@ -272,7 +272,7 @@ namespace uba
 
 		if (async)
 		{
-			m_messageThreadDone.Create(true);
+			m_messageThreadDone.Create(EventResetType_Manual);
 			m_session.m_processMessageThreads.AddWork([this, environment](const WorkContext& context)
 				{
 					m_messageThread = context.thread;
diff --git a/Engine/Source/Programs/UnrealBuildAccelerator/Common/Private/UbaScheduler.cpp b/Engine/Source/Programs/UnrealBuildAccelerator/Common/Private/UbaScheduler.cpp
--- a/Engine/Source/Programs/UnrealBuildAccelerator/Common/Private/UbaScheduler.cpp
+++ b/Engine/Source/Programs/UnrealBuildAccelerator/Common/Private/UbaScheduler.cpp
@@ -177,7 +177,7 @@ namespace uba
 	:	m_session(info.session)
 	,	m_maxLocalProcessors(info.maxLocalProcessors != ~0u ? info.maxLocalProcessors : GetLogicalProcessorCount())
 	,	m_maxParallelCacheQueries(info.maxParallelCacheQueries)
-	,	m_updateThreadWakeup(false)
+	,	m_updateThreadWakeup(EventResetType_Auto)
 	,	m_enableProcessReuse(info.enableProcessReuse)
 	,	m_forceRemote(info.forceRemote)
 	,	m_forceNative(info.forceNative)
@@ -222,7 +222,7 @@ namespace uba
 		{
 			LoadMemTrackTable();
 			m_session.SetNativeProcessCreatedFunc([this](ProcessImpl& process) { NativeProcessCreated(process); });
-			m_memThreadEvent.Create(true);
+			m_memThreadEvent.Create(EventResetType_Manual);
 			m_memThread.Start([this]() { ThreadMemoryCheckLoop();  return 0; }, TC("UbaMemTrackLoop"));
 
 			// Make sure we have an estimate before starting to spawn processes
@@ -1954,7 +1954,7 @@ namespace uba
 		m_activeLocalProcessWeight -= weight;
 		lock.Leave();
 
-		Event done(true);
+		Event done(EventResetType_Manual);
 		struct Context
 		{
 			Logger& logger;
diff --git a/Engine/Source/Programs/UnrealBuildAccelerator/Common/Private/UbaSessionClient.cpp b/Engine/Source/Programs/UnrealBuildAccelerator/Common/Private/UbaSessionClient.cpp
--- a/Engine/Source/Programs/UnrealBuildAccelerator/Common/Private/UbaSessionClient.cpp
+++ b/Engine/Source/Programs/UnrealBuildAccelerator/Common/Private/UbaSessionClient.cpp
@@ -59,7 +59,7 @@ namespace uba
 
 	struct SessionClient::ModuleInfo
 	{
-		ModuleInfo(const tchar* n, const CasKey& c, u32 a) : name(n), casKey(c), attributes(a), done(true) {}
+		ModuleInfo(const tchar* n, const CasKey& c, u32 a) : name(n), casKey(c), attributes(a), done(EventResetType_Manual) {}
 		TString name;
 		CasKey casKey;
 		u32 attributes;
@@ -71,7 +71,7 @@ namespace uba
 	,	m_client(info.client)
 	,	m_name(info.name.data)
 	,	m_terminationTime(~0ull)
-	,	m_waitToSendEvent(false)
+	,	m_waitToSendEvent(EventResetType_Auto)
 	,	m_loop(true)
 	,	m_allowSpawn(true)
 	{
@@ -1113,7 +1113,7 @@ namespace uba
 				item.next->prev = &item;
 			item.readPos = readPos;
 			first = &item;
-			item.done.Create(true);
+			item.done.Create(EventResetType_Manual);
 
 			lock.Leave();
 			bool res = item.done.IsSet(5*60*1000);
@@ -1681,7 +1681,7 @@ namespace uba
 		{
 			Event done;
 			u64 time;
-		} response{true, 0};
+		} response{EventResetType_Manual, 0};
 
 		auto doneFunc = [](bool error, void* userData)
 			{
diff --git a/Engine/Source/Programs/UnrealBuildAccelerator/Common/Private/UbaSessionServer.cpp b/Engine/Source/Programs/UnrealBuildAccelerator/Common/Private/UbaSessionServer.cpp
--- a/Engine/Source/Programs/UnrealBuildAccelerator/Common/Private/UbaSessionServer.cpp
+++ b/Engine/Source/Programs/UnrealBuildAccelerator/Common/Private/UbaSessionServer.cpp
@@ -19,7 +19,7 @@ namespace uba
 		:	m_server(server)
 		,	m_startInfo(si)
 		,	m_processId(processId)
-		,	m_done(true)
+		,	m_done(EventResetType_Manual)
 		{
 			m_startInfo.weight = weight_;
 		}
@@ -321,7 +321,7 @@ namespace uba
 
 		if (knownInputsCount)
 		{
-			remoteProcess->m_knownInputsDone.Create(true);
+			remoteProcess->m_knownInputsDone.Create(EventResetType_Manual);
 
 
 			auto kiBegin = (const tchar*)knownInputs;
diff --git a/Engine/Source/Programs/UnrealBuildAccelerator/Common/Private/UbaStorageClient.cpp b/Engine/Source/Programs/UnrealBuildAccelerator/Common/Private/UbaStorageClient.cpp
--- a/Engine/Source/Programs/UnrealBuildAccelerator/Common/Private/UbaStorageClient.cpp
+++ b/Engine/Source/Programs/UnrealBuildAccelerator/Common/Private/UbaStorageClient.cpp
@@ -975,7 +975,7 @@ namespace uba
 				}
 			}
 			else
-				localFile.hasBeenSent.Create(true);
+				localFile.hasBeenSent.Create(EventResetType_Manual);
 		}
 		lock.Leave();
 
diff --git a/Engine/Source/Programs/UnrealBuildAccelerator/Common/Private/UbaStorageUtils.cpp b/Engine/Source/Programs/UnrealBuildAccelerator/Common/Private/UbaStorageUtils.cpp
--- a/Engine/Source/Programs/UnrealBuildAccelerator/Common/Private/UbaStorageUtils.cpp
+++ b/Engine/Source/Programs/UnrealBuildAccelerator/Common/Private/UbaStorageUtils.cpp
@@ -82,7 +82,7 @@ namespace uba
 			rec->workCount = workCount;
 			rec->fileSize = fileSize;
 			rec->keys.resize(workCount);
-			rec->done.Create(true);
+			rec->done.Create(EventResetType_Manual);
 			rec->refCount = 2;
 
 			auto work = [rec](const WorkContext& context)
@@ -154,7 +154,7 @@ namespace uba
 
 		struct Entry
 		{
-			Entry(u8* slot, u32 i, u32 messageMaxSize) : reader(slot + i * messageMaxSize, 0, SendMaxSize), done(true) {}
+			Entry(u8* slot, u32 i, u32 messageMaxSize) : reader(slot + i * messageMaxSize, 0, SendMaxSize), done(EventResetType_Manual) {}
 			NetworkMessage message;
 			BinaryReader reader;
 			Event done;
@@ -1101,7 +1101,7 @@ namespace uba
 		rec.casCompressionLevel = m_casCompressionLevel;
 
 		for (u32 i=0; i!=workCount; ++i)
-			rec.events[i].Create(true);
+			rec.events[i].Create(EventResetType_Manual);
 
 		TimerScope cts(m_stats.compressWrite);
 
@@ -1351,7 +1351,7 @@ namespace uba
 			rec->writePos = writeData;
 			rec->decompressedSize = decompressedSize;
 			rec->decompressedLeft = decompressedSize;
-			rec->done.Create(true);
+			rec->done.Create(EventResetType_Manual);
 			rec->refCount = 2;
 
 			auto work = [rec](const WorkContext&)
diff --git a/Engine/Source/Programs/UnrealBuildAccelerator/Common/Private/UbaThread.cpp b/Engine/Source/Programs/UnrealBuildAccelerator/Common/Private/UbaThread.cpp
--- a/Engine/Source/Programs/UnrealBuildAccelerator/Common/Private/UbaThread.cpp
+++ b/Engine/Source/Programs/UnrealBuildAccelerator/Common/Private/UbaThread.cpp
@@ -87,7 +87,7 @@ namespace uba
 #else
 		int err = 0;
 
-		m_finished.Create(true);
+		m_finished.Create(EventResetType_Manual);
 		static_assert(sizeof(pthread_t) <= sizeof(m_handle), "");
 		auto& pth = *(pthread_t*)&m_handle;
 
@@ -315,7 +315,7 @@ namespace uba
 		// TODO: None of these approaches work.
 		// signal path will break if thread is in certain system calls.. and ptrace does not work either :-/
 
-		static Event s_ev(false);
+		static Event s_ev(EventResetType_Auto);
 		static const TraverseThreadFunc* s_func;
 		s_func = &func;
 
diff --git a/Engine/Source/Programs/UnrealBuildAccelerator/Common/Private/UbaTrace.cpp b/Engine/Source/Programs/UnrealBuildAccelerator/Common/Private/UbaTrace.cpp
--- a/Engine/Source/Programs/UnrealBuildAccelerator/Common/Private/UbaTrace.cpp
+++ b/Engine/Source/Programs/UnrealBuildAccelerator/Common/Private/UbaTrace.cpp
@@ -86,7 +86,7 @@ namespace uba
 
 	void Trace::StartThread()
 	{
-		m_threadEvent.Create(true);
+		m_threadEvent.Create(EventResetType_Manual);
 		m_thread.Start([this]() { ThreadLoop(); return 0; }, TC("UbaTraceLoop"));
 	}
 
diff --git a/Engine/Source/Programs/UnrealBuildAccelerator/Common/Private/UbaWorkManager.cpp b/Engine/Source/Programs/UnrealBuildAccelerator/Common/Private/UbaWorkManager.cpp
--- a/Engine/Source/Programs/UnrealBuildAccelerator/Common/Private/UbaWorkManager.cpp
+++ b/Engine/Source/Programs/UnrealBuildAccelerator/Common/Private/UbaWorkManager.cpp
@@ -186,7 +186,7 @@ namespace uba
 
 	struct WorkManagerImpl::Worker
 	{
-		Worker(WorkManagerImpl& manager, const tchar* workerDesc) : m_workAvailable(false)
+		Worker(WorkManagerImpl& manager, const tchar* workerDesc) : m_workAvailable(EventResetType_Auto)
 		{
 			m_loop = true;
 			m_thread.Start([&]() { ThreadWorker(manager); return 0; }, workerDesc);
diff --git a/Engine/Source/Programs/UnrealBuildAccelerator/Common/Public/UbaBottleneck.h b/Engine/Source/Programs/UnrealBuildAccelerator/Common/Public/UbaBottleneck.h
--- a/Engine/Source/Programs/UnrealBuildAccelerator/Common/Public/UbaBottleneck.h
+++ b/Engine/Source/Programs/UnrealBuildAccelerator/Common/Public/UbaBottleneck.h
@@ -35,7 +35,7 @@ namespace uba
 			m_last = &ticket;
 
 			if (!ticket.ev.IsCreated())
-				ticket.ev.Create(false);
+				ticket.ev.Create(EventResetType_Auto);
 
 			lock.Leave();
 
diff --git a/Engine/Source/Programs/UnrealBuildAccelerator/Common/Public/UbaCoordinatorWrapper.h b/Engine/Source/Programs/UnrealBuildAccelerator/Common/Public/UbaCoordinatorWrapper.h
--- a/Engine/Source/Programs/UnrealBuildAccelerator/Common/Public/UbaCoordinatorWrapper.h
+++ b/Engine/Source/Programs/UnrealBuildAccelerator/Common/Public/UbaCoordinatorWrapper.h
@@ -56,7 +56,7 @@ namespace uba
 			if (!m_coordinator)
 				return false;
 
-			m_loopCoordinator.Create(true);
+			m_loopCoordinator.Create(EventResetType_Manual);
 			m_networkBackend = &networkBackend;
 			m_networkServer = &networkServer;
 			m_scheduler = scheduler;
diff --git a/Engine/Source/Programs/UnrealBuildAccelerator/Common/Public/UbaFileMapping.h b/Engine/Source/Programs/UnrealBuildAccelerator/Common/Public/UbaFileMapping.h
--- a/Engine/Source/Programs/UnrealBuildAccelerator/Common/Public/UbaFileMapping.h
+++ b/Engine/Source/Programs/UnrealBuildAccelerator/Common/Public/UbaFileMapping.h
@@ -114,7 +114,7 @@ namespace uba
 			Event availableFilesEvent;
 			File* availableFile = nullptr;
 
-			MappingStorage() : availableFilesEvent(false) {}
+			MappingStorage() : availableFilesEvent(EventResetType_Auto) {}
 		};
 
 		File* PopFile(MappingStorage& storage, u64 size, u64 alignment);
diff --git a/Engine/Source/Programs/UnrealBuildAccelerator/Common/Public/UbaStorageServer.h b/Engine/Source/Programs/UnrealBuildAccelerator/Common/Public/UbaStorageServer.h
--- a/Engine/Source/Programs/UnrealBuildAccelerator/Common/Public/UbaStorageServer.h
+++ b/Engine/Source/Programs/UnrealBuildAccelerator/Common/Public/UbaStorageServer.h
@@ -72,7 +72,7 @@ namespace uba
 
 		struct WaitEntry
 		{
-			WaitEntry() : done(true) {}
+			WaitEntry() : done(EventResetType_Manual) {}
 			Event done;
 			bool success = false;
 			u32 refCount = 0;
diff --git a/Engine/Source/Programs/UnrealBuildAccelerator/Common/Public/UbaWorkManager.h b/Engine/Source/Programs/UnrealBuildAccelerator/Common/Public/UbaWorkManager.h
--- a/Engine/Source/Programs/UnrealBuildAccelerator/Common/Public/UbaWorkManager.h
+++ b/Engine/Source/Programs/UnrealBuildAccelerator/Common/Public/UbaWorkManager.h
@@ -155,7 +155,7 @@ namespace uba
 
 		workCount = Min(workCount, u32(size - 1)/BatchSize);
 
-		Event doneEvent(true);
+		Event doneEvent(EventResetType_Manual);
 
 		struct Context
 		{
diff --git a/Engine/Source/Programs/UnrealBuildAccelerator/Core/Private/UbaEvent.cpp b/Engine/Source/Programs/UnrealBuildAccelerator/Core/Private/UbaEvent.cpp
--- a/Engine/Source/Programs/UnrealBuildAccelerator/Core/Private/UbaEvent.cpp
+++ b/Engine/Source/Programs/UnrealBuildAccelerator/Core/Private/UbaEvent.cpp
@@ -24,9 +24,9 @@ namespace uba
 	Futex g_firstEventLock[2];
 	EventImpl* g_firstEvent[2];
 
-	Event::Event(bool manualReset)
+	Event::Event(EventResetType resetType)
 	{
-		Create(manualReset);
+		Create(resetType);
 	}
 
 	Event::~Event()
@@ -34,11 +34,11 @@ namespace uba
 		Destroy();
 	}
 
-	bool Event::Create(bool manualReset)
+	bool Event::Create(EventResetType resetType)
 	{
 		EventImpl* impl = nullptr;
-		SCOPED_FUTEX(g_firstEventLock[manualReset], lock);
-		auto*& first = g_firstEvent[manualReset];
+		SCOPED_FUTEX(g_firstEventLock[resetType], lock);
+		auto*& first = g_firstEvent[resetType];
 		if (first)
 		{
 			impl = first;
@@ -49,12 +49,16 @@ namespace uba
 		if (!impl)
 		{
 			impl = new EventImpl();
-			impl->Create(manualReset);
+			impl->Create(resetType);
 		}
 		else
+		{
 			impl->Reset();
+			#if PLATFORM_WINDOWS
+			impl->resetType = resetType; // union with next
+			#endif
+		}
 
-		impl->ResetManualReset(manualReset);
 		m_impl = impl;
 		return true;
 	}
@@ -65,9 +69,9 @@ namespace uba
 			return;
 		auto impl = (EventImpl*)m_impl;
 		m_impl = nullptr;
-		bool manualReset = impl->GetManualReset();
-		auto*& first = g_firstEvent[manualReset];
-		SCOPED_FUTEX(g_firstEventLock[manualReset], l);
+		auto resetType = impl->GetResetType();
+		auto*& first = g_firstEvent[resetType];
+		SCOPED_FUTEX(g_firstEventLock[resetType], l);
 		impl->next = first;
 		first = impl;
 	}
@@ -129,13 +133,13 @@ namespace uba
 		Destroy();
 	}
 
-	bool SharedEvent::Create(bool manualReset)//, bool shared)
+	bool SharedEvent::Create(EventResetType resetType)//, bool shared)
 	{
 		UBA_ASSERT(*((u8*)m_data) == 0);
 		new (m_data) EventImpl();
 		UBA_ASSERT(*((u8*)m_data) == 0);
 		*((u8*)m_data) = 1;
-		return ((EventImpl&)m_data).Create(manualReset, true);
+		return ((EventImpl&)m_data).Create(resetType, true);
 	}
 
 	bool SharedEvent::Create(StringView name, bool isOwner)
diff --git a/Engine/Source/Programs/UnrealBuildAccelerator/Core/Private/UbaEventPThread.inl b/Engine/Source/Programs/UnrealBuildAccelerator/Core/Private/UbaEventPThread.inl
--- a/Engine/Source/Programs/UnrealBuildAccelerator/Core/Private/UbaEventPThread.inl
+++ b/Engine/Source/Programs/UnrealBuildAccelerator/Core/Private/UbaEventPThread.inl
@@ -26,14 +26,14 @@ namespace uba
 		EventPThread() = default;
 		~EventPThread() { Destroy(); }
 
-		bool Create(bool mr, bool shared = false)
+		bool Create(EventResetType rt, bool shared = false)
 		{
 			static_assert(Atomic<u32>::is_always_lock_free, "atomic<T> not lock free, can't work in shared mem");
 			static_assert(Atomic<TriggerType>::is_always_lock_free, "atomic<T> not lock free, can't work in shared mem");
 			static_assert(Atomic<bool>::is_always_lock_free, "atomic<T> not lock free, can't work in shared mem");
 
 			UBA_ASSERTF(!m_initialized, "Can't create already created Event");
-			m_manualReset = mr;
+			m_resetType = rt;
 
 			pthread_mutexattr_t attrmutex;
 			if (pthread_mutexattr_init(&attrmutex) != 0)
@@ -107,7 +107,7 @@ namespace uba
 				return;
 
 			LockEventMutex();
-			m_manualReset = true;
+			m_resetType = EventResetType_Manual;
 			UnlockEventMutex();
 			Set();
 
@@ -130,7 +130,7 @@ namespace uba
 				return;
 			LockEventMutex();
 
-			if (m_manualReset)
+			if (m_resetType == EventResetType_Manual)
 			{
 				m_triggered = TriggerType_All;
 				if (pthread_cond_broadcast(&m_condition) != 0)
@@ -247,8 +247,7 @@ namespace uba
 			return false;
 		}
 
-		bool GetManualReset() { return m_manualReset; }
-		void ResetManualReset(bool mr) {}
+		EventResetType GetResetType() { return m_resetType; }
 
 		// Must be first!
 		EventPThread* next;
@@ -257,7 +256,7 @@ namespace uba
 		enum TriggerType : u8 { TriggerType_None, TriggerType_One, TriggerType_All };
 
 		Atomic<bool> m_initialized;
-		Atomic<bool> m_manualReset;
+		Atomic<EventResetType> m_resetType;
 		Atomic<TriggerType> m_triggered;
 		Atomic<u32> m_waitingThreads;
 		pthread_mutex_t m_mutex;
diff --git a/Engine/Source/Programs/UnrealBuildAccelerator/Core/Private/UbaEventWin32.inl b/Engine/Source/Programs/UnrealBuildAccelerator/Core/Private/UbaEventWin32.inl
--- a/Engine/Source/Programs/UnrealBuildAccelerator/Core/Private/UbaEventWin32.inl
+++ b/Engine/Source/Programs/UnrealBuildAccelerator/Core/Private/UbaEventWin32.inl
@@ -10,15 +10,35 @@ namespace uba
 		union
 		{
 			EventWin32* next;
-			bool manualReset;
+			EventResetType resetType;
 		};
 
-		bool Create(bool mr) { handle = CreateEvent(nullptr, mr, false, NULL); return handle != 0; }
-		void Set() { SetEvent(handle); }
-		void Reset() { ResetEvent(handle); }
-		bool IsSet(u32 timeOutMs) { return WaitForSingleObject(handle, timeOutMs) == 0;} 
-		bool GetManualReset() { return manualReset; }
-		void ResetManualReset(bool mr) { manualReset = mr; }
+		bool Create(EventResetType rt)
+		{
+			resetType = rt;
+			handle = CreateEvent(nullptr, rt == EventResetType_Manual, false, NULL);
+			return handle != 0;
+		}
+
+		void Set()
+		{
+			SetEvent(handle);
+		}
+
+		void Reset()
+		{
+			ResetEvent(handle);
+		}
+
+		bool IsSet(u32 timeOutMs)
+		{
+			return WaitForSingleObject(handle, timeOutMs) == 0;
+		}
+
+		EventResetType GetResetType()
+		{
+			return resetType;
+		}
 	};
 
 }
diff --git a/Engine/Source/Programs/UnrealBuildAccelerator/Core/Public/UbaEvent.h b/Engine/Source/Programs/UnrealBuildAccelerator/Core/Public/UbaEvent.h
--- a/Engine/Source/Programs/UnrealBuildAccelerator/Core/Public/UbaEvent.h
+++ b/Engine/Source/Programs/UnrealBuildAccelerator/Core/Public/UbaEvent.h
@@ -9,13 +9,19 @@ namespace uba
 	class StringBufferBase;
 	struct StringView;
 
+	enum EventResetType : u8
+	{
+		EventResetType_Auto,
+		EventResetType_Manual,
+	};
+
 	class Event
 	{
 	public:
 		Event() = default;
-		Event(bool manualReset);
+		Event(EventResetType resetType);
 		~Event();
-		bool Create(bool manualReset);
+		bool Create(EventResetType resetType);
 		void Destroy();
 		void Set();
 		void Reset();
@@ -34,7 +40,7 @@ namespace uba
 	public:
 		SharedEvent();
 		~SharedEvent();
-		bool Create(bool manualReset);
+		bool Create(EventResetType resetType);
 		bool Create(StringView name, bool isOwner);
 		void Destroy();
 		void Set();
diff --git a/Engine/Source/Programs/UnrealBuildAccelerator/StorageProxy/Private/UbaStorageProxyApp.cpp b/Engine/Source/Programs/UnrealBuildAccelerator/StorageProxy/Private/UbaStorageProxyApp.cpp
--- a/Engine/Source/Programs/UnrealBuildAccelerator/StorageProxy/Private/UbaStorageProxyApp.cpp
+++ b/Engine/Source/Programs/UnrealBuildAccelerator/StorageProxy/Private/UbaStorageProxyApp.cpp
@@ -177,7 +177,7 @@ int wmain(int argc, wchar_t *argv[])
 		if (exit)
 			return 0;
 
-		Event disconnected(true);
+		Event disconnected(EventResetType_Manual);
 
 		client->RegisterOnDisconnected([&]()
 			{
diff --git a/Engine/Source/Programs/UnrealBuildAccelerator/Test/Private/UbaTestBasics.cpp b/Engine/Source/Programs/UnrealBuildAccelerator/Test/Private/UbaTestBasics.cpp
--- a/Engine/Source/Programs/UnrealBuildAccelerator/Test/Private/UbaTestBasics.cpp
+++ b/Engine/Source/Programs/UnrealBuildAccelerator/Test/Private/UbaTestBasics.cpp
@@ -128,7 +128,7 @@ namespace uba
 	bool TestEvents(LoggerWithWriter& logger, const StringBufferBase& rootDir)
 	{
 		Event ev;
-		if (!ev.Create(true))
+		if (!ev.Create(EventResetType_Manual))
 			return logger.Error(TC("Failed to create event"));
 		if (!TestEventsImpl(logger, rootDir, ev, ev))
 			return false;
@@ -144,7 +144,7 @@ namespace uba
 		mappingBuffer.AddTransient(TC("Temp"), false);
 		MappedView view = mappingBuffer.AllocAndMapView(MappedView_Transient, 1024, 1, TC("Foo"));
 		SharedEvent* ev = new (view.memory) SharedEvent();
-		if (!ev->Create(true))
+		if (!ev->Create(EventResetType_Manual))
 			return logger.Error(TC("Failed to create event"));
 		if (!TestEventsImpl(logger, rootDir, *ev, *ev))
 			return false;
@@ -656,9 +656,9 @@ namespace uba
 
 	bool TestThreads(LoggerWithWriter& logger, const StringBufferBase& rootDir)
 	{
-		Event traverseDone(true);
+		Event traverseDone(EventResetType_Manual);
 
-		Event ev1(true);
+		Event ev1(EventResetType_Manual);
 		Thread t1([&]()
 			{
 				TestFunctionForThread(ev1, traverseDone);
@@ -666,7 +666,7 @@ namespace uba
 			});
 		ev1.IsSet();
 
-		Event ev2(true);
+		Event ev2(EventResetType_Manual);
 		Thread t2([&]()
 			{
 				TestFunctionForThread(ev2, traverseDone);
diff --git a/Engine/Source/Programs/UnrealBuildAccelerator/Test/Private/UbaTestNetwork.cpp b/Engine/Source/Programs/UnrealBuildAccelerator/Test/Private/UbaTestNetwork.cpp
--- a/Engine/Source/Programs/UnrealBuildAccelerator/Test/Private/UbaTestNetwork.cpp
+++ b/Engine/Source/Programs/UnrealBuildAccelerator/Test/Private/UbaTestNetwork.cpp
@@ -238,7 +238,7 @@ namespace uba
 			NetworkMessage msg(client, 1, SessionMessageType_Ping, writer);
 			writer.WriteU32(32);
 			StackBinaryReader<32> reader;
-			Event ev(true);
+			Event ev(EventResetType_Manual);
 			if (!msg.SendAsync(reader, [](bool error, void* userData) { ((Event*)userData)->Set(); }, &ev))
 				return logger.Error(TC("Failed to get message"));
 			ev.IsSet();
diff --git a/Engine/Source/Programs/UnrealBuildAccelerator/Visualizer/Private/UbaVisualizer.cpp b/Engine/Source/Programs/UnrealBuildAccelerator/Visualizer/Private/UbaVisualizer.cpp
--- a/Engine/Source/Programs/UnrealBuildAccelerator/Visualizer/Private/UbaVisualizer.cpp
+++ b/Engine/Source/Programs/UnrealBuildAccelerator/Visualizer/Private/UbaVisualizer.cpp
@@ -249,7 +249,7 @@ namespace uba
 			return false;
 		}
 
-		m_listenTimeout.Create(false);
+		m_listenTimeout.Create(EventResetType_Auto);
 
 		m_listenChannel.Append(channelName);
 		m_looping = true;
@@ -329,7 +329,7 @@ namespace uba
 		if (!StartHwndThread())
 			return true;
 
-		m_clientDisconnect.Create(true);
+		m_clientDisconnect.Create(EventResetType_Manual);
 
 		wchar_t dots[] = TC("....");
 		u32 dotsCounter = 0;
