[UBA]
* Made SharedEvent::IsSempaphore private
* Removed unused event

diff --git a/Engine/Source/Programs/UnrealBuildAccelerator/Common/Private/UbaSessionClient.cpp b/Engine/Source/Programs/UnrealBuildAccelerator/Common/Private/UbaSessionClient.cpp
--- a/Engine/Source/Programs/UnrealBuildAccelerator/Common/Private/UbaSessionClient.cpp
+++ b/Engine/Source/Programs/UnrealBuildAccelerator/Common/Private/UbaSessionClient.cpp
@@ -59,11 +59,10 @@ namespace uba
 
 	struct SessionClient::ModuleInfo
 	{
-		ModuleInfo(const tchar* n, const CasKey& c, u32 a) : name(n), casKey(c), attributes(a), done(EventResetType_Manual) {}
+		ModuleInfo(const tchar* n, const CasKey& c, u32 a) : name(n), casKey(c), attributes(a) {}
 		TString name;
 		CasKey casKey;
 		u32 attributes;
-		Event done;
 	};
 
 	SessionClient::SessionClient(const SessionClientCreateInfo& info)
@@ -415,7 +414,6 @@ namespace uba
 				TrackWorkScope tws(m_client, AsView(TC("FetchModule")), ColorWork);
 				tws.AddHint(m.name);
 
-				auto g = MakeGuard([&]() { m.done.Set(); });
 				CasKey newCasKey;
 				bool storeUncompressed = true;
 				u64 fileSize;
