[UBA]
* Fixed so if detoured process asserts super early we resort to tty to write out error and don't try to resolve callstack

diff --git a/Engine/Source/Programs/UnrealBuildAccelerator/Detours/Private/PosixOS/UbaDetoursFunctionsPosix.cpp b/Engine/Source/Programs/UnrealBuildAccelerator/Detours/Private/PosixOS/UbaDetoursFunctionsPosix.cpp
--- a/Engine/Source/Programs/UnrealBuildAccelerator/Detours/Private/PosixOS/UbaDetoursFunctionsPosix.cpp
+++ b/Engine/Source/Programs/UnrealBuildAccelerator/Detours/Private/PosixOS/UbaDetoursFunctionsPosix.cpp
@@ -3009,19 +3009,29 @@ namespace uba
 		ScopedCriticalSection scs(cs);
 		static auto& sb =  *new StringBuffer<8*1024>;
 		WriteAssertInfo(sb, text, file, line, expr, context);
-		Rpc_ResolveCallstack(sb, 3 + skipCallstackCount, context);
-		Rpc_WriteLog(sb.data, sb.count, true, true);
 
-		if (!allowTerminate)
-			return;
+		if (g_isInitialized)
+		{
+			Rpc_ResolveCallstack(sb, 3 + skipCallstackCount, context);
+			Rpc_WriteLog(sb.data, sb.count, true, true);
 
-		BinaryWriter writer;
-		writer.WriteByte(MessageType_Exit);
-		writer.WriteU32(terminateCode); // Exit code
-		writer.WriteString(""); // Log name
-		g_stats.Write(writer);
-		g_kernelStats.Write(writer);
-		writer.Flush(false);
+			if (!allowTerminate)
+				return;
+
+			BinaryWriter writer;
+			writer.WriteByte(MessageType_Exit);
+			writer.WriteU32(terminateCode); // Exit code
+			writer.WriteString(""); // Log name
+			g_stats.Write(writer);
+			g_kernelStats.Write(writer);
+			writer.Flush(false);
+		}
+		else
+		{
+			g_debugFile = (FileHandle)open("/dev/tty", O_WRONLY);
+			sb.Append('\n');
+			WriteDebug(sb.data, sb.count);
+		}
 
 		CloseCom();
 		TRUE_WRAPPER(_exit)(int(terminateCode));
