[UBA]
* Improved error logging around mach messages/bootstrap

diff --git a/Engine/Source/Programs/UnrealBuildAccelerator/Common/Private/UbaSession.cpp b/Engine/Source/Programs/UnrealBuildAccelerator/Common/Private/UbaSession.cpp
--- a/Engine/Source/Programs/UnrealBuildAccelerator/Common/Private/UbaSession.cpp
+++ b/Engine/Source/Programs/UnrealBuildAccelerator/Common/Private/UbaSession.cpp
@@ -1827,7 +1827,8 @@ namespace uba
 			});
 
 		#if UBA_USE_NATIVE_MAC_SEMAPHORES
-		StartSemaphoreService();
+		if (!StartSemaphoreService())
+			return false;
 		#endif
 
 		return true;
@@ -4037,39 +4038,48 @@ namespace uba
 
 	u32 Session::GetSessionId()
 	{
-		return m_sessionId;
+		return getpid() + (m_sessionId << 20);
 	}
 
-	void Session::StartSemaphoreService()
+	bool Session::StartSemaphoreService()
 	{
 		static Atomic<u32> id;
 		m_sessionId = id++;
 		mach_port_t port = MACH_PORT_NULL;
 		kern_return_t kr;
 		kr = mach_port_allocate(mach_task_self(), MACH_PORT_RIGHT_RECEIVE, &port);
-		UBA_ASSERT(kr == KERN_SUCCESS);
+		if (kr != KERN_SUCCESS)
+			return m_logger.Error(("mach_port_allocate failed (%u)"), kr);
 
 		// Give ourselves a send right so the child can contact us if needed.
 		// (Optional if you only expect to receive.)
 		kr = mach_port_insert_right(mach_task_self(), port, port, MACH_MSG_TYPE_MAKE_SEND);
-		UBA_ASSERT(kr == KERN_SUCCESS);
+		if (kr != KERN_SUCCESS)
+			return m_logger.Error(("mach_port_insert_right failed (%u)"), kr);
 
 		#pragma clang diagnostic push
 		#pragma clang diagnostic ignored "-Wdeprecated-declarations"
 		StringBuffer<64> service;
-		service.Append("UbaSemaphoreService").AppendValue(m_sessionId);
+		service.Append("UbaSemaphoreService").AppendValue(GetSessionId());
 		//name_t service = "com.epicgames.uba";
 		kr = bootstrap_register(bootstrap_port, service.data, port);
-		UBA_ASSERT(kr == KERN_SUCCESS);
+		if (kr != KERN_SUCCESS)
+			return m_logger.Error(("bootstrap_register failed trying to create %s (%u)"), service.data, kr);
 		#pragma clang diagnostic pop
 
 		m_semaphoreServiceInfo = new SemaphoreServiceInfo { port };
 		m_semaphoreServiceThread.Start([this]() { ThreadSemaphoreService(); return 0; }, TC("UbaSemSrv"));
+
+		return true;
 	}
 	
 	void Session::StopSemaphoreService()
 	{
+		#pragma clang diagnostic push
+		#pragma clang diagnostic ignored "-Wdeprecated-declarations"
 		mach_port_destroy(mach_task_self(), m_semaphoreServiceInfo->port);
+		#pragma clang diagnostic pop
+
 		delete m_semaphoreServiceInfo;
 		m_semaphoreServiceInfo = nullptr;
 		m_semaphoreServiceThread.Wait();
diff --git a/Engine/Source/Programs/UnrealBuildAccelerator/Common/Public/UbaSession.h b/Engine/Source/Programs/UnrealBuildAccelerator/Common/Public/UbaSession.h
--- a/Engine/Source/Programs/UnrealBuildAccelerator/Common/Public/UbaSession.h
+++ b/Engine/Source/Programs/UnrealBuildAccelerator/Common/Public/UbaSession.h
@@ -361,7 +361,7 @@ namespace uba
 
 		#if UBA_USE_NATIVE_MAC_SEMAPHORES
 		u32 GetSessionId();
-		void StartSemaphoreService();
+		bool StartSemaphoreService();
 		void StopSemaphoreService();
 		void ThreadSemaphoreService();
 		u32 m_sessionId;
diff --git a/Engine/Source/Programs/UnrealBuildAccelerator/Detours/Private/PosixOS/UbaDetoursMainPosix.cpp b/Engine/Source/Programs/UnrealBuildAccelerator/Detours/Private/PosixOS/UbaDetoursMainPosix.cpp
--- a/Engine/Source/Programs/UnrealBuildAccelerator/Detours/Private/PosixOS/UbaDetoursMainPosix.cpp
+++ b/Engine/Source/Programs/UnrealBuildAccelerator/Detours/Private/PosixOS/UbaDetoursMainPosix.cpp
@@ -105,6 +105,7 @@ static void __attribute__((constructor(102))) PreInitCtor()
 		service.Append(getenv("UBA_SESSION_ID"));
 		mach_port_t server_send = MACH_PORT_NULL;
 		kern_return_t kr = bootstrap_look_up(bootstrap_port, service.data, &server_send);
+		UBA_ASSERTF(kr == KERN_SUCCESS, "bootstrap_look_up failed to lookup %s (%u)", service.data, kr);
 		
 		mach_port_t reply_port = MACH_PORT_NULL;
 		kr = mach_port_allocate(mach_task_self(), MACH_PORT_RIGHT_RECEIVE, &reply_port);
@@ -118,12 +119,12 @@ static void __attribute__((constructor(102))) PreInitCtor()
 		req.header.msgh_local_port  = reply_port;
 		req.pid = u32(pid);
 		kr = mach_msg(&req.header, MACH_SEND_MSG, req.header.msgh_size, 0, MACH_PORT_NULL, MACH_MSG_TIMEOUT_NONE, MACH_PORT_NULL);
-		UBA_ASSERT(kr == KERN_SUCCESS);
+		UBA_ASSERTF(kr == KERN_SUCCESS, "Failed to send mach msg (%u)", kr);
 
 		struct { mach_msg_header_t header; mach_msg_body_t body; mach_msg_port_descriptor_t sem_desc[3]; uint8_t extra[512]; } rep;
 		memset(&rep, 0, sizeof(rep));
 		kr = mach_msg(&rep.header, MACH_RCV_MSG, 0, sizeof(rep), reply_port, MACH_MSG_TIMEOUT_NONE, MACH_PORT_NULL);
-		UBA_ASSERT(kr == KERN_SUCCESS);
+		UBA_ASSERTF(kr == KERN_SUCCESS, "Failed to receive mach msg (%u)", kr);
 
 		StringBuffer<32> b;
 		b.AppendValue(rep.sem_desc[0].name);
