[UBA]
* Added mac dispatch semaphore implementation for events that does not need to be shared. Note this code is buggy and not race condition safe.. will revisit

diff --git a/Engine/Source/Programs/UnrealBuildAccelerator/Core/Private/UbaEvent.cpp b/Engine/Source/Programs/UnrealBuildAccelerator/Core/Private/UbaEvent.cpp
--- a/Engine/Source/Programs/UnrealBuildAccelerator/Core/Private/UbaEvent.cpp
+++ b/Engine/Source/Programs/UnrealBuildAccelerator/Core/Private/UbaEvent.cpp
@@ -10,13 +10,20 @@
 
 #include "UbaEventPThread.inl"
 #include "UbaEventMachSemaphore.inl"
+#include "UbaEventDispatchSemaphore.inl"
 #include "UbaEventPosixSemaphore.inl"
 #include "UbaEventWin32.inl"
 
 namespace uba
 {
-	#if !PLATFORM_WINDOWS
+	#if UBA_USE_NATIVE_MAC_SEMAPHORES
+	using EventImpl = EventDispatchSemaphore;
+	using SharedEventSemaphoreImpl = EventMachSemaphore;
+	using SharedEventImpl = EventPThread;
+	#elif !PLATFORM_WINDOWS
 	using EventImpl = EventPThread;
+	using SharedEventSemaphoreImpl = EventPosixSemaphore;
+	using SharedEventImpl = EventPThread;
 	#else
 	using EventImpl = EventWin32;
 	#endif
@@ -115,17 +122,11 @@ namespace uba
 
 #if !PLATFORM_WINDOWS
 
-	#if UBA_USE_NATIVE_MAC_SEMAPHORES
-	using SharedEventImpl = EventMacSemaphore;
-	#else
-	using SharedEventImpl = EventPosixSemaphore;
-	#endif
-
 	SharedEvent::SharedEvent()
 	{
 		*((u8*)m_data) = 0;
-		static_assert(sizeof(m_data) >= sizeof(EventImpl));
 		static_assert(sizeof(m_data) >= sizeof(SharedEventImpl));
+		static_assert(sizeof(m_data) >= sizeof(SharedEventSemaphoreImpl));
 	}
 
 	SharedEvent::~SharedEvent()
@@ -136,18 +137,18 @@ namespace uba
 	bool SharedEvent::Create(EventResetType resetType)//, bool shared)
 	{
 		UBA_ASSERT(*((u8*)m_data) == 0);
-		new (m_data) EventImpl();
+		new (m_data) SharedEventImpl();
 		UBA_ASSERT(*((u8*)m_data) == 0);
 		*((u8*)m_data) = 1;
-		return ((EventImpl&)m_data).Create(resetType, true);
+		return ((SharedEventImpl&)m_data).Create(resetType, true);
 	}
 
 	bool SharedEvent::Create(StringView name, bool isOwner)
 	{
 		UBA_ASSERT(*((u8*)m_data) == 0);
-		new (m_data) SharedEventImpl();
+		new (m_data) SharedEventSemaphoreImpl();
 		*((u8*)m_data) = 2;
-		return ((SharedEventImpl&)m_data).Create(name, isOwner);
+		return ((SharedEventSemaphoreImpl&)m_data).Create(name, isOwner);
 	}
 
 	void SharedEvent::Destroy()
@@ -157,13 +158,13 @@ namespace uba
 
 		if (IsSemaphore())
 		{
-			((SharedEventImpl&)m_data).Destroy();
-			((SharedEventImpl&)m_data).~SharedEventImpl();
+			((SharedEventSemaphoreImpl&)m_data).Destroy();
+			((SharedEventSemaphoreImpl&)m_data).~SharedEventSemaphoreImpl();
 		}
 		else
 		{
-			((EventImpl&)m_data).Destroy();
-			((EventImpl&)m_data).~EventImpl();
+			((SharedEventImpl&)m_data).Destroy();
+			((SharedEventImpl&)m_data).~SharedEventImpl();
 		}
 		*((u8*)m_data) = 0;
 	}
@@ -174,9 +175,9 @@ namespace uba
 			return;
 
 		if (IsSemaphore())
-			((SharedEventImpl&)m_data).Set();
+			((SharedEventSemaphoreImpl&)m_data).Set();
 		else
-			((EventImpl&)m_data).Set();
+			((SharedEventImpl&)m_data).Set();
 	}
 
 	bool SharedEvent::IsSet(u32 timeOutMs)
@@ -185,9 +186,9 @@ namespace uba
 			return false;
 
 		if (IsSemaphore())
-			return ((SharedEventImpl&)m_data).IsSet(timeOutMs);
+			return ((SharedEventSemaphoreImpl&)m_data).IsSet(timeOutMs);
 		else
-			return ((EventImpl&)m_data).IsSet(timeOutMs);
+			return ((SharedEventImpl&)m_data).IsSet(timeOutMs);
 	}
 
 	bool SharedEvent::IsSemaphore()
@@ -198,7 +199,7 @@ namespace uba
 	void SharedEvent::ToString(StringBufferBase& out)
 	{
 		UBA_ASSERT(IsSemaphore());
-		((SharedEventImpl&)m_data).ToString(out);
+		((SharedEventSemaphoreImpl&)m_data).ToString(out);
 	}
 
 #endif
diff --git a/Engine/Source/Programs/UnrealBuildAccelerator/Core/Private/UbaEventDispatchSemaphore.inl b/Engine/Source/Programs/UnrealBuildAccelerator/Core/Private/UbaEventDispatchSemaphore.inl
new file mode 100644
--- /dev/null
+++ b/Engine/Source/Programs/UnrealBuildAccelerator/Core/Private/UbaEventDispatchSemaphore.inl
@@ -0,0 +1,145 @@
+// Copyright Epic Games, Inc. All Rights Reserved.
+
+#include "UbaDefinitions.h"
+
+#if UBA_USE_NATIVE_MAC_SEMAPHORES
+
+#include <dispatch/dispatch.h>
+
+namespace uba
+{
+	class EventDispatchSemaphore
+	{
+	public:
+		~EventDispatchSemaphore()
+		{
+			Destroy();
+		}
+
+		bool Create(bool manualReset, bool shared = false)
+		{
+			UBA_ASSERT(!shared);
+			UBA_ASSERT(!m_sem);
+
+			m_sem = dispatch_semaphore_create(0);
+			UBA_ASSERT(m_sem);
+			if (!m_sem)
+				return false;
+
+			m_manualReset.store(manualReset, std::memory_order_relaxed);
+			m_signaled.store(false, std::memory_order_relaxed);
+			return true;
+		}
+
+		void Destroy()
+		{
+			if (!m_sem)
+				return;
+
+			// Normalize state
+			m_signaled.store(false, std::memory_order_relaxed);
+
+			// Drain any surplus count to 0 so we don't leak resources conceptually
+			while (dispatch_semaphore_wait(m_sem, DISPATCH_TIME_NOW) == 0)
+			{
+				// drain
+			}
+
+			// Release semaphore
+			dispatch_release(m_sem);
+
+			m_sem = 0;
+		}
+
+		void Reset()
+		{
+			// Clear signaled state (manual reset) and drain surplus tokens (both modes)
+			m_signaled.store(false, std::memory_order_release);
+
+			// Bring the underlying semaphore count back to 0
+			while (dispatch_semaphore_wait(m_sem, DISPATCH_TIME_NOW) == 0)
+			{
+				// drain
+			}
+		}
+
+		bool IsCreated()
+		{
+			return m_sem != 0;
+		}
+
+		void Set()
+		{
+			if (m_manualReset.load(std::memory_order_acquire))
+			{
+				m_signaled.store(true, std::memory_order_release);
+				while (dispatch_semaphore_signal(m_sem) == 0)
+				{
+				}
+			}
+			else
+			{
+				dispatch_semaphore_signal(m_sem);
+			}
+		}
+
+		bool IsSet(u32 timeoutMs = ~0u)
+		{
+			if (m_manualReset.load(std::memory_order_acquire) && m_signaled.load(std::memory_order_acquire))
+				return true;
+
+			// Compute dispatch timeout
+			dispatch_time_t when;
+			if (timeoutMs == ~0u)
+				when = DISPATCH_TIME_FOREVER;
+			else
+				when = dispatch_time(DISPATCH_TIME_NOW,  u64(timeoutMs) * 1000000ull);
+
+			// Wait on the semaphore
+			long res = dispatch_semaphore_wait(m_sem, when);
+			if (res != 0)
+				return false;
+			return true;
+		}
+
+		bool GetManualReset()
+		{
+			return m_manualReset.load(std::memory_order_acquire);
+		}
+
+		void ResetManualReset(bool mr)
+		{
+			bool cur = m_manualReset.load(std::memory_order_acquire);
+			if (mr == cur)
+				return;
+
+			if (mr)
+			{
+				m_manualReset.store(true, std::memory_order_release);
+				m_signaled.store(false, std::memory_order_release);
+				while (dispatch_semaphore_wait(m_sem, DISPATCH_TIME_NOW) == 0)
+				{
+					// drain
+				}
+			}
+			else
+			{
+				bool wasSignaled = m_signaled.exchange(false, std::memory_order_acq_rel);
+
+				while (dispatch_semaphore_wait(m_sem, DISPATCH_TIME_NOW) == 0) { /* drain */ }
+				if (wasSignaled)
+					dispatch_semaphore_signal(m_sem);
+
+				m_manualReset.store(false, std::memory_order_release);
+			}
+		}
+	
+		EventDispatchSemaphore* next = nullptr;
+
+		dispatch_semaphore_t m_sem = 0;
+		Atomic<bool> m_signaled = false;
+		Atomic<bool> m_manualReset = false;
+	};
+}
+
+#endif
diff --git a/Engine/Source/Programs/UnrealBuildAccelerator/Core/Private/UbaEventMachSemaphore.inl b/Engine/Source/Programs/UnrealBuildAccelerator/Core/Private/UbaEventMachSemaphore.inl
--- a/Engine/Source/Programs/UnrealBuildAccelerator/Core/Private/UbaEventMachSemaphore.inl
+++ b/Engine/Source/Programs/UnrealBuildAccelerator/Core/Private/UbaEventMachSemaphore.inl
@@ -8,10 +8,10 @@
 
 namespace uba
 {
-	class EventMacSemaphore
+	class EventMachSemaphore
 	{
 	public:
-		~EventMacSemaphore()
+		~EventMachSemaphore()
 		{
 			Destroy();
 		}
diff --git a/Engine/Source/Programs/UnrealBuildAccelerator/Core/Private/UbaEventPThread.inl b/Engine/Source/Programs/UnrealBuildAccelerator/Core/Private/UbaEventPThread.inl
--- a/Engine/Source/Programs/UnrealBuildAccelerator/Core/Private/UbaEventPThread.inl
+++ b/Engine/Source/Programs/UnrealBuildAccelerator/Core/Private/UbaEventPThread.inl
@@ -250,7 +250,7 @@ namespace uba
 		EventResetType GetResetType() { return m_resetType; }
 
 		// Must be first!
-		EventPThread* next;
+		EventPThread* next = nullptr;
 
 	private:
 		enum TriggerType : u8 { TriggerType_None, TriggerType_One, TriggerType_All };
diff --git a/Engine/Source/Programs/UnrealBuildAccelerator/Core/Public/UbaDefinitions.h b/Engine/Source/Programs/UnrealBuildAccelerator/Core/Public/UbaDefinitions.h
--- a/Engine/Source/Programs/UnrealBuildAccelerator/Core/Public/UbaDefinitions.h
+++ b/Engine/Source/Programs/UnrealBuildAccelerator/Core/Public/UbaDefinitions.h
@@ -3,7 +3,7 @@
 #pragma once
 
 // Use native mach semaphores for IPC.
-#define UBA_USE_NATIVE_MAC_SEMAPHORES PLATFORM_MAC
+#define UBA_USE_NATIVE_MAC_SEMAPHORES 0//PLATFORM_MAC
 
 // Experimental. Testing to use file mappings that are backed by files to reduce page file pressure
 #define UBA_ENABLE_ON_DISK_FILE_MAPPINGS 0
