[UBA]
* Added initializer for Atomic<x> fields

diff --git a/Engine/Source/Programs/UnrealBuildAccelerator/Common/Private/UbaNetworkClient.cpp b/Engine/Source/Programs/UnrealBuildAccelerator/Common/Private/UbaNetworkClient.cpp
--- a/Engine/Source/Programs/UnrealBuildAccelerator/Common/Private/UbaNetworkClient.cpp
+++ b/Engine/Source/Programs/UnrealBuildAccelerator/Common/Private/UbaNetworkClient.cpp
@@ -27,15 +27,15 @@ namespace uba
 		Connection(NetworkClient& o) : owner(o), disconnectedEvent(EventResetType_Manual) {}
 		NetworkClient& owner;
 		void* backendConnection = nullptr;
-		Atomic<u32> connected;
+		Atomic<u32> connected = 0;
 		Event disconnectedEvent;
 		NetworkBackend* backend = nullptr;
 		Guid cryptoIV;
 
 		#if UBA_TRACK_NETWORK_TIMES
-		Atomic<u64> lastSendTime;
-		Atomic<u64> lastHeaderRecvTime;
-		Atomic<u64> lastBodyRecvTime;
+		Atomic<u64> lastSendTime = 0;
+		Atomic<u64> lastHeaderRecvTime = 0;
+		Atomic<u64> lastBodyRecvTime = 0;
 		#endif
 	};
 
@@ -94,7 +94,6 @@ namespace uba
 		{
 			RecvContext(NetworkClient& c, NetworkBackend& b, void* bc) : client(c), backend(b), backendConnection(bc), recvEvent(EventResetType_Manual), exitScopeEvent(EventResetType_Manual)
 			{
-				error = 255;
 			}
 
 			~RecvContext()
@@ -116,7 +115,7 @@ namespace uba
 			void* backendConnection;
 			Event recvEvent;
 			Event exitScopeEvent;
-			Atomic<u8> error;
+			Atomic<u8> error = 255;
 		};
 
 		backend.SetRecvTimeout(backendConnection, 10*1000, nullptr, nullptr); // 10 seconds timeout during adding connection. RecvContext set it to right time in ctor
diff --git a/Engine/Source/Programs/UnrealBuildAccelerator/Common/Private/UbaNetworkServer.cpp b/Engine/Source/Programs/UnrealBuildAccelerator/Common/Private/UbaNetworkServer.cpp
--- a/Engine/Source/Programs/UnrealBuildAccelerator/Common/Private/UbaNetworkServer.cpp
+++ b/Engine/Source/Programs/UnrealBuildAccelerator/Common/Private/UbaNetworkServer.cpp
@@ -33,10 +33,10 @@ namespace uba
 		Guid uid;
 		u32 id;
 		u32 queuedConnectionCount = 0;
-		Atomic<u32> refCount;
-		Atomic<u32> connectionCount;
-		Atomic<u64> sendBytes;
-		Atomic<u64> recvBytes;
+		Atomic<u32> refCount = 0;
+		Atomic<u32> connectionCount = 0;
+		Atomic<u64> sendBytes = 0;
+		Atomic<u64> recvBytes = 0;
 	};
 
 	struct NetworkServer::WorkerContext
@@ -110,8 +110,8 @@ namespace uba
 
 		WorkerContext* m_context = nullptr;
 
-		Atomic<bool> m_loop;
-		Atomic<bool> m_inUse;
+		Atomic<bool> m_loop = false;
+		Atomic<bool> m_inUse = false;
 		Thread m_thread;
 
 		Worker(const Worker&) = delete;
@@ -130,8 +130,6 @@ namespace uba
 		,	m_id(id)
 		,	m_backendConnection(backendConnection)
 		{
-			m_activeWorkerCount = 1;
-
 			m_backend.SetDisconnectCallback(m_backendConnection, this, [](void* context, const Guid& connectionUid, void* connection)
 				{
 					auto& conn = *(Connection*)context;
@@ -553,9 +551,9 @@ namespace uba
 		CryptoKey m_cryptoKey;
 		Guid m_cryptoIV;
 		Event m_disconnectCallbackCalled;
-		Atomic<int> m_activeWorkerCount;
-		Atomic<int> m_disconnectCalled;
-		Atomic<bool> m_disconnected;
+		Atomic<int> m_activeWorkerCount = 1;
+		Atomic<int> m_disconnectCalled = 0;
+		Atomic<bool> m_disconnected = false;
 		u32 m_id = 0;
 		u32 m_recvTimeoutCount = 0;
 		bool m_shouldDisconnect = false;
diff --git a/Engine/Source/Programs/UnrealBuildAccelerator/Common/Private/UbaSessionClient.cpp b/Engine/Source/Programs/UnrealBuildAccelerator/Common/Private/UbaSessionClient.cpp
--- a/Engine/Source/Programs/UnrealBuildAccelerator/Common/Private/UbaSessionClient.cpp
+++ b/Engine/Source/Programs/UnrealBuildAccelerator/Common/Private/UbaSessionClient.cpp
@@ -1868,8 +1868,8 @@ namespace uba
 			ProcessRec(ProcessImpl* impl) : handle(impl) {}
 			ProcessHandle handle;
 			Futex lock;
-			Atomic<bool> isKilled;
-			Atomic<bool> isDone;
+			Atomic<bool> isKilled = false;
+			Atomic<bool> isDone = false;
 			float weight = 1.0f;
 		};
 		List<ProcessRec> activeProcesses;
diff --git a/Engine/Source/Programs/UnrealBuildAccelerator/Common/Private/UbaSessionServer.cpp b/Engine/Source/Programs/UnrealBuildAccelerator/Common/Private/UbaSessionServer.cpp
--- a/Engine/Source/Programs/UnrealBuildAccelerator/Common/Private/UbaSessionServer.cpp
+++ b/Engine/Source/Programs/UnrealBuildAccelerator/Common/Private/UbaSessionServer.cpp
@@ -364,7 +364,7 @@ namespace uba
 
 					Container container((const tchar*)knownInputsCopy, knownInputsCount);
 
-					Atomic<u32> keysIndex;
+					Atomic<u32> keysIndex = 0;
 					const TString& workingDir = remoteProcess->m_startInfo.workingDirStr;
 
 					m_server.ParallelFor(knownInputsCount, container, [&](const WorkContext&, auto& it)
