[UBA]
* Improved mac dispatch semaphore implementation

diff --git a/Engine/Source/Programs/UnrealBuildAccelerator/Core/Private/UbaEventDispatchSemaphore.inl b/Engine/Source/Programs/UnrealBuildAccelerator/Core/Private/UbaEventDispatchSemaphore.inl
--- a/Engine/Source/Programs/UnrealBuildAccelerator/Core/Private/UbaEventDispatchSemaphore.inl
+++ b/Engine/Source/Programs/UnrealBuildAccelerator/Core/Private/UbaEventDispatchSemaphore.inl
@@ -11,134 +11,146 @@ namespace uba
 	class EventDispatchSemaphore
 	{
 	public:
-		~EventDispatchSemaphore()
+		EventDispatchSemaphore()
 		{
-			Destroy();
 		}
 
-		bool Create(bool manualReset, bool shared = false)
+		EventDispatchSemaphore(EventResetType resetType)
 		{
-			UBA_ASSERT(!shared);
-			UBA_ASSERT(!m_sem);
+			Create(resetType);
+		}
 
-			m_sem = dispatch_semaphore_create(0);
-			UBA_ASSERT(m_sem);
-			if (!m_sem)
-				return false;
+		~EventDispatchSemaphore()
+		{
+			Destroy();
+		}
 
-			m_manualReset.store(manualReset, std::memory_order_relaxed);
-			m_signaled.store(false, std::memory_order_relaxed);
-			return true;
+		void Create(EventResetType resetType)
+		{
+			m_waiters = 0;
+			m_setLatch = 0;
+			//m_autoCredit = 0; // Union
+			m_resetType = resetType;
+			m_gate = dispatch_semaphore_create(0);
 		}
 
 		void Destroy()
 		{
-			if (!m_sem)
-				return;
-
-			// Normalize state
-			m_signaled.store(false, std::memory_order_relaxed);
-
-			// Drain any surplus count to 0 so we don't leak resources conceptually
-			while (dispatch_semaphore_wait(m_sem, DISPATCH_TIME_NOW) == 0)
+			if (m_gate)
 			{
-				// drain
+				//dispatch_release(m_gate);
+				m_gate = 0;
 			}
-
-			// Release semaphore
-			dispatch_release(m_sem);
-
-			m_sem = 0;
 		}
 
-		void Reset()
+		bool IsSet(u64 timeoutMs = ~0u)
 		{
-			// Clear signaled state (manual reset) and drain surplus tokens (both modes)
-			m_signaled.store(false, std::memory_order_release);
-
-			// Bring the underlying semaphore count back to 0
-			while (dispatch_semaphore_wait(m_sem, DISPATCH_TIME_NOW) == 0)
+			if (m_resetType == EventResetType_Auto)
 			{
-				// drain
-			}
-		}
+				if (m_autoCredit.exchange(0, std::memory_order_acq_rel) == 1) // Fast-path: consume pending credit if present.
+					return true;
 
-		bool IsCreated()
-		{
-			return m_sem != 0;
-		}
+				m_waiters.fetch_add(1, std::memory_order_acq_rel); // Register as a waiter.
 
-		void Set()
-		{
-			if (m_manualReset.load(std::memory_order_acquire))
-			{
-				m_signaled.store(true, std::memory_order_release);
-				while (dispatch_semaphore_signal(m_sem) == 0)
+				if (m_autoCredit.exchange(0, std::memory_order_acq_rel) == 1) // Re-check after registration to close the set-during-entry race.
 				{
+					m_waiters.fetch_sub(1, std::memory_order_acq_rel);
+					return true;
 				}
+
+				bool res = WaitDispatch(m_gate, timeoutMs);
+				m_waiters.fetch_sub(1, std::memory_order_acq_rel); // We were either signaled or timed out. Adjust waiter count.
+				return res;
 			}
-			else
+			else // ManualReset
 			{
-				dispatch_semaphore_signal(m_sem);
-			}
-		}
+				if (m_setLatch.load(std::memory_order_acquire)) // Fast-path: already latched open.
+					return true;
 
-		bool IsSet(u32 timeoutMs = ~0u)
-		{
-			if (m_manualReset.load(std::memory_order_acquire) && m_signaled.load(std::memory_order_acquire))
-				return true;
+				m_waiters.fetch_add(1, std::memory_order_acq_rel); // Register as a waiter first…
 
-			// Compute dispatch timeout
-			dispatch_time_t when;
-			if (timeoutMs == ~0u)
-				when = DISPATCH_TIME_FOREVER;
-			else
-				when = dispatch_time(DISPATCH_TIME_NOW,  u64(timeoutMs) * 1000000ull);
+				if (m_setLatch.load(std::memory_order_acquire)) // …then re-check latch to avoid missing a concurrent Set().
+				{
+					m_waiters.fetch_sub(1, std::memory_order_acq_rel);
+					return true;
+				}
 
-			// Wait on the semaphore
-			long res = dispatch_semaphore_wait(m_sem, when);
-			if (res != 0)
-				return false;
-			return true;
-		}
+				bool res = WaitDispatch(m_gate, timeoutMs); // Block. On success, Set() will have drained waiters for us.
 
-		bool GetManualReset()
-		{
-			return m_manualReset.load(std::memory_order_acquire);
+				if (!res) // Timed out or error: remove our registration.
+					m_waiters.fetch_sub(1, std::memory_order_acq_rel); 
+
+				return res;
+			}
 		}
 
-		void ResetManualReset(bool mr)
+		// Set the event:
+		//  - AutoReset: deliver exactly one credit (wake one waiter or latch for next).
+		//  - ManualReset: latch open and wake all currently registered waiters exactly once.
+		void Set()
 		{
-			bool cur = m_manualReset.load(std::memory_order_acquire);
-			if (mr == cur)
-				return;
-
-			if (mr)
+			if (m_resetType == EventResetType_Auto)
 			{
-				m_manualReset.store(true, std::memory_order_release);
-				m_signaled.store(false, std::memory_order_release);
-				while (dispatch_semaphore_wait(m_sem, DISPATCH_TIME_NOW) == 0)
+				// First, try to create a pending credit if none exists.
+				int expected = 0;
+				if (m_autoCredit.compare_exchange_strong(expected, 1, std::memory_order_acq_rel, std::memory_order_acquire))
 				{
-					// drain
+					if (m_waiters.load(std::memory_order_acquire) > 0) // If waiters exist, transfer the credit to exactly one waiter now.
+					{
+						m_autoCredit.store(0, std::memory_order_release); // Move credit to the queue; ensure no double-credit remains.
+						dispatch_semaphore_signal(m_gate);
+					}
 				}
+				// If a credit already existed, collapse extra Set() calls (no-op).
 			}
 			else
 			{
-				bool wasSignaled = m_signaled.exchange(false, std::memory_order_acq_rel);
+				// ManualReset: latch and wake all currently registered waiters.
+				int was = m_setLatch.exchange(1, std::memory_order_acq_rel);
 
-				while (dispatch_semaphore_wait(m_sem, DISPATCH_TIME_NOW) == 0) { /* drain */ }
-				if (wasSignaled)
-					dispatch_semaphore_signal(m_sem);
+				if (was == 1) // already set
+					return;
 
-				m_manualReset.store(false, std::memory_order_release);
+				// Snapshot and drain current waiters exactly once.
+				int n = m_waiters.exchange(0, std::memory_order_acq_rel);
+				for (int i = 0; i < n; ++i)
+					dispatch_semaphore_signal(m_gate);
 			}
 		}
-	
+
+		void Reset()
+		{
+			UBA_ASSERT(m_waiters.load() == 0);
+			m_setLatch.store(0, std::memory_order_release);
+			//m_autoCredit.store(0, std::memory_order_release); // union
+		}
+
+		EventResetType GetResetType()
+		{
+			return m_resetType;
+		}
+
 		EventDispatchSemaphore* next = nullptr;
+	
+	private:
+
+		static bool WaitDispatch(dispatch_semaphore_t sem, u64 timeoutMs)
+		{
+			if (timeoutMs == ~0u)
+				timeoutMs = 24llu * 60 * 60 * 1000; // 24 hours (FOREVER can cause sigill)
+			u64 timeoutNs = timeoutMs * 1'000'000ull;
+			dispatch_time_t t = dispatch_time(DISPATCH_TIME_NOW, timeoutNs);
+			return dispatch_semaphore_wait(sem, t) == 0;
+		}
 
-		dispatch_semaphore_t m_sem = 0;
-		Atomic<bool> m_signaled = false;
-		Atomic<bool> m_manualReset = false;
+		EventResetType m_resetType;
+		dispatch_semaphore_t m_gate = 0;
+		Atomic<int> m_waiters;    // >= 0
+		union
+		{
+			Atomic<int> m_setLatch;   // 0/1
+			Atomic<int> m_autoCredit; // 0/1
+		};
 	};
 }
 
