[UBA]
* Added initialization to all Atomic<x> fields

diff --git a/Engine/Source/Programs/UnrealBuildAccelerator/Agent/Private/UbaAgent.cpp b/Engine/Source/Programs/UnrealBuildAccelerator/Agent/Private/UbaAgent.cpp
--- a/Engine/Source/Programs/UnrealBuildAccelerator/Agent/Private/UbaAgent.cpp
+++ b/Engine/Source/Programs/UnrealBuildAccelerator/Agent/Private/UbaAgent.cpp
@@ -133,8 +133,8 @@ namespace uba
 	ReaderWriterLock* g_exitLock = new ReaderWriterLock();
 	LoggerWithWriter* g_logger;
 	SessionClient* g_sessionClient;
-	Atomic<bool> g_shouldExit;
-	Atomic<bool> g_ctrlPressed;
+	Atomic<bool> g_shouldExit = false;
+	Atomic<bool> g_ctrlPressed = false;
 
 	bool ShouldExit()
 	{
@@ -1208,7 +1208,7 @@ namespace uba
 
 
 			StorageProxy* storageProxy = nullptr;
-			Atomic<NetworkServer*> proxyNetworkServer;
+			Atomic<NetworkServer*> proxyNetworkServer = nullptr;
 			auto psg = MakeGuard([&]() { delete proxyNetworkServer; });
 			auto pg = MakeGuard([&]() { delete storageProxy; });
 			TString proxyNetworkServerPrefix;
@@ -1244,7 +1244,7 @@ namespace uba
 					return true;
 				};
 
-			Atomic<bool> isDisconnected;
+			Atomic<bool> isDisconnected = false;
 			client->RegisterOnDisconnected([&]() { isDisconnected = true; networkBackend->StopListen(); if (auto proxyServer = proxyNetworkServer.load()) proxyServer->DisconnectClients(); });
 
 			struct NetworkBackends
diff --git a/Engine/Source/Programs/UnrealBuildAccelerator/CacheService/Private/UbaCacheService.cpp b/Engine/Source/Programs/UnrealBuildAccelerator/CacheService/Private/UbaCacheService.cpp
--- a/Engine/Source/Programs/UnrealBuildAccelerator/CacheService/Private/UbaCacheService.cpp
+++ b/Engine/Source/Programs/UnrealBuildAccelerator/CacheService/Private/UbaCacheService.cpp
@@ -82,7 +82,7 @@ namespace uba
 
 	ReaderWriterLock* g_exitLock = new ReaderWriterLock();
 	LoggerWithWriter* g_logger;
-	Atomic<bool> g_shouldExit;
+	Atomic<bool> g_shouldExit = false;
 
 	bool ShouldExit()
 	{
diff --git a/Engine/Source/Programs/UnrealBuildAccelerator/Cli/Private/UbaCli.cpp b/Engine/Source/Programs/UnrealBuildAccelerator/Cli/Private/UbaCli.cpp
--- a/Engine/Source/Programs/UnrealBuildAccelerator/Cli/Private/UbaCli.cpp
+++ b/Engine/Source/Programs/UnrealBuildAccelerator/Cli/Private/UbaCli.cpp
@@ -1390,7 +1390,7 @@ namespace uba
 			scheduler.GetStats(queued, activeLocal, activeRemote, outFinished);
 
 			bool success = true;
-			Atomic<u32> counter;
+			Atomic<u32> counter = 0;
 			static Event finished(EventResetType_Manual);
 
 			scheduler.SetProcessFinishedCallback([&](const ProcessHandle& ph)
diff --git a/Engine/Source/Programs/UnrealBuildAccelerator/Common/Private/UbaCacheClient.cpp b/Engine/Source/Programs/UnrealBuildAccelerator/Common/Private/UbaCacheClient.cpp
--- a/Engine/Source/Programs/UnrealBuildAccelerator/Common/Private/UbaCacheClient.cpp
+++ b/Engine/Source/Programs/UnrealBuildAccelerator/Common/Private/UbaCacheClient.cpp
@@ -84,21 +84,21 @@ namespace uba
 
 		Futex serverPathTableNetworkLock;
 		CompactPathTable serverPathTable;
-		Atomic<u32> serverPathTableSize;
+		Atomic<u32> serverPathTableSize = 0;
 
 		Futex serverCasKeyTableNetworkLock;
 		CompactCasKeyTable serverCasKeyTable;
-		Atomic<u32> serverCasKeyTableSize;
+		Atomic<u32> serverCasKeyTableSize = 0;
 
 		CompactPathTable sendPathTable;
 		Futex sendPathTableNetworkLock;
 		u32 pathTableSizeSent = 0;
-		Atomic<u32> pathTableSizeMaxToSend;
+		Atomic<u32> pathTableSizeMaxToSend = 0;
 
 		CompactCasKeyTable sendCasKeyTable;
 		Futex sendCasKeyTableNetworkLock;
 		u32 casKeyTableSizeSent = 0;
-		Atomic<u32> casKeyTableSizeMaxToSend;
+		Atomic<u32> casKeyTableSizeMaxToSend = 0;
 	};
 
 	////////////////////////////////////////////////////////////////////////////////////////////////////
@@ -172,6 +172,7 @@ namespace uba
 		m_validateCacheWritesOutput = info.validateCacheWritesOutput;
 		m_useCacheHit = info.useCacheHit;
 		m_useRoots = info.useRoots;
+		m_connected = false;
 
 		m_client.RegisterOnConnected([this, hint = TString(info.hint)]()
 			{
diff --git a/Engine/Source/Programs/UnrealBuildAccelerator/Common/Private/UbaCacheServer.cpp b/Engine/Source/Programs/UnrealBuildAccelerator/Common/Private/UbaCacheServer.cpp
--- a/Engine/Source/Programs/UnrealBuildAccelerator/Common/Private/UbaCacheServer.cpp
+++ b/Engine/Source/Programs/UnrealBuildAccelerator/Common/Private/UbaCacheServer.cpp
@@ -767,7 +767,7 @@ namespace uba
 
 		u64 totalCasSize = 0;
 
-		struct CasFileInfo { CasFileInfo(u32 s = 0) : size(s) {} u32 size; Atomic<bool> isUsed; }; // These are compressed cas, should never be over 4gb
+		struct CasFileInfo { CasFileInfo(u32 s = 0) : size(s) {} u32 size; Atomic<bool> isUsed = false; }; // These are compressed cas, should never be over 4gb
 
 		MemoryBlock existingCasMemoryBlock;
 		HashMap<CasKey, CasFileInfo> existingCas;
@@ -843,14 +843,14 @@ namespace uba
 		u32 workerCountToUse = workerCount > 0 ? workerCount - 1 : 0;
 		u32 workerCountToUseForBuckets = Min(workerCountToUse, u32(m_buckets.size()));
 
-		Atomic<u64> totalEntryCount;
-		Atomic<u64> deleteEntryCount;
-		Atomic<u64> expiredEntryCount;
-		Atomic<u64> overflowedEntryCount;
-		Atomic<u64> missingOutputEntryCount;
-		Atomic<u64> missingInputEntryCount;
+		Atomic<u64> totalEntryCount = 0;
+		Atomic<u64> deleteEntryCount = 0;
+		Atomic<u64> expiredEntryCount = 0;
+		Atomic<u64> overflowedEntryCount = 0;
+		Atomic<u64> missingOutputEntryCount = 0;
+		Atomic<u64> missingInputEntryCount = 0;
 
-		Atomic<u64> activeDropCount;
+		Atomic<u64> activeDropCount = 0;
 		auto dropCasGuard = MakeGuard([&]() { while (activeDropCount != 0) Sleep(1); });
 
 		auto EnsureBucketContextInitialized = [&](Bucket& bucket)
diff --git a/Engine/Source/Programs/UnrealBuildAccelerator/Common/Private/UbaFileAccessor.cpp b/Engine/Source/Programs/UnrealBuildAccelerator/Common/Private/UbaFileAccessor.cpp
--- a/Engine/Source/Programs/UnrealBuildAccelerator/Common/Private/UbaFileAccessor.cpp
+++ b/Engine/Source/Programs/UnrealBuildAccelerator/Common/Private/UbaFileAccessor.cpp
@@ -22,7 +22,7 @@ namespace uba
 	HANDLE asHANDLE(FileHandle fh);
 	#else
 	int asFileDescriptor(FileHandle fh);
-	Atomic<u32> g_tempFileCounter;
+	Atomic<u32> g_tempFileCounter = 0;
 	#endif
 		
 #if PLATFORM_WINDOWS
diff --git a/Engine/Source/Programs/UnrealBuildAccelerator/Common/Private/UbaFileMapping.cpp b/Engine/Source/Programs/UnrealBuildAccelerator/Common/Private/UbaFileMapping.cpp
--- a/Engine/Source/Programs/UnrealBuildAccelerator/Common/Private/UbaFileMapping.cpp
+++ b/Engine/Source/Programs/UnrealBuildAccelerator/Common/Private/UbaFileMapping.cpp
@@ -21,7 +21,7 @@ namespace uba
 {
 	#if UBA_DEBUG_FILE_MAPPING
 	Futex g_fileMappingsLock;
-	struct DebugFileMapping { TString hint; Atomic<u64> viewCount; };
+	struct DebugFileMapping { TString hint; Atomic<u64> viewCount = 0; };
 	UnorderedMap<HANDLE, DebugFileMapping> g_fileMappings;
 	UnorderedMap<const void*, HANDLE> g_viewMappings;
 	#endif
@@ -62,7 +62,7 @@ namespace uba
 	int asFileDescriptor(FileHandle fh);
 
 	Futex g_mappingUidCounterLock;
-	Atomic<u64> g_mappingUidCounter;
+	Atomic<u64> g_mappingUidCounter = 0;
 #endif
 
 	FileMappingHandle CreateMemoryMappingW(Logger& logger, u32 flProtect, u64 maxSize, const tchar* name, const tchar* hint)
diff --git a/Engine/Source/Programs/UnrealBuildAccelerator/Common/Private/UbaNetworkBackendMemory.cpp b/Engine/Source/Programs/UnrealBuildAccelerator/Common/Private/UbaNetworkBackendMemory.cpp
--- a/Engine/Source/Programs/UnrealBuildAccelerator/Common/Private/UbaNetworkBackendMemory.cpp
+++ b/Engine/Source/Programs/UnrealBuildAccelerator/Common/Private/UbaNetworkBackendMemory.cpp
@@ -25,13 +25,13 @@ namespace uba
 			void* disconnectContext = nullptr;
 
 			#if UBA_VALIDATE_MULTIPLE_SEND
-			Atomic<u32> sendCount;
+			Atomic<u32> sendCount = 0;
 			#endif
 		};
 
 		Peer peer[2];
 		Guid uid;
-		Atomic<bool> connected;
+		Atomic<bool> connected = false;
 	};
 
 	NetworkBackendMemory::NetworkBackendMemory(LogWriter& writer, const tchar* prefix)
diff --git a/Engine/Source/Programs/UnrealBuildAccelerator/Common/Private/UbaNetworkBackendQuic.cpp b/Engine/Source/Programs/UnrealBuildAccelerator/Common/Private/UbaNetworkBackendQuic.cpp
--- a/Engine/Source/Programs/UnrealBuildAccelerator/Common/Private/UbaNetworkBackendQuic.cpp
+++ b/Engine/Source/Programs/UnrealBuildAccelerator/Common/Private/UbaNetworkBackendQuic.cpp
@@ -11,7 +11,7 @@
 
 namespace uba
 {
-	Atomic<u32> connectionCount;
+	Atomic<u32> connectionCount = 0;
 
 	NetworkBackendQuic::Connection::Connection(NetworkBackendQuic& b, HQUIC ch)
 	:	backend(b)
diff --git a/Engine/Source/Programs/UnrealBuildAccelerator/Common/Private/UbaScheduler.cpp b/Engine/Source/Programs/UnrealBuildAccelerator/Common/Private/UbaScheduler.cpp
--- a/Engine/Source/Programs/UnrealBuildAccelerator/Common/Private/UbaScheduler.cpp
+++ b/Engine/Source/Programs/UnrealBuildAccelerator/Common/Private/UbaScheduler.cpp
@@ -1960,7 +1960,7 @@ namespace uba
 			Logger& logger;
 			Event& done;
 			Atomic<u32> counter;
-		} context { logger, done };
+		} context { logger, done, 0 };
 
 		auto exitedFunc = [](void* userData, const ProcessHandle& ph, ProcessExitedResponse&)
 			{
diff --git a/Engine/Source/Programs/UnrealBuildAccelerator/Common/Private/UbaSessionClient.cpp b/Engine/Source/Programs/UnrealBuildAccelerator/Common/Private/UbaSessionClient.cpp
--- a/Engine/Source/Programs/UnrealBuildAccelerator/Common/Private/UbaSessionClient.cpp
+++ b/Engine/Source/Programs/UnrealBuildAccelerator/Common/Private/UbaSessionClient.cpp
@@ -71,8 +71,6 @@ namespace uba
 	,	m_name(info.name.data)
 	,	m_terminationTime(~0ull)
 	,	m_waitToSendEvent(EventResetType_Auto)
-	,	m_loop(true)
-	,	m_allowSpawn(true)
 	{
 		m_maxProcessCount = info.maxProcessCount;
 		m_dedicated = info.dedicated;
diff --git a/Engine/Source/Programs/UnrealBuildAccelerator/Common/Private/UbaSessionServer.cpp b/Engine/Source/Programs/UnrealBuildAccelerator/Common/Private/UbaSessionServer.cpp
--- a/Engine/Source/Programs/UnrealBuildAccelerator/Common/Private/UbaSessionServer.cpp
+++ b/Engine/Source/Programs/UnrealBuildAccelerator/Common/Private/UbaSessionServer.cpp
@@ -143,7 +143,6 @@ namespace uba
 	SessionServer::SessionServer(const SessionServerCreateInfo& info, const u8* environment, u32 environmentSize)
 	:	Session(info, TC("UbaSessionServer"), false, info.server)
 	,	m_server(info.server)
-	,	m_maxRemoteProcessCount(~0u)
 	{
 		m_server.RegisterOnClientDisconnected(ServiceId, [this](const Guid& clientUid, u32 clientId) { OnDisconnected(clientUid, clientId); });
 
diff --git a/Engine/Source/Programs/UnrealBuildAccelerator/Common/Private/UbaStorage.cpp b/Engine/Source/Programs/UnrealBuildAccelerator/Common/Private/UbaStorage.cpp
--- a/Engine/Source/Programs/UnrealBuildAccelerator/Common/Private/UbaStorage.cpp
+++ b/Engine/Source/Programs/UnrealBuildAccelerator/Common/Private/UbaStorage.cpp
@@ -640,7 +640,7 @@ namespace uba
 
 	void StorageImpl::TraverseAllCasFiles(const tchar* dir, const Function<void(const StringBufferBase& fullPath, const DirectoryEntry& e)>& func, bool allowParallel)
 	{
-		Atomic<u32> workLeft;
+		Atomic<u32> workLeft = 0;
 
 		auto traverseCasFileDir = [&func, &workLeft, this](const StringView& casFileDir)
 			{
@@ -1446,10 +1446,10 @@ namespace uba
 		List<TString> directories;
 		directories.push_back(searchPath);
 
-		Atomic<u32> foundFiles;
-		Atomic<u32> trackedFiles;
-		Atomic<u32> testedFiles;
-		Atomic<u32> errorCount;
+		Atomic<u32> foundFiles = 0;
+		Atomic<u32> trackedFiles = 0;
+		Atomic<u32> testedFiles = 0;
+		Atomic<u32> errorCount = 0;
 
 		u64 startTime = GetTime();
 		
@@ -1556,7 +1556,7 @@ namespace uba
 		u32 deleteCount = 0;
 
 		{
-			Atomic<u32> atomicDeleteCount;
+			Atomic<u32> atomicDeleteCount = 0;
 			TraverseDir(m_logger, m_rootDir, [&](const DirectoryEntry& e)
 				{
 					if (!IsDirectory(e.attributes))
diff --git a/Engine/Source/Programs/UnrealBuildAccelerator/Common/Private/UbaStorageUtils.cpp b/Engine/Source/Programs/UnrealBuildAccelerator/Common/Private/UbaStorageUtils.cpp
--- a/Engine/Source/Programs/UnrealBuildAccelerator/Common/Private/UbaStorageUtils.cpp
+++ b/Engine/Source/Programs/UnrealBuildAccelerator/Common/Private/UbaStorageUtils.cpp
@@ -64,9 +64,9 @@ namespace uba
 		{
 			struct WorkRec
 			{
-				Atomic<u64> refCount;
-				Atomic<u64> counter;
-				Atomic<u64> doneCounter;
+				Atomic<u64> refCount = 2;
+				Atomic<u64> counter = 0;
+				Atomic<u64> doneCounter = 0;
 				u8* fileMem = nullptr;
 				u64 workCount = 0;
 				u64 fileSize = 0;
@@ -83,7 +83,6 @@ namespace uba
 			rec->fileSize = fileSize;
 			rec->keys.resize(workCount);
 			rec->done.Create(EventResetType_Manual);
-			rec->refCount = 2;
 
 			auto work = [rec](const WorkContext& context)
 			{
@@ -359,7 +358,7 @@ namespace uba
 				u32 count;
 			} container(sendCount);
 
-			Atomic<bool> error;
+			Atomic<bool> error = false;
 
 			workManager->ParallelFor(connectionCount - 1, container, [&](const WorkContext&, auto& it)
 				{
@@ -1077,7 +1076,7 @@ namespace uba
 			LinearWriter* destination;
 			BufferCache* bufferSlots;
 			Atomic<u64> refCount;
-			Atomic<u64> compressCounter;
+			Atomic<u64> compressCounter = 0;
 			Event* events;
 			const u8* uncompressedData = nullptr;
 			u64 written = 0;
@@ -1332,7 +1331,7 @@ namespace uba
 			{
 				Logger* logger;
 				const tchar* hint;
-				Atomic<u64> refCount;
+				Atomic<u64> refCount = 2;
 				const u8* readPos = nullptr;
 				u8* writePos = nullptr;
 
@@ -1341,7 +1340,7 @@ namespace uba
 				u64 decompressedLeft = 0;
 				u64 written = 0;
 				Event done;
-				Atomic<bool> error;
+				Atomic<bool> error = false;
 			};
 
 			WorkRec* rec = new WorkRec();
@@ -1352,7 +1351,6 @@ namespace uba
 			rec->decompressedSize = decompressedSize;
 			rec->decompressedLeft = decompressedSize;
 			rec->done.Create(EventResetType_Manual);
-			rec->refCount = 2;
 
 			auto work = [rec](const WorkContext&)
 				{
diff --git a/Engine/Source/Programs/UnrealBuildAccelerator/Common/Private/UbaWorkManager.cpp b/Engine/Source/Programs/UnrealBuildAccelerator/Common/Private/UbaWorkManager.cpp
--- a/Engine/Source/Programs/UnrealBuildAccelerator/Common/Private/UbaWorkManager.cpp
+++ b/Engine/Source/Programs/UnrealBuildAccelerator/Common/Private/UbaWorkManager.cpp
@@ -188,7 +188,6 @@ namespace uba
 	{
 		Worker(WorkManagerImpl& manager, const tchar* workerDesc) : m_workAvailable(EventResetType_Auto)
 		{
-			m_loop = true;
 			m_thread.Start([&]() { ThreadWorker(manager); return 0; }, workerDesc);
 		}
 		~Worker()
@@ -244,7 +243,7 @@ namespace uba
 
 		Worker* m_nextWorker = nullptr;
 		Worker* m_prevWorker = nullptr;
-		Atomic<bool> m_loop;
+		Atomic<bool> m_loop = true;
 		Event m_workAvailable;
 		Thread m_thread;
 	};
diff --git a/Engine/Source/Programs/UnrealBuildAccelerator/Common/Public/UbaBottleneck.h b/Engine/Source/Programs/UnrealBuildAccelerator/Common/Public/UbaBottleneck.h
--- a/Engine/Source/Programs/UnrealBuildAccelerator/Common/Public/UbaBottleneck.h
+++ b/Engine/Source/Programs/UnrealBuildAccelerator/Common/Public/UbaBottleneck.h
@@ -72,7 +72,7 @@ namespace uba
 		Futex m_lock;
 		BottleneckTicket* m_first = nullptr;
 		BottleneckTicket* m_last = nullptr;
-		Atomic<u32> m_activeCount;
+		Atomic<u32> m_activeCount = 0;
 		u32 m_maxCount;
 	};
 
diff --git a/Engine/Source/Programs/UnrealBuildAccelerator/Common/Public/UbaCacheBucket.h b/Engine/Source/Programs/UnrealBuildAccelerator/Common/Public/UbaCacheBucket.h
--- a/Engine/Source/Programs/UnrealBuildAccelerator/Common/Public/UbaCacheBucket.h
+++ b/Engine/Source/Programs/UnrealBuildAccelerator/Common/Public/UbaCacheBucket.h
@@ -23,12 +23,12 @@ namespace uba
 		u64 m_id;
 		u64 totalEntryCount = 0;
 		u64 totalEntrySize = 0;
-		Atomic<bool> hasDeletedEntries;
-		Atomic<bool> needsSave;
+		Atomic<bool> hasDeletedEntries = false;
+		Atomic<bool> needsSave = false;
 
 		// Times are in file time from creation of cache database
-		Atomic<u64> lastSavedTime;
-		Atomic<u64> lastUsedTime;
+		Atomic<u64> lastSavedTime = 0;
+		Atomic<u64> lastUsedTime = 0;
 		u64 oldestUsedTime = 0; 
 
 		u32 index = ~0u;
diff --git a/Engine/Source/Programs/UnrealBuildAccelerator/Common/Public/UbaCacheClient.h b/Engine/Source/Programs/UnrealBuildAccelerator/Common/Public/UbaCacheClient.h
--- a/Engine/Source/Programs/UnrealBuildAccelerator/Common/Public/UbaCacheClient.h
+++ b/Engine/Source/Programs/UnrealBuildAccelerator/Common/Public/UbaCacheClient.h
@@ -151,7 +151,7 @@ namespace uba
 		Futex m_sendOneAtTheTimeLock;
 
 		Futex m_directoryPreparserLock;
-		struct PreparedDir { Futex lock; Atomic<bool> done; };
+		struct PreparedDir { Futex lock; Atomic<bool> done = false; };
 		UnorderedMap<StringKey, PreparedDir> m_directoryPreparser;
 
 		struct PathHash { TString path; CasKey hash; };
diff --git a/Engine/Source/Programs/UnrealBuildAccelerator/Common/Public/UbaCacheServer.h b/Engine/Source/Programs/UnrealBuildAccelerator/Common/Public/UbaCacheServer.h
--- a/Engine/Source/Programs/UnrealBuildAccelerator/Common/Public/UbaCacheServer.h
+++ b/Engine/Source/Programs/UnrealBuildAccelerator/Common/Public/UbaCacheServer.h
@@ -99,12 +99,12 @@ namespace uba
 
 		StringBuffer<MaxPath> m_rootDir;
 
-		Atomic<u32> m_addsSinceMaintenance;
-		Atomic<u64> m_cacheKeyFetchCount;
-		Atomic<u64> m_cacheKeyHitCount;
-		Atomic<bool> m_isRunningMaintenance;
-		Atomic<bool> m_bucketIsOverflowing;
-		Atomic<bool> m_forceMaintenance;
+		Atomic<u32> m_addsSinceMaintenance = 0;
+		Atomic<u64> m_cacheKeyFetchCount = 0;
+		Atomic<u64> m_cacheKeyHitCount = 0;
+		Atomic<bool> m_isRunningMaintenance = false;
+		Atomic<bool> m_bucketIsOverflowing = false;
+		Atomic<bool> m_forceMaintenance = false;
 
 		Futex m_bucketsLock;
 		Map<u64, Bucket> m_buckets;
@@ -114,8 +114,8 @@ namespace uba
 
 		UnorderedSet<CasKey> m_trackedDeletes;
 
-		Atomic<bool> m_shutdownRequested;
-		Atomic<u64> m_totalEntryCount;
+		Atomic<bool> m_shutdownRequested = false;
+		Atomic<u64> m_totalEntryCount = 0;
 
 		u64 m_maintenanceReserveSize = 0;
 		u64 m_bucketCasTableMaxSize = 0;
diff --git a/Engine/Source/Programs/UnrealBuildAccelerator/Common/Public/UbaDependencyCrawler.h b/Engine/Source/Programs/UnrealBuildAccelerator/Common/Public/UbaDependencyCrawler.h
--- a/Engine/Source/Programs/UnrealBuildAccelerator/Common/Public/UbaDependencyCrawler.h
+++ b/Engine/Source/Programs/UnrealBuildAccelerator/Common/Public/UbaDependencyCrawler.h
@@ -80,7 +80,7 @@ namespace uba
 			Vector<IncludeRoot> includeRoots;
 			Pch* pch = nullptr;
 
-			Atomic<u32> refCount;
+			Atomic<u32> refCount = 0;
 		};
 
 		struct InstanceRef
diff --git a/Engine/Source/Programs/UnrealBuildAccelerator/Common/Public/UbaLogger.h b/Engine/Source/Programs/UnrealBuildAccelerator/Common/Public/UbaLogger.h
--- a/Engine/Source/Programs/UnrealBuildAccelerator/Common/Public/UbaLogger.h
+++ b/Engine/Source/Programs/UnrealBuildAccelerator/Common/Public/UbaLogger.h
@@ -72,7 +72,7 @@ namespace uba
 		MutableLogger(LogWriter& writer, const tchar* prefix) : LoggerWithWriter(writer, prefix) {}
 		virtual void Log(LogEntryType type, const tchar* str, u32 strLen) override { if (!isMuted) LoggerWithWriter::Log(type, str, strLen); }
 		using LoggerWithWriter::Log;
-		Atomic<u32> isMuted;
+		Atomic<u32> isMuted = false;
 	};
 
 	////////////////////////////////////////////////////////////////////////////////////////////////////
diff --git a/Engine/Source/Programs/UnrealBuildAccelerator/Common/Public/UbaNetworkBackendTcp.h b/Engine/Source/Programs/UnrealBuildAccelerator/Common/Public/UbaNetworkBackendTcp.h
--- a/Engine/Source/Programs/UnrealBuildAccelerator/Common/Public/UbaNetworkBackendTcp.h
+++ b/Engine/Source/Programs/UnrealBuildAccelerator/Common/Public/UbaNetworkBackendTcp.h
@@ -81,8 +81,8 @@ namespace uba
 		Futex m_connectionsLock;
 		List<Connection> m_connections;
 
-		Atomic<u64> m_totalSend;
-		Atomic<u64> m_totalRecv;
+		Atomic<u64> m_totalSend = 0;
+		Atomic<u64> m_totalRecv = 0;
 
 		bool m_disableNagle;
 
diff --git a/Engine/Source/Programs/UnrealBuildAccelerator/Common/Public/UbaNetworkClient.h b/Engine/Source/Programs/UnrealBuildAccelerator/Common/Public/UbaNetworkClient.h
--- a/Engine/Source/Programs/UnrealBuildAccelerator/Common/Public/UbaNetworkClient.h
+++ b/Engine/Source/Programs/UnrealBuildAccelerator/Common/Public/UbaNetworkClient.h
@@ -97,10 +97,10 @@ namespace uba
 		u32 m_recvBufferSize;
 		u32 m_sendBufferSize;
 		u32 m_receiveTimeoutSeconds;
-		Atomic<u64> m_sendBytes;
-		Atomic<u64> m_recvBytes;
-		Atomic<u32> m_recvCount;
-		Atomic<bool> m_isDisconnecting;
+		Atomic<u64> m_sendBytes = 0;
+		Atomic<u64> m_recvBytes = 0;
+		Atomic<u32> m_recvCount = 0;
+		Atomic<bool> m_isDisconnecting = false;
 		Timer m_sendTimer;
 
 		Futex m_serverUidLock;
@@ -108,7 +108,7 @@ namespace uba
 
 		Event m_isConnected;
 		Event m_isOrWasConnected;
-		Atomic<u32> m_connectionCount;
+		Atomic<u32> m_connectionCount = 0;
 		bool m_useMessagePriority;
 
 		Futex m_onConnectedFunctionsLock;
diff --git a/Engine/Source/Programs/UnrealBuildAccelerator/Common/Public/UbaNetworkMessage.h b/Engine/Source/Programs/UnrealBuildAccelerator/Common/Public/UbaNetworkMessage.h
--- a/Engine/Source/Programs/UnrealBuildAccelerator/Common/Public/UbaNetworkMessage.h
+++ b/Engine/Source/Programs/UnrealBuildAccelerator/Common/Public/UbaNetworkMessage.h
@@ -42,7 +42,7 @@ namespace uba
 		u32 m_responseCapacity = 0;
 		u16 m_id = 0;
 		MessagePriority m_priority;
-		Atomic<u32> m_error;
+		Atomic<u32> m_error = 0;
 		NetworkClient::Connection* m_connection = nullptr;
 		Guid m_cryptoIV;
 		friend NetworkClient;
diff --git a/Engine/Source/Programs/UnrealBuildAccelerator/Common/Public/UbaNetworkServer.h b/Engine/Source/Programs/UnrealBuildAccelerator/Common/Public/UbaNetworkServer.h
--- a/Engine/Source/Programs/UnrealBuildAccelerator/Common/Public/UbaNetworkServer.h
+++ b/Engine/Source/Programs/UnrealBuildAccelerator/Common/Public/UbaNetworkServer.h
@@ -203,9 +203,9 @@ namespace uba
 		Timer m_sendTimer;
 		Timer m_encryptTimer;
 		Timer m_decryptTimer;
-		Atomic<u64> m_sendBytes;
-		Atomic<u64> m_recvBytes;
-		Atomic<u32> m_recvCount;
+		Atomic<u64> m_sendBytes = 0;
+		Atomic<u64> m_recvBytes = 0;
+		Atomic<u32> m_recvCount = 0;
 
 		Trace* m_trace = nullptr;
 
diff --git a/Engine/Source/Programs/UnrealBuildAccelerator/Common/Public/UbaProcess.h b/Engine/Source/Programs/UnrealBuildAccelerator/Common/Public/UbaProcess.h
--- a/Engine/Source/Programs/UnrealBuildAccelerator/Common/Public/UbaProcess.h
+++ b/Engine/Source/Programs/UnrealBuildAccelerator/Common/Public/UbaProcess.h
@@ -36,7 +36,7 @@ namespace uba
 		void AddRef();
 		void Release();
 		virtual ~Process() {}
-		Atomic<u32> m_refCount;
+		Atomic<u32> m_refCount = 0;
 		friend ProcessHandle;
 	};
 
@@ -154,8 +154,8 @@ namespace uba
 		u32 m_exitCode = ~0u;
 		u32 m_messageCount = 0;
 		Atomic<bool> m_cancelAllowed = true;
-		Atomic<bool> m_cancelled;
-		Atomic<bool> m_hasExited;
+		Atomic<bool> m_cancelled = false;
+		Atomic<bool> m_hasExited = false;
 		MessageType m_firstMessageError = MessageType_None;
 		bool m_runningRemote = false;
 		bool m_gotExitMessage = false;
diff --git a/Engine/Source/Programs/UnrealBuildAccelerator/Common/Public/UbaScheduler.h b/Engine/Source/Programs/UnrealBuildAccelerator/Common/Public/UbaScheduler.h
--- a/Engine/Source/Programs/UnrealBuildAccelerator/Common/Public/UbaScheduler.h
+++ b/Engine/Source/Programs/UnrealBuildAccelerator/Common/Public/UbaScheduler.h
@@ -180,7 +180,7 @@ namespace uba
 
 		u32 m_processEntriesCacheStart = 0;
 
-		Atomic<bool> m_loop;
+		Atomic<bool> m_loop = false;
 		bool m_enableProcessReuse;
 		bool m_forceRemote;
 		bool m_forceNative;
@@ -191,16 +191,16 @@ namespace uba
 
 		float m_activeLocalProcessWeight = 0.0f;
 		u32 m_activeCacheQueries = 0;
-		Atomic<u32> m_activeCacheDownloads;
+		Atomic<u32> m_activeCacheDownloads = 0;
 
-		Atomic<u32> m_totalProcesses;
-		Atomic<u32> m_queuedProcesses;
-		Atomic<u32> m_activeLocalProcesses;
-		Atomic<u32> m_activeRemoteProcesses;
-		Atomic<u32> m_finishedProcesses;
-		Atomic<u32> m_errorCount;
-		Atomic<u32> m_cacheHitCount;
-		Atomic<u32> m_cacheMissCount;
+		Atomic<u32> m_totalProcesses = 0;
+		Atomic<u32> m_queuedProcesses = 0;
+		Atomic<u32> m_activeLocalProcesses = 0;
+		Atomic<u32> m_activeRemoteProcesses = 0;
+		Atomic<u32> m_finishedProcesses = 0;
+		Atomic<u32> m_errorCount = 0;
+		Atomic<u32> m_cacheHitCount = 0;
+		Atomic<u32> m_cacheMissCount = 0;
 
 		Vector<CacheClient*> m_cacheClients;
 		Vector<RootPaths*> m_rootPaths;
@@ -215,8 +215,8 @@ namespace uba
 		Event m_memThreadEvent;
 		Thread m_memThread;
 		bool m_memDelaySpawning = false;
-		Atomic<u64> m_memEstimatedUsage;
-		Atomic<u64> m_memWaitThreshold;
+		Atomic<u64> m_memEstimatedUsage = 0;
+		Atomic<u64> m_memWaitThreshold = 0;
 		u8 m_memStartWaitPercent = 0;
 		u8 m_memStartKillPercent = 0;
 		u8 m_memTracingLevel = 0;
diff --git a/Engine/Source/Programs/UnrealBuildAccelerator/Common/Public/UbaSession.h b/Engine/Source/Programs/UnrealBuildAccelerator/Common/Public/UbaSession.h
--- a/Engine/Source/Programs/UnrealBuildAccelerator/Common/Public/UbaSession.h
+++ b/Engine/Source/Programs/UnrealBuildAccelerator/Common/Public/UbaSession.h
@@ -252,7 +252,7 @@ namespace uba
 
 		u64 m_keepOutputFileMemoryMapsThreshold;
 
-		Atomic<u32> m_processIdCounter;
+		Atomic<u32> m_processIdCounter = 0;
 
 		MemoryBlock m_directoryTableMemory;
 
diff --git a/Engine/Source/Programs/UnrealBuildAccelerator/Common/Public/UbaSessionClient.h b/Engine/Source/Programs/UnrealBuildAccelerator/Common/Public/UbaSessionClient.h
--- a/Engine/Source/Programs/UnrealBuildAccelerator/Common/Public/UbaSessionClient.h
+++ b/Engine/Source/Programs/UnrealBuildAccelerator/Common/Public/UbaSessionClient.h
@@ -137,10 +137,10 @@ namespace uba
 		bool m_remoteExecutionEnabled = true;
 		bool m_useDependencyCrawler = false;
 		
-		Atomic<const tchar*> m_terminationReason;
-		Atomic<u64> m_terminationTime;
-		Atomic<u32> m_maxProcessCount;
-		Atomic<float> m_cpuUsage;
+		Atomic<const tchar*> m_terminationReason = nullptr;
+		Atomic<u64> m_terminationTime = 0;
+		Atomic<u32> m_maxProcessCount = 0;
+		Atomic<float> m_cpuUsage = 0.0f;
 
 		struct ApplicationEnvironment { Futex lock; TString virtualApplication; TString realApplication; };
 		Futex m_handledApplicationEnvironmentsLock;
@@ -167,21 +167,21 @@ namespace uba
 
 		Event m_waitToSendEvent;
 		Thread m_loopThread;
-		Atomic<bool> m_loop;
-		Atomic<bool> m_sendPing;
-		Atomic<bool> m_allowSpawn;
+		Atomic<bool> m_loop = true;
+		Atomic<bool> m_sendPing = false;
+		Atomic<bool> m_allowSpawn = true;
 
 		Function<void(const ProcessHandle&)> m_processFinished;
 
 		SessionSummaryStats m_stats;
 
-		Atomic<u64> m_bestPing;
-		Atomic<u64> m_lastPing;
+		Atomic<u64> m_bestPing = 0;
+		Atomic<u64> m_lastPing = 0;
 		u64 m_lastPingSendTime = 0;
 		u32 m_pingTimeoutSecondsPrintCallstacks = 0;
 
-		Atomic<u64> m_memAvail;
-		Atomic<u64> m_memTotal;
+		Atomic<u64> m_memAvail = 0;
+		Atomic<u64> m_memTotal = 0;
 
 		UnorderedMap<CasKey, Vector<u8>> m_hostRunCache;
 		Futex m_hostRunCacheLock;
diff --git a/Engine/Source/Programs/UnrealBuildAccelerator/Common/Public/UbaSessionServer.h b/Engine/Source/Programs/UnrealBuildAccelerator/Common/Public/UbaSessionServer.h
--- a/Engine/Source/Programs/UnrealBuildAccelerator/Common/Public/UbaSessionServer.h
+++ b/Engine/Source/Programs/UnrealBuildAccelerator/Common/Public/UbaSessionServer.h
@@ -131,7 +131,7 @@ namespace uba
 
 		NetworkServer& m_server;
 		u32 m_uiLanguage;
-		Atomic<u32> m_maxRemoteProcessCount;
+		Atomic<u32> m_maxRemoteProcessCount = ~0u;
 		bool m_resetCas;
 		bool m_remoteExecutionEnabled;
 		bool m_nameToHashTableEnabled;
@@ -195,7 +195,7 @@ namespace uba
 
 		UnorderedMap<StringKey, CasKey> m_nameToHashLookup;
 		Futex m_nameToHashLookupLock;
-		Atomic<bool> m_nameToHashInitialized;
+		Atomic<bool> m_nameToHashInitialized = false;
 
 		ReaderWriterLock m_receivedFilesLock;
 		UnorderedMap<StringKey, CasKey> m_receivedFiles;
diff --git a/Engine/Source/Programs/UnrealBuildAccelerator/Common/Public/UbaStorageProxy.h b/Engine/Source/Programs/UnrealBuildAccelerator/Common/Public/UbaStorageProxy.h
--- a/Engine/Source/Programs/UnrealBuildAccelerator/Common/Public/UbaStorageProxy.h
+++ b/Engine/Source/Programs/UnrealBuildAccelerator/Common/Public/UbaStorageProxy.h
@@ -73,14 +73,14 @@ namespace uba
 
 		TString m_name;
 
-		Atomic<u32> m_inProcessClientId;
+		Atomic<u32> m_inProcessClientId = 0;
 
 		struct FileEntry
 		{
 			Futex lock;
 			u8* memory = nullptr;
 			u64 size = 0;
-			Atomic<u64> received;
+			Atomic<u64> received = 0;
 			CasKey casKey;
 			u32 trackId = 0;
 			u16 fetchId = 0;
diff --git a/Engine/Source/Programs/UnrealBuildAccelerator/Common/Public/UbaStorageServer.h b/Engine/Source/Programs/UnrealBuildAccelerator/Common/Public/UbaStorageServer.h
--- a/Engine/Source/Programs/UnrealBuildAccelerator/Common/Public/UbaStorageServer.h
+++ b/Engine/Source/Programs/UnrealBuildAccelerator/Common/Public/UbaStorageServer.h
@@ -86,8 +86,8 @@ namespace uba
 			MappedView mappedView;
 			FileAccessor* fileAccessor = nullptr;
 			CasEntry* casEntry = nullptr;
-			Atomic<u64> totalWritten;
-			Atomic<u64> recvCasTime;
+			Atomic<u64> totalWritten = 0;
+			Atomic<u64> recvCasTime = 0;
 			u64 fileSize = 0;
 			u64 actualSize = 0;
 		};
@@ -98,9 +98,9 @@ namespace uba
 		struct ActiveFetch
 		{
 			u32 clientId = ~0u;
-			Atomic<u64> left;
+			Atomic<u64> left = 0;
 			CasEntry* casEntry;
-			Atomic<u64> sendCasTime;
+			Atomic<u64> sendCasTime = 0;
 
 			FileHandle readFileHandle = InvalidFileHandle;
 			MappedView mappedView;
@@ -112,7 +112,7 @@ namespace uba
 		};
 		ReaderWriterLock m_activeFetchesLock;
 		UnorderedMap<u16, ActiveFetch> m_activeFetches;
-		Atomic<u32> m_activeUnmap;
+		Atomic<u32> m_activeUnmap = 0;
 
 		ReaderWriterLock m_activeRefCountsLock;
 		UnorderedMap<u32, List<CasEntry*>> m_activeRefCounts;
diff --git a/Engine/Source/Programs/UnrealBuildAccelerator/Common/Public/UbaTrace.h b/Engine/Source/Programs/UnrealBuildAccelerator/Common/Public/UbaTrace.h
--- a/Engine/Source/Programs/UnrealBuildAccelerator/Common/Public/UbaTrace.h
+++ b/Engine/Source/Programs/UnrealBuildAccelerator/Common/Public/UbaTrace.h
@@ -183,7 +183,7 @@ namespace uba
 		Futex m_stringsLock;
 		UnorderedMap<StringKey, u32> m_strings;
 
-		Atomic<u32> m_workCounter;
+		Atomic<u32> m_workCounter = 0;
 
 		Futex m_threadCallbackLock;
 		Function<void()> m_threadCallback;
diff --git a/Engine/Source/Programs/UnrealBuildAccelerator/Common/Public/UbaWorkManager.h b/Engine/Source/Programs/UnrealBuildAccelerator/Common/Public/UbaWorkManager.h
--- a/Engine/Source/Programs/UnrealBuildAccelerator/Common/Public/UbaWorkManager.h
+++ b/Engine/Source/Programs/UnrealBuildAccelerator/Common/Public/UbaWorkManager.h
@@ -63,7 +63,7 @@ namespace uba
 		void ParallelFor(u32 workCount, TContainer& container, TFunc&& func, const StringView& description = AsView(TC("")), WorkPriority priority = WorkPriority_Normal, Bottleneck* bottleneck = nullptr);
 
 	protected:
-		Atomic<WorkTracker*> m_workTracker;
+		Atomic<WorkTracker*> m_workTracker = nullptr;
 	};
 
 	struct Work { WorkFunction func; Bottleneck* bottleneck; TString desc; Work* next; };
@@ -104,7 +104,7 @@ namespace uba
 		TString m_workerDesc;
 		List<Worker> m_workers;
 		WorkQueue m_workQueue;
-		Atomic<u32> m_activeWorkerCount;
+		Atomic<u32> m_activeWorkerCount = 0;
 
 		Futex m_availableWorkersLock;
 		Worker* m_firstAvailableWorker = nullptr;
diff --git a/Engine/Source/Programs/UnrealBuildAccelerator/Host/Private/UbaExports.cpp b/Engine/Source/Programs/UnrealBuildAccelerator/Host/Private/UbaExports.cpp
--- a/Engine/Source/Programs/UnrealBuildAccelerator/Host/Private/UbaExports.cpp
+++ b/Engine/Source/Programs/UnrealBuildAccelerator/Host/Private/UbaExports.cpp
@@ -77,7 +77,7 @@ namespace uba
 	{
 	public:
 		CacheClientWithCounter(const CacheClientCreateInfo& info) : CacheClient(info) {}
-		Atomic<u32> active;
+		Atomic<u32> active = 0;
 	};
 
 	struct CacheClientActiveScope
diff --git a/Engine/Source/Programs/UnrealBuildAccelerator/Test/Private/UbaTestNetwork.cpp b/Engine/Source/Programs/UnrealBuildAccelerator/Test/Private/UbaTestNetwork.cpp
--- a/Engine/Source/Programs/UnrealBuildAccelerator/Test/Private/UbaTestNetwork.cpp
+++ b/Engine/Source/Programs/UnrealBuildAccelerator/Test/Private/UbaTestNetwork.cpp
@@ -14,7 +14,7 @@ namespace uba
 	bool TestSockets(LoggerWithWriter& logger, const StringBufferBase& testRootDir)
 	{
 		static constexpr u16 port = 1346;
-		Atomic<u8> result;
+		Atomic<u8> result = 0;
 		NetworkBackendTcp tcp(logger.m_writer);
 		Thread t([&]()
 			{
diff --git a/Engine/Source/Programs/UnrealBuildAccelerator/Test/Private/UbaTestStress.cpp b/Engine/Source/Programs/UnrealBuildAccelerator/Test/Private/UbaTestStress.cpp
--- a/Engine/Source/Programs/UnrealBuildAccelerator/Test/Private/UbaTestStress.cpp
+++ b/Engine/Source/Programs/UnrealBuildAccelerator/Test/Private/UbaTestStress.cpp
@@ -88,8 +88,8 @@ namespace uba
 
 		u32 counter = 0;
 
-		Atomic<u32> activeCount;
-		Atomic<bool> error;
+		Atomic<u32> activeCount = 0;
+		Atomic<bool> error = false;
 		while (true)
 		{
 			StringBuffer<512> fileName;
diff --git a/Engine/Source/Programs/UnrealBuildAccelerator/Visualizer/Private/UbaVisualizer.h b/Engine/Source/Programs/UnrealBuildAccelerator/Visualizer/Private/UbaVisualizer.h
--- a/Engine/Source/Programs/UnrealBuildAccelerator/Visualizer/Private/UbaVisualizer.h
+++ b/Engine/Source/Programs/UnrealBuildAccelerator/Visualizer/Private/UbaVisualizer.h
@@ -195,7 +195,7 @@ namespace uba
 
 		ProcessBrushes m_processBrushes[2]; // Non-selected and selected
 
-		Atomic<bool> m_looping;
+		Atomic<bool> m_looping = false;
 		HWND m_hwnd = 0;
 		HWND m_parentHwnd = 0;
 		COLORREF m_textColor = {};
