[UBA]
* Implemented process-local semaphore based Event that works on all platforms

diff --git a/Engine/Source/Programs/UnrealBuildAccelerator/Common/Private/UbaFileAccessor.cpp b/Engine/Source/Programs/UnrealBuildAccelerator/Common/Private/UbaFileAccessor.cpp
--- a/Engine/Source/Programs/UnrealBuildAccelerator/Common/Private/UbaFileAccessor.cpp
+++ b/Engine/Source/Programs/UnrealBuildAccelerator/Common/Private/UbaFileAccessor.cpp
@@ -64,6 +64,11 @@ namespace uba
 
 	bool FileAccessor::InternalCreateWrite(bool allowRead, u32 flagsAndAttributes, u64 fileSize, const tchar* tempPath, bool isMemoryMap)
 	{
+		#if PLATFORM_WINDOWS
+		UBA_ASSERT((flagsAndAttributes & FILE_FLAG_OVERLAPPED) == 0 || Event::IsNative());
+		#endif
+
+
 		UBA_ASSERT(flagsAndAttributes != 0);
 		m_size = fileSize;
 		m_flagsAndAttributes = flagsAndAttributes;
diff --git a/Engine/Source/Programs/UnrealBuildAccelerator/Common/Private/UbaNetworkBackendTcp.cpp b/Engine/Source/Programs/UnrealBuildAccelerator/Common/Private/UbaNetworkBackendTcp.cpp
--- a/Engine/Source/Programs/UnrealBuildAccelerator/Common/Private/UbaNetworkBackendTcp.cpp
+++ b/Engine/Source/Programs/UnrealBuildAccelerator/Common/Private/UbaNetworkBackendTcp.cpp
@@ -255,7 +255,7 @@ namespace uba
 		m_disableNagle = info.disableNagle;
 
 		#if PLATFORM_WINDOWS
-		m_useOverlappedSend = info.useOverlappedSend;
+		m_useOverlappedSend = info.useOverlappedSend && Event::IsNative();
 		m_iocpWorkerCount = u16(info.iocpWorkerCount);
 		#endif
 
diff --git a/Engine/Source/Programs/UnrealBuildAccelerator/Common/Private/UbaSession.cpp b/Engine/Source/Programs/UnrealBuildAccelerator/Common/Private/UbaSession.cpp
--- a/Engine/Source/Programs/UnrealBuildAccelerator/Common/Private/UbaSession.cpp
+++ b/Engine/Source/Programs/UnrealBuildAccelerator/Common/Private/UbaSession.cpp
@@ -489,7 +489,7 @@ namespace uba
 		// On one machine we get twice as fast without overlap if no bottleneck. On another machine (ntfs compression on) we get twice as slow without overlap
 		// Both machines behaves well with overlap AND bottleneck. Both machine are 128 logical core thread rippers.
 		bool useFileMapForWrite = fileSize && fileSize <= m_writeFilesFileMapMax; // ::CreateFileMappingW does not work for zero-length files
-		bool useOverlap = !useFileMapForWrite && fileSize >= m_writeFilesNoBufferingMin;//fileSize > 2*1024*1024;
+		bool useOverlap = !useFileMapForWrite && fileSize >= m_writeFilesNoBufferingMin && Event::IsNative();//fileSize > 2*1024*1024;
 
 		u32 attributes = DefaultAttributes();
 		if (useOverlap)
diff --git a/Engine/Source/Programs/UnrealBuildAccelerator/Common/Private/UbaSessionServer.cpp b/Engine/Source/Programs/UnrealBuildAccelerator/Common/Private/UbaSessionServer.cpp
--- a/Engine/Source/Programs/UnrealBuildAccelerator/Common/Private/UbaSessionServer.cpp
+++ b/Engine/Source/Programs/UnrealBuildAccelerator/Common/Private/UbaSessionServer.cpp
@@ -2180,7 +2180,7 @@ namespace uba
 		m_receivedFiles.erase(findIt);
 		lock2.Leave();
 
-		m_storage.DropCasFile(casKey, true, filePath.data);
+		m_storage.DropCasFile(casKey, true, filePath.data, true);
 		return res;
 	}
 
diff --git a/Engine/Source/Programs/UnrealBuildAccelerator/Common/Private/UbaStorageClient.cpp b/Engine/Source/Programs/UnrealBuildAccelerator/Common/Private/UbaStorageClient.cpp
--- a/Engine/Source/Programs/UnrealBuildAccelerator/Common/Private/UbaStorageClient.cpp
+++ b/Engine/Source/Programs/UnrealBuildAccelerator/Common/Private/UbaStorageClient.cpp
@@ -583,7 +583,7 @@ namespace uba
 					else
 					{
 						u32 extraFlags = DefaultAttributes();
-						bool useOverlap = !IsRunningWine() && isCompressed == IsCompressed(casKey) && sizeOnDisk > 1024 * 1024;
+						bool useOverlap = !IsRunningWine() && isCompressed == IsCompressed(casKey) && sizeOnDisk > 1024 * 1024 && Event::IsNative();
 						if (useOverlap)
 							extraFlags |= FILE_FLAG_OVERLAPPED;
 						if (!destinationFile.CreateWrite(false, extraFlags, sizeOnDisk, m_tempPath.data))
diff --git a/Engine/Source/Programs/UnrealBuildAccelerator/Common/Private/UbaThread.cpp b/Engine/Source/Programs/UnrealBuildAccelerator/Common/Private/UbaThread.cpp
--- a/Engine/Source/Programs/UnrealBuildAccelerator/Common/Private/UbaThread.cpp
+++ b/Engine/Source/Programs/UnrealBuildAccelerator/Common/Private/UbaThread.cpp
@@ -171,7 +171,7 @@ namespace uba
 			};
 
 #if PLATFORM_WINDOWS // Optimization, not needed in initial implementation
-		if (wakeupEvent)
+		if (wakeupEvent && Event::IsNative())
 		{
 			HANDLE h[] = { m_handle, wakeupEvent->GetHandle() };
 			DWORD res = WaitForMultipleObjects(2, h, false, milliseconds);
diff --git a/Engine/Source/Programs/UnrealBuildAccelerator/Common/Public/UbaProcess.h b/Engine/Source/Programs/UnrealBuildAccelerator/Common/Public/UbaProcess.h
--- a/Engine/Source/Programs/UnrealBuildAccelerator/Common/Public/UbaProcess.h
+++ b/Engine/Source/Programs/UnrealBuildAccelerator/Common/Public/UbaProcess.h
@@ -118,9 +118,9 @@ namespace uba
 		FileMappingAllocator::Allocation m_comMemory;
 
 	#if PLATFORM_WINDOWS
-		Event m_cancelEvent;
-		Event m_writeEvent;
-		Event m_readEvent;
+		SharedEvent m_cancelEvent;
+		SharedEvent m_writeEvent;
+		SharedEvent m_readEvent;
 	#else
 		Futex m_comMemoryLock;
 		SharedEvent& m_cancelEvent;
diff --git a/Engine/Source/Programs/UnrealBuildAccelerator/Core/Private/UbaEvent.cpp b/Engine/Source/Programs/UnrealBuildAccelerator/Core/Private/UbaEvent.cpp
--- a/Engine/Source/Programs/UnrealBuildAccelerator/Core/Private/UbaEvent.cpp
+++ b/Engine/Source/Programs/UnrealBuildAccelerator/Core/Private/UbaEvent.cpp
@@ -7,6 +7,26 @@
 #include "UbaSynchronization.h"
 
 #define UBA_TEST_WAIT_QUALITY 0
+#define UBA_TEST_WINDOWS_SEMAPHORES 0
+
+namespace uba
+{
+	struct ResetTypeWithPadding
+	{
+		u8 padding; // This is to map into SharedEvent under the hood
+		EventResetType value;
+	};
+
+	struct EventHeader
+	{
+		EventHeader() { m_next = nullptr; }
+		union
+		{
+			ResetTypeWithPadding m_resetType;
+			EventHeader* m_next;
+		};
+	};
+}
 
 #include "UbaEventPThread.inl"
 #include "UbaEventMachSemaphore.inl"
@@ -14,22 +34,35 @@
 #include "UbaEventPosixSemaphore.inl"
 #include "UbaEventWin32.inl"
 
+
 namespace uba
 {
-	#if UBA_USE_NATIVE_MAC_SEMAPHORES
+#if UBA_USE_NATIVE_MAC_SEMAPHORES
+	bool Event::IsNative() { return false; }
 	using EventImpl = EventDispatchSemaphore;
 	using SharedEventSemaphoreImpl = EventMachSemaphore;
 	using SharedEventImpl = EventPThread;
-	#elif !PLATFORM_WINDOWS
+#elif !PLATFORM_WINDOWS
+	bool Event::IsNative() { return false; }
 	using EventImpl = EventPThread;
 	using SharedEventSemaphoreImpl = EventPosixSemaphore;
 	using SharedEventImpl = EventPThread;
+#else
+	using SharedEventSemaphoreImpl = EventWin32;
+	using SharedEventImpl = EventWin32;
+	#if UBA_TEST_WINDOWS_SEMAPHORES
+	using EventImpl = EventDispatchSemaphore;
+	HANDLE GetNativeHandle(EventImpl& e) { return 0; }
+	bool Event::IsNative() { return false; }
 	#else
 	using EventImpl = EventWin32;
+	HANDLE GetNativeHandle(EventImpl& e) { return e.m_handle; }
+	bool Event::IsNative() { return true; }
 	#endif
+#endif
 
 	Futex g_firstEventLock[2];
-	EventImpl* g_firstEvent[2];
+	EventHeader* g_firstEvent[2];
 
 	Event::Event(EventResetType resetType)
 	{
@@ -48,8 +81,8 @@ namespace uba
 		auto*& first = g_firstEvent[resetType];
 		if (first)
 		{
-			impl = first;
-			first = impl->next;
+			impl = (EventImpl*)first;
+			first = impl->m_next;
 		}
 		lock.Leave();
 
@@ -61,9 +94,7 @@ namespace uba
 		else
 		{
 			impl->Reset();
-			#if PLATFORM_WINDOWS
-			impl->resetType = resetType; // union with next
-			#endif
+			impl->m_resetType.value = resetType; // union with next
 		}
 
 		m_impl = impl;
@@ -76,10 +107,10 @@ namespace uba
 			return;
 		auto impl = (EventImpl*)m_impl;
 		m_impl = nullptr;
-		auto resetType = impl->GetResetType();
+		auto resetType = impl->m_resetType.value;
 		auto*& first = g_firstEvent[resetType];
 		SCOPED_FUTEX(g_firstEventLock[resetType], l);
-		impl->next = first;
+		impl->m_next = first;
 		first = impl;
 	}
 
@@ -112,16 +143,14 @@ namespace uba
 		#if PLATFORM_WINDOWS
 		if (!m_impl)
 			return nullptr;
-		return ((EventImpl*)m_impl)->handle;
+		return GetNativeHandle(*(EventImpl*)m_impl);
 		#else
 		UBA_ASSERTF(false, "Event::GetHandle not available");
-		return 0;
+		return nullptr;
 		#endif
 	}
 
 
-#if !PLATFORM_WINDOWS
-
 	SharedEvent::SharedEvent()
 	{
 		*((u8*)m_data) = 0;
@@ -140,7 +169,10 @@ namespace uba
 		new (m_data) SharedEventImpl();
 		UBA_ASSERT(*((u8*)m_data) == 0);
 		*((u8*)m_data) = 1;
-		return ((SharedEventImpl&)m_data).Create(resetType, true);
+		if (!((SharedEventImpl&)m_data).Create(resetType, true))
+			return false;
+		UBA_ASSERT(*((u8*)m_data) == 1);
+		return true;
 	}
 
 	bool SharedEvent::Create(StringView name, bool isOwner)
@@ -202,5 +234,14 @@ namespace uba
 		((SharedEventSemaphoreImpl&)m_data).ToString(out);
 	}
 
-#endif
+	void* SharedEvent::GetHandle()
+	{
+		UBA_ASSERT(!IsSemaphore());
+		#if PLATFORM_WINDOWS
+		return ((SharedEventImpl*)m_data)->m_handle;
+		#else
+		UBA_ASSERTF(false, "Event::GetHandle not available");
+		return nullptr;
+		#endif
+	}
 }
diff --git a/Engine/Source/Programs/UnrealBuildAccelerator/Core/Private/UbaEventDispatchSemaphore.inl b/Engine/Source/Programs/UnrealBuildAccelerator/Core/Private/UbaEventDispatchSemaphore.inl
--- a/Engine/Source/Programs/UnrealBuildAccelerator/Core/Private/UbaEventDispatchSemaphore.inl
+++ b/Engine/Source/Programs/UnrealBuildAccelerator/Core/Private/UbaEventDispatchSemaphore.inl
@@ -2,18 +2,18 @@
 
 #include "UbaDefinitions.h"
 
-#if UBA_USE_NATIVE_MAC_SEMAPHORES
-
+//#if UBA_USE_NATIVE_MAC_SEMAPHORES
+#if PLATFORM_LINUX
+#include <semaphore.h>
+#elif PLATFORM_MAC
 #include <dispatch/dispatch.h>
+#endif
 
 namespace uba
 {
-	class EventDispatchSemaphore
+	struct EventDispatchSemaphore : EventHeader
 	{
-	public:
-		EventDispatchSemaphore()
-		{
-		}
+		EventDispatchSemaphore() = default;
 
 		EventDispatchSemaphore(EventResetType resetType)
 		{
@@ -27,131 +27,194 @@ namespace uba
 
 		void Create(EventResetType resetType)
 		{
-			m_waiters = 0;
-			m_setLatch = 0;
-			//m_autoCredit = 0; // Union
-			m_resetType = resetType;
-			m_gate = dispatch_semaphore_create(0);
-		}
+			m_resetType.value = resetType;
+			m_triggered = TriggerType_None;
+
+			#if PLATFORM_WINDOWS
+			m_sem = CreateSemaphoreExW(nullptr, 0, 100000, nullptr, 0, SEMAPHORE_MODIFY_STATE | SYNCHRONIZE);
+			#elif PLATFORM_LINUX
+			m_sem = new sem_t;
+			sem_init(m_sem, 0, 0);
+			#else
+			m_sem = dispatch_semaphore_create(0);
+			#endif
+ 		}
 
 		void Destroy()
 		{
-			if (m_gate)
+			if (m_sem)
 			{
-				//dispatch_release(m_gate);
-				m_gate = 0;
+				#if PLATFORM_WINDOWS
+				//CloseHandle(m_sem);
+				#elif PLATFORM_LINUX
+				sem_destroy(m_sem);
+				delete m_sem;
+				#else
+				//dispatch_release(m_sem);
+				#endif
+				m_sem = 0;
 			}
 		}
 
-		bool IsSet(u64 timeoutMs = ~0u)
+		void Set()
 		{
-			if (m_resetType == EventResetType_Auto)
-			{
-				if (m_autoCredit.exchange(0, std::memory_order_acq_rel) == 1) // Fast-path: consume pending credit if present.
-					return true;
+			UBA_ASSERT(m_sem);
+
+			SCOPED_FUTEX(m_lock, lock);
 
-				m_waiters.fetch_add(1, std::memory_order_acq_rel); // Register as a waiter.
+			if (m_resetType.value == EventResetType_Manual)
+			{
+				m_triggered = TriggerType_All;
 
-				if (m_autoCredit.exchange(0, std::memory_order_acq_rel) == 1) // Re-check after registration to close the set-during-entry race.
+				// Release semaphore for all waiters... also fine releasing semaphore more than needed since it is a manual reset.
+				for (Waiter* it=m_oldestWaiter; it; it = it->next)
 				{
-					m_waiters.fetch_sub(1, std::memory_order_acq_rel);
-					return true;
+					#if PLATFORM_WINDOWS
+					LONG prev = 0;
+					ReleaseSemaphore(m_sem, 1, &prev);
+					#elif PLATFORM_LINUX
+					sem_post(m_sem);
+					#else
+					dispatch_semaphore_signal(m_sem);
+					#endif
 				}
-
-				bool res = WaitDispatch(m_gate, timeoutMs);
-				m_waiters.fetch_sub(1, std::memory_order_acq_rel); // We were either signaled or timed out. Adjust waiter count.
-				return res;
 			}
-			else // ManualReset
+			else
 			{
-				if (m_setLatch.load(std::memory_order_acquire)) // Fast-path: already latched open.
-					return true;
-
-				m_waiters.fetch_add(1, std::memory_order_acq_rel); // Register as a waiter first…
-
-				if (m_setLatch.load(std::memory_order_acquire)) // …then re-check latch to avoid missing a concurrent Set().
+				if (m_oldestWaiter)
+				{
+					#if PLATFORM_WINDOWS
+					LONG prev = 0;
+					ReleaseSemaphore(m_sem, 1, &prev);
+					#elif PLATFORM_LINUX
+					sem_post(m_sem);
+					#else
+					dispatch_semaphore_signal(m_sem);
+					#endif
+				}
+				else
 				{
-					m_waiters.fetch_sub(1, std::memory_order_acq_rel);
-					return true;
+					m_triggered = TriggerType_One;
 				}
+			}
+		}
 
-				bool res = WaitDispatch(m_gate, timeoutMs); // Block. On success, Set() will have drained waiters for us.
+		void Reset()
+		{
+			UBA_ASSERT(m_sem);
 
-				if (!res) // Timed out or error: remove our registration.
-					m_waiters.fetch_sub(1, std::memory_order_acq_rel); 
+			UBA_ASSERT(!m_oldestWaiter);
 
-				return res;
+			// Drain potential releases
+			#if PLATFORM_WINDOWS
+			while (WaitForSingleObject(m_sem, 0) == WAIT_OBJECT_0)
+			{
+			}
+			#elif PLATFORM_LINUX
+			while (sem_trywait(m_sem) == 0)
+			{
 			}
+			#else
+			UBA_ASSERT(false);
+			#endif
+
+			m_triggered = TriggerType_None;
 		}
 
-		// Set the event:
-		//  - AutoReset: deliver exactly one credit (wake one waiter or latch for next).
-		//  - ManualReset: latch open and wake all currently registered waiters exactly once.
-		void Set()
+		bool IsSet(u64 timeoutMs = ~0u)
 		{
-			if (m_resetType == EventResetType_Auto)
+			UBA_ASSERT(m_sem);
+
+			SCOPED_FUTEX(m_lock, lock);
+			if (m_triggered == TriggerType_All)
+				return true;
+
+			if (timeoutMs == 0 && m_triggered == TriggerType_None)
+				return false;
+
+			Waiter waiter;
+			if (m_newestWaiter)
 			{
-				// First, try to create a pending credit if none exists.
-				int expected = 0;
-				if (m_autoCredit.compare_exchange_strong(expected, 1, std::memory_order_acq_rel, std::memory_order_acquire))
-				{
-					if (m_waiters.load(std::memory_order_acquire) > 0) // If waiters exist, transfer the credit to exactly one waiter now.
-					{
-						m_autoCredit.store(0, std::memory_order_release); // Move credit to the queue; ensure no double-credit remains.
-						dispatch_semaphore_signal(m_gate);
-					}
-				}
-				// If a credit already existed, collapse extra Set() calls (no-op).
+				m_newestWaiter->next = &waiter;
+				waiter.prev = m_newestWaiter;
+			}
+			else if (m_triggered == TriggerType_One) // If we are the first waiter and one is set, we take that trigger
+			{
+				m_triggered = TriggerType_None;
+				return true;
 			}
 			else
+				m_oldestWaiter = &waiter;
+
+			m_newestWaiter = &waiter;
+
+			// Unlock and wait for semaphore
+			lock.Leave();
+
+		#if PLATFORM_WINDOWS
+			DWORD timeout = INFINITE;
+			if (timeoutMs != ~0u)
+				timeout = DWORD(timeoutMs);
+			
+			DWORD r = WaitForSingleObject(m_sem, timeout);
+			bool timedOut = r != WAIT_OBJECT_0;(void)timedOut;
+		#elif PLATFORM_LINUX
+			u64 timeoutNs = u64(timeoutMs) * 1'000'000;
+			struct timespec now;
+			clock_gettime(CLOCK_REALTIME, &now);
+			u64 absTimeNano = ToNanoSeconds(now) + timeoutNs;
+			auto absTime = ToTimeSpecFromNano<timespec>(absTimeNano);
+			bool timedOut = false;
+			while (sem_timedwait(m_sem, &absTime) == -1)
 			{
-				// ManualReset: latch and wake all currently registered waiters.
-				int was = m_setLatch.exchange(1, std::memory_order_acq_rel);
-
-				if (was == 1) // already set
-					return;
-
-				// Snapshot and drain current waiters exactly once.
-				int n = m_waiters.exchange(0, std::memory_order_acq_rel);
-				for (int i = 0; i < n; ++i)
-					dispatch_semaphore_signal(m_gate);
+				if (errno == EINTR)
+					continue;
+				UBA_ASSERTF(errno == ETIMEDOUT, "sem_wait failed");
+				timedOut = true;
+				break;
 			}
-		}
+		#else
+			UBA_ASSERT(false);
+		#endif
 
-		void Reset()
-		{
-			UBA_ASSERT(m_waiters.load() == 0);
-			m_setLatch.store(0, std::memory_order_release);
-			//m_autoCredit.store(0, std::memory_order_release); // union
-		}
+			// Take lock and disconnect waiter
+			lock.Enter();
 
-		EventResetType GetResetType()
-		{
-			return m_resetType;
-		}
+			if (waiter.prev)
+				waiter.prev->next = waiter.next;
+			else
+				m_oldestWaiter = waiter.next;
 
-		EventDispatchSemaphore* next = nullptr;
-	
-	private:
+			if (waiter.next)
+				waiter.next->prev = waiter.prev;
+			else
+				m_newestWaiter = waiter.prev;
 
-		static bool WaitDispatch(dispatch_semaphore_t sem, u64 timeoutMs)
-		{
-			if (timeoutMs == ~0u)
-				timeoutMs = 24llu * 60 * 60 * 1000; // 24 hours (FOREVER can cause sigill)
-			u64 timeoutNs = timeoutMs * 1'000'000ull;
-			dispatch_time_t t = dispatch_time(DISPATCH_TIME_NOW, timeoutNs);
-			return dispatch_semaphore_wait(sem, t) == 0;
+			return !timedOut || m_triggered == TriggerType_All;
 		}
 
-		EventResetType m_resetType;
-		dispatch_semaphore_t m_gate = 0;
-		Atomic<int> m_waiters;    // >= 0
-		union
+		struct Waiter
 		{
-			Atomic<int> m_setLatch;   // 0/1
-			Atomic<int> m_autoCredit; // 0/1
+			Waiter* next = nullptr;
+			Waiter* prev = nullptr;
 		};
+		Waiter* m_oldestWaiter = nullptr;
+		Waiter* m_newestWaiter = nullptr;
+
+		enum TriggerType : u8 { TriggerType_None, TriggerType_One, TriggerType_All };
+
+		TriggerType m_triggered = TriggerType_None;
+
+		Futex m_lock;
+
+#if PLATFORM_WINDOWS
+		HANDLE m_sem = 0;
+#elif PLATFORM_LINUX
+		sem_t* m_sem;
+#else
+		dispatch_semaphore_t m_sem = 0;
+#endif
 	};
 }
 
-#endif
+//#endif
diff --git a/Engine/Source/Programs/UnrealBuildAccelerator/Core/Private/UbaEventMachSemaphore.inl b/Engine/Source/Programs/UnrealBuildAccelerator/Core/Private/UbaEventMachSemaphore.inl
--- a/Engine/Source/Programs/UnrealBuildAccelerator/Core/Private/UbaEventMachSemaphore.inl
+++ b/Engine/Source/Programs/UnrealBuildAccelerator/Core/Private/UbaEventMachSemaphore.inl
@@ -8,9 +8,8 @@
 
 namespace uba
 {
-	class EventMachSemaphore
+	struct EventMachSemaphore : EventHeader
 	{
-	public:
 		~EventMachSemaphore()
 		{
 			Destroy();
@@ -65,8 +64,6 @@ namespace uba
 			out.AppendValue(m_sem);
 		}
 
-	private:
-		u8 m_dummy; // This is to map into SharedEvent under the hood
 		bool m_isOwner = false;
 		semaphore_t m_sem = MACH_PORT_NULL;
 	};
diff --git a/Engine/Source/Programs/UnrealBuildAccelerator/Core/Private/UbaEventPThread.inl b/Engine/Source/Programs/UnrealBuildAccelerator/Core/Private/UbaEventPThread.inl
--- a/Engine/Source/Programs/UnrealBuildAccelerator/Core/Private/UbaEventPThread.inl
+++ b/Engine/Source/Programs/UnrealBuildAccelerator/Core/Private/UbaEventPThread.inl
@@ -21,7 +21,7 @@ namespace uba
 		return ts;
 	}
 
-	struct EventPThread
+	struct EventPThread : EventHeader
 	{
 		EventPThread() = default;
 		~EventPThread() { Destroy(); }
@@ -33,7 +33,7 @@ namespace uba
 			static_assert(Atomic<bool>::is_always_lock_free, "atomic<T> not lock free, can't work in shared mem");
 
 			UBA_ASSERTF(!m_initialized, "Can't create already created Event");
-			m_resetType = rt;
+			m_resetType.value = rt;
 
 			pthread_mutexattr_t attrmutex;
 			if (pthread_mutexattr_init(&attrmutex) != 0)
@@ -107,7 +107,7 @@ namespace uba
 				return;
 
 			LockEventMutex();
-			m_resetType = EventResetType_Manual;
+			m_resetType.value = EventResetType_Manual;
 			UnlockEventMutex();
 			Set();
 
@@ -130,7 +130,7 @@ namespace uba
 				return;
 			LockEventMutex();
 
-			if (m_resetType == EventResetType_Manual)
+			if (m_resetType.value == EventResetType_Manual)
 			{
 				m_triggered = TriggerType_All;
 				if (pthread_cond_broadcast(&m_condition) != 0)
@@ -247,16 +247,10 @@ namespace uba
 			return false;
 		}
 
-		EventResetType GetResetType() { return m_resetType; }
-
-		// Must be first!
-		EventPThread* next = nullptr;
-
 	private:
 		enum TriggerType : u8 { TriggerType_None, TriggerType_One, TriggerType_All };
 
 		Atomic<bool> m_initialized;
-		Atomic<EventResetType> m_resetType;
 		Atomic<TriggerType> m_triggered;
 		Atomic<u32> m_waitingThreads;
 		pthread_mutex_t m_mutex;
diff --git a/Engine/Source/Programs/UnrealBuildAccelerator/Core/Private/UbaEventPosixSemaphore.inl b/Engine/Source/Programs/UnrealBuildAccelerator/Core/Private/UbaEventPosixSemaphore.inl
--- a/Engine/Source/Programs/UnrealBuildAccelerator/Core/Private/UbaEventPosixSemaphore.inl
+++ b/Engine/Source/Programs/UnrealBuildAccelerator/Core/Private/UbaEventPosixSemaphore.inl
@@ -2,15 +2,13 @@
 
 #if !PLATFORM_WINDOWS
 
-#if !PLATFORM_WINDOWS
 #include <sys/time.h>
 #include <new>
 #include <semaphore.h>
-#endif
 
 namespace uba
 {
-	class EventPosixSemaphore
+	struct EventPosixSemaphore : EventHeader
 	{
 	public:
 		~EventPosixSemaphore()
@@ -148,7 +146,6 @@ namespace uba
 		}
 
 	private:
-		u8 m_dummy; // This is to map into SharedEvent under the hood
 		bool m_isOwner = false;
 		sem_t* m_snd = nullptr;
 		sem_t* m_rcv = nullptr;
diff --git a/Engine/Source/Programs/UnrealBuildAccelerator/Core/Private/UbaEventWin32.inl b/Engine/Source/Programs/UnrealBuildAccelerator/Core/Private/UbaEventWin32.inl
--- a/Engine/Source/Programs/UnrealBuildAccelerator/Core/Private/UbaEventWin32.inl
+++ b/Engine/Source/Programs/UnrealBuildAccelerator/Core/Private/UbaEventWin32.inl
@@ -4,40 +4,47 @@
 
 namespace uba
 {
-	struct EventWin32
+	struct EventWin32 : EventHeader
 	{
-		HANDLE handle;
-		union
+		HANDLE m_handle;
+
+		bool Create(EventResetType rt, bool = true)
 		{
-			EventWin32* next;
-			EventResetType resetType;
-		};
+			m_resetType.value = rt;
+			m_handle = CreateEvent(nullptr, rt == EventResetType_Manual, false, NULL);
+			return m_handle != 0;
+		}
 
-		bool Create(EventResetType rt)
+		void Destroy()
 		{
-			resetType = rt;
-			handle = CreateEvent(nullptr, rt == EventResetType_Manual, false, NULL);
-			return handle != 0;
+			CloseHandle(m_handle);
+			m_handle = 0;
 		}
 
 		void Set()
 		{
-			SetEvent(handle);
+			SetEvent(m_handle);
 		}
 
 		void Reset()
 		{
-			ResetEvent(handle);
+			ResetEvent(m_handle);
 		}
 
 		bool IsSet(u32 timeOutMs)
 		{
-			return WaitForSingleObject(handle, timeOutMs) == 0;
+			return WaitForSingleObject(m_handle, timeOutMs) == 0;
+		}
+
+		bool Create(StringView name, bool isOwner)
+		{
+			UBA_ASSERT(false);
+			return false;
 		}
 
-		EventResetType GetResetType()
+		void ToString(StringBufferBase& out)
 		{
-			return resetType;
+			UBA_ASSERT(false);
 		}
 	};
 
diff --git a/Engine/Source/Programs/UnrealBuildAccelerator/Core/Public/UbaEvent.h b/Engine/Source/Programs/UnrealBuildAccelerator/Core/Public/UbaEvent.h
--- a/Engine/Source/Programs/UnrealBuildAccelerator/Core/Public/UbaEvent.h
+++ b/Engine/Source/Programs/UnrealBuildAccelerator/Core/Public/UbaEvent.h
@@ -28,13 +28,14 @@ namespace uba
 		bool IsCreated();
 		bool IsSet(u32 timeOutMs = ~0u);
 		void* GetHandle();
+		static bool IsNative();
+
 	private:
 		void* m_impl = nullptr;
 		Event(const Event&) = delete;
 		void operator=(const Event&) = delete;
 	};
 
-#if !PLATFORM_WINDOWS
 	class SharedEvent
 	{
 	public:
@@ -46,15 +47,17 @@ namespace uba
 		void Set();
 		bool IsSet(u32 timeOutMs = ~0u);
 		void ToString(StringBufferBase& out);
+		void* GetHandle();
 	private:
 		bool IsSemaphore();
 		#if PLATFORM_MAC || defined(__aarch64__)
 		u64 m_data[16];
-		#else
+		#elif !PLATFORM_WINDOWS
 		u64 m_data[13];
+		#else
+		u64 m_data[2];
 		#endif
 		SharedEvent(const SharedEvent&) = delete;
 		void operator=(const SharedEvent&) = delete;
 	};
-#endif
 }
diff --git a/Engine/Source/Programs/UnrealBuildAccelerator/Test/Private/UbaTestBasics.cpp b/Engine/Source/Programs/UnrealBuildAccelerator/Test/Private/UbaTestBasics.cpp
--- a/Engine/Source/Programs/UnrealBuildAccelerator/Test/Private/UbaTestBasics.cpp
+++ b/Engine/Source/Programs/UnrealBuildAccelerator/Test/Private/UbaTestBasics.cpp
@@ -358,6 +358,9 @@ namespace uba
 
 	bool TestOverlappedIO(LoggerWithWriter& logger, const StringBufferBase& rootDir)
 	{
+		if (!Event::IsNative())
+			return true;
+
 		StringBuffer<> testFileName(rootDir);
 		testFileName.Append(TCV("UbaTestFile"));
 
