[UBA]
* Fixed dispatcher semaphore event for macos

diff --git a/Engine/Source/Programs/UnrealBuildAccelerator/Core/Private/UbaEventDispatchSemaphore.inl b/Engine/Source/Programs/UnrealBuildAccelerator/Core/Private/UbaEventDispatchSemaphore.inl
--- a/Engine/Source/Programs/UnrealBuildAccelerator/Core/Private/UbaEventDispatchSemaphore.inl
+++ b/Engine/Source/Programs/UnrealBuildAccelerator/Core/Private/UbaEventDispatchSemaphore.inl
@@ -108,15 +108,13 @@ namespace uba
 			// Drain potential releases
 			#if PLATFORM_WINDOWS
 			while (WaitForSingleObject(m_sem, 0) == WAIT_OBJECT_0)
-			{
-			}
 			#elif PLATFORM_LINUX
 			while (sem_trywait(m_sem) == 0)
-			{
-			}
 			#else
-			UBA_ASSERT(false);
+			while (dispatch_semaphore_wait(m_sem, DISPATCH_TIME_NOW) == 0)
 			#endif
+			{				
+			}
 
 			m_triggered = TriggerType_None;
 		}
@@ -159,6 +157,8 @@ namespace uba
 			DWORD r = WaitForSingleObject(m_sem, timeout);
 			bool timedOut = r != WAIT_OBJECT_0;(void)timedOut;
 		#elif PLATFORM_LINUX
+			if (timeoutMs == ~0u)
+				timeoutMs = 24llu*60*60*1000; // 24hrs
 			u64 timeoutNs = u64(timeoutMs) * 1'000'000;
 			struct timespec now;
 			clock_gettime(CLOCK_REALTIME, &now);
@@ -174,7 +174,10 @@ namespace uba
 				break;
 			}
 		#else
-			UBA_ASSERT(false);
+			if (timeoutMs == ~0u)
+				timeoutMs = 24llu*60*60*1000; // 24hrs
+			dispatch_time_t t = dispatch_time(DISPATCH_TIME_NOW, (int64_t)timeoutMs * 1000000ll);
+			bool timedOut = dispatch_semaphore_wait(m_sem, t) != 0;
 		#endif
 
 			// Take lock and disconnect waiter
