[UBA]
* Cleanups of new event code

diff --git a/Engine/Source/Programs/UnrealBuildAccelerator/Common/Private/UbaEnvironment.cpp b/Engine/Source/Programs/UnrealBuildAccelerator/Common/Private/UbaEnvironment.cpp
--- a/Engine/Source/Programs/UnrealBuildAccelerator/Common/Private/UbaEnvironment.cpp
+++ b/Engine/Source/Programs/UnrealBuildAccelerator/Common/Private/UbaEnvironment.cpp
@@ -130,7 +130,7 @@ namespace uba
 	void ElevateCurrentThreadPriority()
 	{
 		#if PLATFORM_WINDOWS
-		SetThreadPriority(GetCurrentThread(), THREAD_PRIORITY_HIGHEST);
+		//SetThreadPriority(GetCurrentThread(), THREAD_PRIORITY_HIGHEST);
 		#endif
 	}
 
diff --git a/Engine/Source/Programs/UnrealBuildAccelerator/Common/Private/UbaFileAccessor.cpp b/Engine/Source/Programs/UnrealBuildAccelerator/Common/Private/UbaFileAccessor.cpp
--- a/Engine/Source/Programs/UnrealBuildAccelerator/Common/Private/UbaFileAccessor.cpp
+++ b/Engine/Source/Programs/UnrealBuildAccelerator/Common/Private/UbaFileAccessor.cpp
@@ -65,7 +65,7 @@ namespace uba
 	bool FileAccessor::InternalCreateWrite(bool allowRead, u32 flagsAndAttributes, u64 fileSize, const tchar* tempPath, bool isMemoryMap)
 	{
 		#if PLATFORM_WINDOWS
-		UBA_ASSERT((flagsAndAttributes & FILE_FLAG_OVERLAPPED) == 0 || Event::IsNative());
+		UBA_ASSERT((flagsAndAttributes & FILE_FLAG_OVERLAPPED) == 0 || EventIsNative);
 		#endif
 
 
diff --git a/Engine/Source/Programs/UnrealBuildAccelerator/Common/Private/UbaNetworkBackendTcp.cpp b/Engine/Source/Programs/UnrealBuildAccelerator/Common/Private/UbaNetworkBackendTcp.cpp
--- a/Engine/Source/Programs/UnrealBuildAccelerator/Common/Private/UbaNetworkBackendTcp.cpp
+++ b/Engine/Source/Programs/UnrealBuildAccelerator/Common/Private/UbaNetworkBackendTcp.cpp
@@ -2,6 +2,7 @@
 
 #include "UbaNetworkBackendTcp.h"
 #include "UbaConfig.h"
+#include "UbaDefinitions.h"
 #include "UbaEnvironment.h"
 #include "UbaEvent.h"
 #include "UbaHash.h"
@@ -255,7 +256,7 @@ namespace uba
 		m_disableNagle = info.disableNagle;
 
 		#if PLATFORM_WINDOWS
-		m_useOverlappedSend = info.useOverlappedSend && Event::IsNative();
+		m_useOverlappedSend = info.useOverlappedSend && EventIsNative;
 		m_iocpWorkerCount = u16(info.iocpWorkerCount);
 		#endif
 
diff --git a/Engine/Source/Programs/UnrealBuildAccelerator/Common/Private/UbaStorageClient.cpp b/Engine/Source/Programs/UnrealBuildAccelerator/Common/Private/UbaStorageClient.cpp
--- a/Engine/Source/Programs/UnrealBuildAccelerator/Common/Private/UbaStorageClient.cpp
+++ b/Engine/Source/Programs/UnrealBuildAccelerator/Common/Private/UbaStorageClient.cpp
@@ -2,6 +2,7 @@
 
 #include "UbaStorageClient.h"
 #include "UbaConfig.h"
+#include "UbaDefinitions.h"
 #include "UbaDirectoryIterator.h"
 #include "UbaFileAccessor.h"
 #include "UbaNetworkClient.h"
@@ -583,7 +584,7 @@ namespace uba
 					else
 					{
 						u32 extraFlags = DefaultAttributes();
-						bool useOverlap = !IsRunningWine() && isCompressed == IsCompressed(casKey) && sizeOnDisk > 1024 * 1024 && Event::IsNative();
+						bool useOverlap = !IsRunningWine() && isCompressed == IsCompressed(casKey) && sizeOnDisk > 1024 * 1024 && EventIsNative;
 						if (useOverlap)
 							extraFlags |= FILE_FLAG_OVERLAPPED;
 						if (!destinationFile.CreateWrite(false, extraFlags, sizeOnDisk, m_tempPath.data))
diff --git a/Engine/Source/Programs/UnrealBuildAccelerator/Common/Private/UbaThread.cpp b/Engine/Source/Programs/UnrealBuildAccelerator/Common/Private/UbaThread.cpp
--- a/Engine/Source/Programs/UnrealBuildAccelerator/Common/Private/UbaThread.cpp
+++ b/Engine/Source/Programs/UnrealBuildAccelerator/Common/Private/UbaThread.cpp
@@ -171,7 +171,7 @@ namespace uba
 			};
 
 #if PLATFORM_WINDOWS // Optimization, not needed in initial implementation
-		if (wakeupEvent && Event::IsNative())
+		if (wakeupEvent && EventIsNative)
 		{
 			HANDLE h[] = { m_handle, wakeupEvent->GetHandle() };
 			DWORD res = WaitForMultipleObjects(2, h, false, milliseconds);
diff --git a/Engine/Source/Programs/UnrealBuildAccelerator/Core/Private/UbaEvent.cpp b/Engine/Source/Programs/UnrealBuildAccelerator/Core/Private/UbaEvent.cpp
--- a/Engine/Source/Programs/UnrealBuildAccelerator/Core/Private/UbaEvent.cpp
+++ b/Engine/Source/Programs/UnrealBuildAccelerator/Core/Private/UbaEvent.cpp
@@ -7,13 +7,12 @@
 #include "UbaSynchronization.h"
 
 #define UBA_TEST_WAIT_QUALITY 0
-#define UBA_TEST_WINDOWS_SEMAPHORES 0
 
 namespace uba
 {
 	struct ResetTypeWithPadding
 	{
-		u8 padding; // This is to map into SharedEvent under the hood
+		u8 padding;
 		EventResetType value;
 	};
 
@@ -37,33 +36,45 @@ namespace uba
 
 namespace uba
 {
-#if UBA_USE_NATIVE_MAC_SEMAPHORES
-	bool Event::IsNative() { return false; }
-	using EventImpl = EventDispatchSemaphore;
-	using SharedEventSemaphoreImpl = EventMachSemaphore;
-	using SharedEventImpl = EventPThread;
-#elif !PLATFORM_WINDOWS
-	bool Event::IsNative() { return false; }
-	using EventImpl = EventPThread;
-	using SharedEventSemaphoreImpl = EventPosixSemaphore;
-	using SharedEventImpl = EventPThread;
-#else
-	using SharedEventSemaphoreImpl = EventWin32;
-	using SharedEventImpl = EventWin32;
-	#if UBA_TEST_WINDOWS_SEMAPHORES
-	using EventImpl = EventDispatchSemaphore;
-	HANDLE GetNativeHandle(EventImpl& e) { return 0; }
-	bool Event::IsNative() { return false; }
+	////////////////////////////////////////////////////////////////////////////////////////////////////
+
+	// Local event type
+	using SharedEventDefaultImpl = 
+	#if PLATFORM_WINDOWS
+		EventWin32;
+	#else
+		EventPThread;
+	#endif
+
+	// Semaphore based IPC event.
+	using SharedEventSemaphoreImpl = 
+
+	#if UBA_USE_NATIVE_MAC_SEMAPHORES
+		EventMachSemaphore;
+	#elif !PLATFORM_WINDOWS
+		EventPosixSemaphore;
 	#else
-	using EventImpl = EventWin32;
-	HANDLE GetNativeHandle(EventImpl& e) { return e.m_handle; }
-	bool Event::IsNative() { return true; }
+		EventWin32;
+	#endif
+
+	// Default IPC event
+	using EventImpl = 
+
+	#if UBA_USE_NATIVE_MAC_SEMAPHORES
+		EventDispatchSemaphore;
+	#elif !PLATFORM_WINDOWS
+		EventPThread;
+	#elif UBA_TEST_WINDOWS_SEMAPHORES
+		EventDispatchSemaphore;
+	#else
+		EventWin32;
 	#endif
-#endif
 
 	Futex g_firstEventLock[2];
 	EventHeader* g_firstEvent[2];
 
+	////////////////////////////////////////////////////////////////////////////////////////////////////
+
 	Event::Event(EventResetType resetType)
 	{
 		Create(resetType);
@@ -143,18 +154,20 @@ namespace uba
 		#if PLATFORM_WINDOWS
 		if (!m_impl)
 			return nullptr;
-		return GetNativeHandle(*(EventImpl*)m_impl);
+		return ((EventImpl*)m_impl)->GetHandle();
 		#else
 		UBA_ASSERTF(false, "Event::GetHandle not available");
 		return nullptr;
 		#endif
 	}
 
+	////////////////////////////////////////////////////////////////////////////////////////////////////
 
 	SharedEvent::SharedEvent()
 	{
+		new (m_data) EventHeader();
 		*((u8*)m_data) = 0;
-		static_assert(sizeof(m_data) >= sizeof(SharedEventImpl));
+		static_assert(sizeof(m_data) >= sizeof(SharedEventDefaultImpl));
 		static_assert(sizeof(m_data) >= sizeof(SharedEventSemaphoreImpl));
 	}
 
@@ -166,10 +179,10 @@ namespace uba
 	bool SharedEvent::Create(EventResetType resetType)//, bool shared)
 	{
 		UBA_ASSERT(*((u8*)m_data) == 0);
-		new (m_data) SharedEventImpl();
+		new (m_data) SharedEventDefaultImpl();
 		UBA_ASSERT(*((u8*)m_data) == 0);
 		*((u8*)m_data) = 1;
-		if (!((SharedEventImpl&)m_data).Create(resetType, true))
+		if (!((SharedEventDefaultImpl&)m_data).Create(resetType, true))
 			return false;
 		UBA_ASSERT(*((u8*)m_data) == 1);
 		return true;
@@ -195,8 +208,8 @@ namespace uba
 		}
 		else
 		{
-			((SharedEventImpl&)m_data).Destroy();
-			((SharedEventImpl&)m_data).~SharedEventImpl();
+			((SharedEventDefaultImpl&)m_data).Destroy();
+			((SharedEventDefaultImpl&)m_data).~SharedEventDefaultImpl();
 		}
 		*((u8*)m_data) = 0;
 	}
@@ -209,7 +222,7 @@ namespace uba
 		if (IsSemaphore())
 			((SharedEventSemaphoreImpl&)m_data).Set();
 		else
-			((SharedEventImpl&)m_data).Set();
+			((SharedEventDefaultImpl&)m_data).Set();
 	}
 
 	bool SharedEvent::IsSet(u32 timeOutMs)
@@ -220,7 +233,7 @@ namespace uba
 		if (IsSemaphore())
 			return ((SharedEventSemaphoreImpl&)m_data).IsSet(timeOutMs);
 		else
-			return ((SharedEventImpl&)m_data).IsSet(timeOutMs);
+			return ((SharedEventDefaultImpl&)m_data).IsSet(timeOutMs);
 	}
 
 	bool SharedEvent::IsSemaphore()
@@ -238,7 +251,7 @@ namespace uba
 	{
 		UBA_ASSERT(!IsSemaphore());
 		#if PLATFORM_WINDOWS
-		return ((SharedEventImpl*)m_data)->m_handle;
+		return ((SharedEventDefaultImpl*)m_data)->GetHandle();
 		#else
 		UBA_ASSERTF(false, "Event::GetHandle not available");
 		return nullptr;
diff --git a/Engine/Source/Programs/UnrealBuildAccelerator/Core/Private/UbaEventDispatchSemaphore.inl b/Engine/Source/Programs/UnrealBuildAccelerator/Core/Private/UbaEventDispatchSemaphore.inl
--- a/Engine/Source/Programs/UnrealBuildAccelerator/Core/Private/UbaEventDispatchSemaphore.inl
+++ b/Engine/Source/Programs/UnrealBuildAccelerator/Core/Private/UbaEventDispatchSemaphore.inl
@@ -68,34 +68,14 @@ namespace uba
 
 				// Release semaphore for all waiters... also fine releasing semaphore more than needed since it is a manual reset.
 				for (Waiter* it=m_oldestWaiter; it; it = it->next)
-				{
-					#if PLATFORM_WINDOWS
-					LONG prev = 0;
-					ReleaseSemaphore(m_sem, 1, &prev);
-					#elif PLATFORM_LINUX
-					sem_post(m_sem);
-					#else
-					dispatch_semaphore_signal(m_sem);
-					#endif
-				}
+					Signal();
 			}
 			else
 			{
 				if (m_oldestWaiter)
-				{
-					#if PLATFORM_WINDOWS
-					LONG prev = 0;
-					ReleaseSemaphore(m_sem, 1, &prev);
-					#elif PLATFORM_LINUX
-					sem_post(m_sem);
-					#else
-					dispatch_semaphore_signal(m_sem);
-					#endif
-				}
+					Signal();
 				else
-				{
 					m_triggered = TriggerType_One;
-				}
 			}
 		}
 
@@ -105,16 +85,7 @@ namespace uba
 
 			UBA_ASSERT(!m_oldestWaiter);
 
-			// Drain potential releases
-			#if PLATFORM_WINDOWS
-			while (WaitForSingleObject(m_sem, 0) == WAIT_OBJECT_0)
-			#elif PLATFORM_LINUX
-			while (sem_trywait(m_sem) == 0)
-			#else
-			while (dispatch_semaphore_wait(m_sem, DISPATCH_TIME_NOW) == 0)
-			#endif
-			{				
-			}
+			Drain();
 
 			m_triggered = TriggerType_None;
 		}
@@ -149,14 +120,46 @@ namespace uba
 			// Unlock and wait for semaphore
 			lock.Leave();
 
-		#if PLATFORM_WINDOWS
+			bool success = Wait(timeoutMs);
+
+			// Take lock and disconnect waiter
+			lock.Enter();
+
+			if (waiter.prev)
+				waiter.prev->next = waiter.next;
+			else
+				m_oldestWaiter = waiter.next;
+
+			if (waiter.next)
+				waiter.next->prev = waiter.prev;
+			else
+				m_newestWaiter = waiter.prev;
+
+			return success || m_triggered == TriggerType_All;
+		}
+
+		void Signal()
+		{
+			#if PLATFORM_WINDOWS
+			LONG prev = 0;
+			ReleaseSemaphore(m_sem, 1, &prev);
+			#elif PLATFORM_LINUX
+			sem_post(m_sem);
+			#else
+			dispatch_semaphore_signal(m_sem);
+			#endif
+		}
+
+		bool Wait(u64 timeoutMs)
+		{
+			#if PLATFORM_WINDOWS
 			DWORD timeout = INFINITE;
 			if (timeoutMs != ~0u)
 				timeout = DWORD(timeoutMs);
 			
 			DWORD r = WaitForSingleObject(m_sem, timeout);
-			bool timedOut = r != WAIT_OBJECT_0;(void)timedOut;
-		#elif PLATFORM_LINUX
+			return r == WAIT_OBJECT_0;
+			#elif PLATFORM_LINUX
 			if (timeoutMs == ~0u)
 				timeoutMs = 24llu*60*60*1000; // 24hrs
 			u64 timeoutNs = u64(timeoutMs) * 1'000'000;
@@ -164,36 +167,38 @@ namespace uba
 			clock_gettime(CLOCK_REALTIME, &now);
 			u64 absTimeNano = ToNanoSeconds(now) + timeoutNs;
 			auto absTime = ToTimeSpecFromNano<timespec>(absTimeNano);
-			bool timedOut = false;
 			while (sem_timedwait(m_sem, &absTime) == -1)
 			{
 				if (errno == EINTR)
 					continue;
 				UBA_ASSERTF(errno == ETIMEDOUT, "sem_wait failed");
-				timedOut = true;
-				break;
+				return false;
 			}
-		#else
+			return true;
+			#else
 			if (timeoutMs == ~0u)
 				timeoutMs = 24llu*60*60*1000; // 24hrs
 			dispatch_time_t t = dispatch_time(DISPATCH_TIME_NOW, (int64_t)timeoutMs * 1000000ll);
-			bool timedOut = dispatch_semaphore_wait(m_sem, t) != 0;
-		#endif
-
-			// Take lock and disconnect waiter
-			lock.Enter();
-
-			if (waiter.prev)
-				waiter.prev->next = waiter.next;
-			else
-				m_oldestWaiter = waiter.next;
+			return dispatch_semaphore_wait(m_sem, t) == 0;
+			#endif
+		}
 
-			if (waiter.next)
-				waiter.next->prev = waiter.prev;
-			else
-				m_newestWaiter = waiter.prev;
+		void Drain()
+		{
+			#if PLATFORM_WINDOWS
+			while (WaitForSingleObject(m_sem, 0) == WAIT_OBJECT_0)
+			#elif PLATFORM_LINUX
+			while (sem_trywait(m_sem) == 0)
+			#else
+			while (dispatch_semaphore_wait(m_sem, DISPATCH_TIME_NOW) == 0)
+			#endif
+			{				
+			}
+		}
 
-			return !timedOut || m_triggered == TriggerType_All;
+		void* GetHandle()
+		{
+			return nullptr;
 		}
 
 		struct Waiter
diff --git a/Engine/Source/Programs/UnrealBuildAccelerator/Core/Private/UbaEventWin32.inl b/Engine/Source/Programs/UnrealBuildAccelerator/Core/Private/UbaEventWin32.inl
--- a/Engine/Source/Programs/UnrealBuildAccelerator/Core/Private/UbaEventWin32.inl
+++ b/Engine/Source/Programs/UnrealBuildAccelerator/Core/Private/UbaEventWin32.inl
@@ -46,6 +46,11 @@ namespace uba
 		{
 			UBA_ASSERT(false);
 		}
+
+		void* GetHandle()
+		{
+			return m_handle;
+		}
 	};
 
 }
diff --git a/Engine/Source/Programs/UnrealBuildAccelerator/Core/Public/UbaDefinitions.h b/Engine/Source/Programs/UnrealBuildAccelerator/Core/Public/UbaDefinitions.h
--- a/Engine/Source/Programs/UnrealBuildAccelerator/Core/Public/UbaDefinitions.h
+++ b/Engine/Source/Programs/UnrealBuildAccelerator/Core/Public/UbaDefinitions.h
@@ -10,4 +10,14 @@
 
 // Disable pthread cancellation to reduce work when waking up from event.
 // No uba code should rely on cancelling threads and shutdown is gracefully so this can be disabled
-#define UBA_DISABLE_PTHREAD_CANCELLATIONS 1
\ No newline at end of file
+#define UBA_DISABLE_PTHREAD_CANCELLATIONS 1
+
+
+// Enable this to test semaphore implementation for local events
+#define UBA_TEST_WINDOWS_SEMAPHORES 0
+
+
+namespace uba
+{
+	inline constexpr bool EventIsNative = bool(PLATFORM_WINDOWS && !UBA_TEST_WINDOWS_SEMAPHORES);
+}
\ No newline at end of file
diff --git a/Engine/Source/Programs/UnrealBuildAccelerator/Core/Public/UbaEvent.h b/Engine/Source/Programs/UnrealBuildAccelerator/Core/Public/UbaEvent.h
--- a/Engine/Source/Programs/UnrealBuildAccelerator/Core/Public/UbaEvent.h
+++ b/Engine/Source/Programs/UnrealBuildAccelerator/Core/Public/UbaEvent.h
@@ -6,15 +6,21 @@
 
 namespace uba
 {
+	////////////////////////////////////////////////////////////////////////////////////////////////////
+
 	class StringBufferBase;
 	struct StringView;
 
+	////////////////////////////////////////////////////////////////////////////////////////////////////
+
 	enum EventResetType : u8
 	{
 		EventResetType_Auto,
 		EventResetType_Manual,
 	};
 
+	////////////////////////////////////////////////////////////////////////////////////////////////////
+
 	class Event
 	{
 	public:
@@ -36,6 +42,8 @@ namespace uba
 		void operator=(const Event&) = delete;
 	};
 
+	////////////////////////////////////////////////////////////////////////////////////////////////////
+
 	class SharedEvent
 	{
 	public:
@@ -60,4 +68,6 @@ namespace uba
 		SharedEvent(const SharedEvent&) = delete;
 		void operator=(const SharedEvent&) = delete;
 	};
+
+	////////////////////////////////////////////////////////////////////////////////////////////////////
 }
diff --git a/Engine/Source/Programs/UnrealBuildAccelerator/Test/Private/UbaTestBasics.cpp b/Engine/Source/Programs/UnrealBuildAccelerator/Test/Private/UbaTestBasics.cpp
--- a/Engine/Source/Programs/UnrealBuildAccelerator/Test/Private/UbaTestBasics.cpp
+++ b/Engine/Source/Programs/UnrealBuildAccelerator/Test/Private/UbaTestBasics.cpp
@@ -2,6 +2,7 @@
 
 #include "UbaBinaryParser.h"
 #include "UbaBinaryReaderWriter.h"
+#include "UbaDefinitions.h"
 #include "UbaFileAccessor.h"
 #include "UbaFileMapping.h"
 #include "UbaPathUtils.h"
@@ -358,7 +359,7 @@ namespace uba
 
 	bool TestOverlappedIO(LoggerWithWriter& logger, const StringBufferBase& rootDir)
 	{
-		if (!Event::IsNative())
+		if (!EventIsNative)
 			return true;
 
 		StringBuffer<> testFileName(rootDir);
