[UBA]
* Cleanups of new event code

diff --git a/Engine/Source/Programs/UnrealBuildAccelerator/Common/Private/UbaSession.cpp b/Engine/Source/Programs/UnrealBuildAccelerator/Common/Private/UbaSession.cpp
--- a/Engine/Source/Programs/UnrealBuildAccelerator/Common/Private/UbaSession.cpp
+++ b/Engine/Source/Programs/UnrealBuildAccelerator/Common/Private/UbaSession.cpp
@@ -489,7 +489,7 @@ namespace uba
 		// On one machine we get twice as fast without overlap if no bottleneck. On another machine (ntfs compression on) we get twice as slow without overlap
 		// Both machines behaves well with overlap AND bottleneck. Both machine are 128 logical core thread rippers.
 		bool useFileMapForWrite = fileSize && fileSize <= m_writeFilesFileMapMax; // ::CreateFileMappingW does not work for zero-length files
-		bool useOverlap = !useFileMapForWrite && fileSize >= m_writeFilesNoBufferingMin && Event::IsNative();//fileSize > 2*1024*1024;
+		bool useOverlap = !useFileMapForWrite && fileSize >= m_writeFilesNoBufferingMin && EventIsNative;//fileSize > 2*1024*1024;
 
 		u32 attributes = DefaultAttributes();
 		if (useOverlap)
