[UBA]
* Fixed so local events can toggle between normal and semaphore version

diff --git a/Engine/Source/Programs/UnrealBuildAccelerator/Core/Private/UbaEvent.cpp b/Engine/Source/Programs/UnrealBuildAccelerator/Core/Private/UbaEvent.cpp
--- a/Engine/Source/Programs/UnrealBuildAccelerator/Core/Private/UbaEvent.cpp
+++ b/Engine/Source/Programs/UnrealBuildAccelerator/Core/Private/UbaEvent.cpp
@@ -57,19 +57,19 @@ namespace uba
 		EventWin32;
 	#endif
 
-	// Default IPC event
-	using EventImpl = 
-
-	#if UBA_USE_NATIVE_MAC_SEMAPHORES
-		EventDispatchSemaphore;
-	#elif !PLATFORM_WINDOWS
-		EventPThread;
-	#elif UBA_TEST_WINDOWS_SEMAPHORES
-		EventDispatchSemaphore;
-	#else
+	// Default local event
+	using LocalEventDefaultImpl = 
+	#if PLATFORM_WINDOWS
 		EventWin32;
+	#else
+		EventPThread;
 	#endif
 
+	// Semaphore basedlocal event
+	using LocalEventSemaphoreImpl = 
+		EventDispatchSemaphore;
+
+	bool g_useUseSemaphore = UseSemaphoreForLocalEvents();
 	Futex g_firstEventLock[2];
 	EventHeader* g_firstEvent[2];
 
@@ -87,27 +87,33 @@ namespace uba
 
 	bool Event::Create(EventResetType resetType)
 	{
-		EventImpl* impl = nullptr;
+		EventHeader* impl = nullptr;
 		SCOPED_FUTEX(g_firstEventLock[resetType], lock);
 		auto*& first = g_firstEvent[resetType];
 		if (first)
 		{
-			impl = (EventImpl*)first;
+			impl = (EventHeader*)first;
 			first = impl->m_next;
 		}
 		lock.Leave();
 
-		if (!impl)
+		if (impl)
 		{
-			impl = new EventImpl();
-			impl->Create(resetType);
+			m_impl = impl;
+			impl->m_resetType.value = resetType;
+			Reset();
+			return true;
+		}
+		if (g_useUseSemaphore)
+		{
+			impl = new LocalEventSemaphoreImpl();
+			((LocalEventSemaphoreImpl*)impl)->Create(resetType);
 		}
 		else
 		{
-			impl->Reset();
-			impl->m_resetType.value = resetType; // union with next
+			impl = new LocalEventDefaultImpl();
+			((LocalEventDefaultImpl*)impl)->Create(resetType);
 		}
-
 		m_impl = impl;
 		return true;
 	}
@@ -116,7 +122,7 @@ namespace uba
 	{
 		if (!m_impl)
 			return;
-		auto impl = (EventImpl*)m_impl;
+		auto impl = (EventHeader*)m_impl;
 		m_impl = nullptr;
 		auto resetType = impl->m_resetType.value;
 		auto*& first = g_firstEvent[resetType];
@@ -127,14 +133,22 @@ namespace uba
 
 	void Event::Set()
 	{
-		if (m_impl)
-			((EventImpl*)m_impl)->Set();
+		if (!m_impl)
+			return;
+		if (g_useUseSemaphore)
+			((LocalEventSemaphoreImpl*)m_impl)->Set();
+		else
+			((LocalEventDefaultImpl*)m_impl)->Set();
 	}
 
 	void Event::Reset()
 	{
-		if (m_impl)
-			((EventImpl*)m_impl)->Reset();
+		if (!m_impl)
+			return;
+		if (g_useUseSemaphore)
+			((LocalEventSemaphoreImpl*)m_impl)->Reset();
+		else
+			((LocalEventDefaultImpl*)m_impl)->Reset();
 	}
 
 	bool Event::IsCreated()
@@ -144,17 +158,20 @@ namespace uba
 
 	bool Event::IsSet(u32 timeOutMs)
 	{
-		if (m_impl)
-			return ((EventImpl*)m_impl)->IsSet(timeOutMs);
-		return false;
+		if (!m_impl)
+			return false;
+		if (g_useUseSemaphore)
+			return ((LocalEventSemaphoreImpl*)m_impl)->IsSet(timeOutMs);
+		else
+			return ((LocalEventDefaultImpl*)m_impl)->IsSet(timeOutMs);
 	}
 
 	void* Event::GetHandle()
 	{
 		#if PLATFORM_WINDOWS
-		if (!m_impl)
+		if (!m_impl || g_useUseSemaphore)
 			return nullptr;
-		return ((EventImpl*)m_impl)->GetHandle();
+		return ((LocalEventDefaultImpl*)m_impl)->GetHandle();
 		#else
 		UBA_ASSERTF(false, "Event::GetHandle not available");
 		return nullptr;
diff --git a/Engine/Source/Programs/UnrealBuildAccelerator/Core/Private/UbaPlatform.cpp b/Engine/Source/Programs/UnrealBuildAccelerator/Core/Private/UbaPlatform.cpp
--- a/Engine/Source/Programs/UnrealBuildAccelerator/Core/Private/UbaPlatform.cpp
+++ b/Engine/Source/Programs/UnrealBuildAccelerator/Core/Private/UbaPlatform.cpp
@@ -1,6 +1,7 @@
 // Copyright Epic Games, Inc. All Rights Reserved.
 
 #include "UbaPlatform.h"
+#include "UbaDefinitions.h"
 #include "UbaProcessStats.h"
 #include <oodle2.h>
 
@@ -668,6 +669,19 @@ namespace uba
 		return true;
 	}
 
+	bool UseSemaphoreForLocalEvents()
+	{
+		#if PLATFORM_WINDOWS && UBA_TEST_WINDOWS_SEMAPHORES
+		return true;
+		#else
+		return IsRunningDarling();
+ 		#endif
+	}
+
+	bool UseSemaphoreForSharedEvents()
+	{
+		return IsRunningDarling();
+	}
 
 #if !PLATFORM_WINDOWS
 	void strcpy_s(tchar* dest, u64 destCapacity, const tchar* source)
@@ -692,11 +706,6 @@ namespace uba
 		out.Append(TCV("/uba_")).AppendHex(uid);
 	}
 
-	bool UseSemaphoreForSharedEvents()
-	{
-		return IsRunningDarling();
-	}
-
 	bool GetPhysicallyInstalledSystemMemory(u64& outKb)
 	{
 	#if PLATFORM_MAC
diff --git a/Engine/Source/Programs/UnrealBuildAccelerator/Core/Public/UbaPlatform.h b/Engine/Source/Programs/UnrealBuildAccelerator/Core/Public/UbaPlatform.h
--- a/Engine/Source/Programs/UnrealBuildAccelerator/Core/Public/UbaPlatform.h
+++ b/Engine/Source/Programs/UnrealBuildAccelerator/Core/Public/UbaPlatform.h
@@ -113,6 +113,8 @@ namespace uba
 	void AddExceptionHandler();
 	void InitMemory();
 	bool SetMemoryWorkingSet(u64 minWorkingSet);
+	bool UseSemaphoreForLocalEvents();
+	bool UseSemaphoreForSharedEvents();
 
 #if PLATFORM_WINDOWS
 	inline constexpr bool CaseInsensitiveFs = true;
@@ -165,7 +167,6 @@ namespace uba
 	void strcpy_s(tchar* dest, u64 destCapacity, const tchar* source);
 	void strcat_s(tchar* dest, u64 destCapacity, const tchar* source);
 	void GetMappingHandleName(StringBufferBase& out, u64 uid);
-	bool UseSemaphoreForSharedEvents();
 	bool GetPhysicallyInstalledSystemMemory(u64& outKb);
 	inline u64 FromTimeSpec(const timespec& ts) { return u64(ts.tv_sec) * 10'000'000ull + u64(ts.tv_nsec/100); }
 	inline timespec ToTimeSpec(u64 time) { timespec ts; ts.tv_sec = time / 10'000'000ull; ts.tv_nsec = (time - (u64(ts.tv_sec) * 10'000'000ull)) * 100; return ts; }
diff --git a/Engine/Source/Programs/UnrealBuildAccelerator/Detours/Private/PosixOS/UbaDetoursMainPosix.cpp b/Engine/Source/Programs/UnrealBuildAccelerator/Detours/Private/PosixOS/UbaDetoursMainPosix.cpp
--- a/Engine/Source/Programs/UnrealBuildAccelerator/Detours/Private/PosixOS/UbaDetoursMainPosix.cpp
+++ b/Engine/Source/Programs/UnrealBuildAccelerator/Detours/Private/PosixOS/UbaDetoursMainPosix.cpp
@@ -90,10 +90,6 @@ static void __attribute__((constructor(102))) PreInitCtor()
 
 		#if UBA_USE_NATIVE_MAC_SEMAPHORES
 
-		//int tty = (FileHandle)open("/dev/tty", O_WRONLY);
-		//auto log = [tty](const char* str, int len) { write(tty, str, len); };
-		//auto logLine = [&](const char* str, bool newline = true) { log(str, TStrlen(str)); if (newline) log("\n", 1); };
-
 		const char* pidStr = getenv("UBA_PID");
 		UBA_ASSERT(pidStr);
 		u64 pid;
