[UBA]
* Removed usage of stackalloc in ParseArguments function and instead have a fixed stack buffer and only allocate when really needed

diff --git a/Engine/Source/Programs/UnrealBuildAccelerator/Core/Public/UbaProcessUtils.h b/Engine/Source/Programs/UnrealBuildAccelerator/Core/Public/UbaProcessUtils.h
--- a/Engine/Source/Programs/UnrealBuildAccelerator/Core/Public/UbaProcessUtils.h
+++ b/Engine/Source/Programs/UnrealBuildAccelerator/Core/Public/UbaProcessUtils.h
@@ -4,18 +4,6 @@
 
 #include "UbaPlatform.h"
 
-#if defined(__clang_analyzer__)
-#define stackAlloc(x) malloc(x)
-#define stackFree(x) free(x)
-#elif PLATFORM_WINDOWS
-#include <malloc.h>
-#define stackAlloc(x) _malloca(x)
-#define stackFree(x) _freea(x)
-#else
-#define stackAlloc(x) alloca(x)
-#define stackFree(x)
-#endif
-
 namespace uba
 {
 	// func is void(const CharType* arg, u32 strArg)
@@ -34,8 +22,12 @@ namespace uba
 		CharType lastChar = 0;
 		bool isBackslashOwned = false;
 
-		CharType* argWriteBuffer = static_cast<CharType*>( stackAlloc(sizeof(CharType) * (argumentsLen + 1)) );
-		auto g = MakeGuard([&]() { stackFree(argWriteBuffer); });
+		
+		CharType temp[5];
+
+		CharType* argWriteBuffer = temp;
+		u64 argBufferSize = sizeof_array(temp);
+		auto g = MakeGuard([&]() { if (argWriteBuffer != temp) delete[] argWriteBuffer; });
 
 		for (const CharType* it = arguments; !isEnd; lastChar = currentChar, ++it)
 		{
@@ -74,6 +66,17 @@ namespace uba
 						isBackslashOwned = false;
 
 					*argWriteIt++ = *argIt;
+
+					u64 writtenCount = argWriteIt - argWriteBuffer;
+					if (writtenCount + 1 == argBufferSize)
+					{
+						UBA_ASSERT(argWriteBuffer == temp);
+						argBufferSize = argumentsLen + 1;
+						argWriteBuffer = new CharType[argBufferSize];
+						argWriteIt = argWriteBuffer + writtenCount;
+						memcpy(argWriteBuffer, temp, writtenCount *sizeof(CharType));
+					}
+
 					lastChar2 = *argIt;
 					++argIt;
 				}
