Allow SxS manifest generation to occur as normal for UBA tools, now that the underlying issue of SxS manifests conflicting has been resolved.

#Synced-CL 48208542

diff --git a/Engine/Source/Programs/UnrealBuildAccelerator/UbaAgent.Target.cs b/Engine/Source/Programs/UnrealBuildAccelerator/UbaAgent.Target.cs
--- a/Engine/Source/Programs/UnrealBuildAccelerator/UbaAgent.Target.cs
+++ b/Engine/Source/Programs/UnrealBuildAccelerator/UbaAgent.Target.cs
@@ -65,11 +65,6 @@ public class UbaAgentTarget : TargetRules
 				"UBA_AGENT_EXECUTABLE=L\"UbaAgent.exe\"",
 				"UBA_DETOURS_LIBRARY_ANSI=\"UbaDetours.dll\"",
 			});
-
-			if (!Target.Architectures.SingleArchitecture.bIsX64)
-			{
-				Rules.WindowsPlatform.ManifestFile = null;
-			}
 		}
 		else
 		{
