[UBA]
* Improved assert description in scheduler when trying to remove process that already has been removed

diff --git a/Engine/Source/Programs/UnrealBuildAccelerator/Common/Private/UbaScheduler.cpp b/Engine/Source/Programs/UnrealBuildAccelerator/Common/Private/UbaScheduler.cpp
--- a/Engine/Source/Programs/UnrealBuildAccelerator/Common/Private/UbaScheduler.cpp
+++ b/Engine/Source/Programs/UnrealBuildAccelerator/Common/Private/UbaScheduler.cpp
@@ -1318,7 +1318,7 @@ namespace uba
 			StringKey key(impl.m_nativeProcessId, impl.m_nativeProcessCreationTime);
 			SCOPED_FUTEX(m_processesLock, l);
 			bool res = m_processes.erase(key) == 1;
-			UBA_ASSERT(res);(void)res;
+			UBA_ASSERTF(res, TC("Tried to erase local process from process list but can't find it (%s)"), info.GetDescription());(void)res;
 		}
 	}
 
