[UBA]
* Fixed so LastErrorToText can be used in detoured code

diff --git a/Engine/Source/Programs/UnrealBuildAccelerator/Common/Private/UbaLogger.cpp b/Engine/Source/Programs/UnrealBuildAccelerator/Common/Private/UbaLogger.cpp
--- a/Engine/Source/Programs/UnrealBuildAccelerator/Common/Private/UbaLogger.cpp
+++ b/Engine/Source/Programs/UnrealBuildAccelerator/Common/Private/UbaLogger.cpp
@@ -350,10 +350,6 @@ namespace uba
 #endif
 	}
 
-	LastErrorToText::LastErrorToText() : LastErrorToText(GetLastError())
-	{
-	}
-
 	LastErrorToText::LastErrorToText(u32 lastError)
 	{
 #if PLATFORM_WINDOWS
diff --git a/Engine/Source/Programs/UnrealBuildAccelerator/Core/Private/UbaStringBuffer.cpp b/Engine/Source/Programs/UnrealBuildAccelerator/Core/Private/UbaStringBuffer.cpp
--- a/Engine/Source/Programs/UnrealBuildAccelerator/Core/Private/UbaStringBuffer.cpp
+++ b/Engine/Source/Programs/UnrealBuildAccelerator/Core/Private/UbaStringBuffer.cpp
@@ -459,6 +459,8 @@ namespace uba
 		#endif
 	}
 
+	////////////////////////////////////////////////////////////////////////////////////////////////////
+
 	bool StringView::Contains(tchar c) const
 	{
 		for (const tchar* i = data, *e = data + count; i!=e; ++i)
@@ -488,4 +490,12 @@ namespace uba
 	{
 		return StringView(s, TStrlen(s));
 	}
+
+	////////////////////////////////////////////////////////////////////////////////////////////////////
+
+	LastErrorToText::LastErrorToText() : LastErrorToText(GetLastError())
+	{
+	}
+
+	////////////////////////////////////////////////////////////////////////////////////////////////////
 }
diff --git a/Engine/Source/Programs/UnrealBuildAccelerator/Core/Public/UbaStringBuffer.h b/Engine/Source/Programs/UnrealBuildAccelerator/Core/Public/UbaStringBuffer.h
--- a/Engine/Source/Programs/UnrealBuildAccelerator/Core/Public/UbaStringBuffer.h
+++ b/Engine/Source/Programs/UnrealBuildAccelerator/Core/Public/UbaStringBuffer.h
@@ -8,6 +8,8 @@ namespace uba
 {
 	struct StringView;
 
+	////////////////////////////////////////////////////////////////////////////////////////////////////
+
 	bool StartsWith(const tchar* data, const tchar* str, bool ignoreCase = true);
 	bool EndsWith(const tchar* str, u64 strLen, const tchar* value, u64 valueLen, bool ignoreCase);
 	bool Contains(const tchar* str, const tchar* sub, bool ignoreCase = true, const tchar** pos = nullptr);
@@ -21,6 +23,8 @@ namespace uba
 	inline tchar ToUpper(tchar c) { return (c >= 'a' && c <= 'z') ? (c - 'a' + 'A') : c;}
 	const tchar* GetFileName(const tchar* path);
 
+	////////////////////////////////////////////////////////////////////////////////////////////////////
+
 	class StringBufferBase
 	{
 	public:
@@ -105,6 +109,8 @@ namespace uba
 		tchar buf[Capacity];
 	};
 
+	////////////////////////////////////////////////////////////////////////////////////////////////////
+
 	enum NoInitEnum { NoInit };
 
 	struct StringView
@@ -146,6 +152,8 @@ namespace uba
 
 	#define TCV(x) AsView(TC(x))
 
+	////////////////////////////////////////////////////////////////////////////////////////////////////
+
 	struct LastErrorToText : StringBuffer<256>
 	{
 		LastErrorToText();
@@ -166,4 +174,6 @@ namespace uba
 		}
 		return *this;
 	}
+
+	////////////////////////////////////////////////////////////////////////////////////////////////////
 }
\ No newline at end of file
diff --git a/Engine/Source/Programs/UnrealBuildAccelerator/Detours/Private/PosixOS/UbaDetoursFunctionsPosix.cpp b/Engine/Source/Programs/UnrealBuildAccelerator/Detours/Private/PosixOS/UbaDetoursFunctionsPosix.cpp
--- a/Engine/Source/Programs/UnrealBuildAccelerator/Detours/Private/PosixOS/UbaDetoursFunctionsPosix.cpp
+++ b/Engine/Source/Programs/UnrealBuildAccelerator/Detours/Private/PosixOS/UbaDetoursFunctionsPosix.cpp
@@ -3036,6 +3036,11 @@ namespace uba
 		CloseCom();
 		TRUE_WRAPPER(_exit)(int(terminateCode));
 	}
+
+	LastErrorToText::LastErrorToText(u32 lastError)
+	{
+		Append(strerror(int(lastError)));
+	}
 }
 
 extern "C"
diff --git a/Engine/Source/Programs/UnrealBuildAccelerator/Detours/Private/Windows/UbaDetoursFunctionsWin.cpp b/Engine/Source/Programs/UnrealBuildAccelerator/Detours/Private/Windows/UbaDetoursFunctionsWin.cpp
--- a/Engine/Source/Programs/UnrealBuildAccelerator/Detours/Private/Windows/UbaDetoursFunctionsWin.cpp
+++ b/Engine/Source/Programs/UnrealBuildAccelerator/Detours/Private/Windows/UbaDetoursFunctionsWin.cpp
@@ -58,6 +58,7 @@ DETOURED_FUNCTIONS_MEMORY_NON_WINE
 
 namespace uba
 {
+	extern HANDLE g_hostProcess;
 
 bool g_useMiMalloc;
 constexpr u64  g_pageSize = 64*1024;
@@ -525,7 +526,15 @@ UBA_NOINLINE void UbaAssert(const tchar* text, const char* file, u32 line, const
 	#endif
 }
 
-extern HANDLE g_hostProcess;
+LastErrorToText::LastErrorToText(u32 lastError)
+{
+	SuppressDetourScope _;
+	size_t size = ::FormatMessageW(FORMAT_MESSAGE_FROM_SYSTEM | FORMAT_MESSAGE_IGNORE_INSERTS, NULL, lastError, MAKELANGID(LANG_ENGLISH, SUBLANG_ENGLISH_US), data, capacity, NULL);
+	if (!size)
+		AppendValue(lastError);
+	else
+		Resize(size - 2);
+}
 
 const wchar_t* HandleToName(HANDLE handle)
 {
