[UBA]
* Implemented SharedMemoryAllocator which is an allocator that tracks mapped memory that can be shared between processes. First iteration, will need some more cleanup

diff --git a/Engine/Source/Programs/UnrealBuildAccelerator/Common/Private/UbaSharedMemoryAllocator.cpp b/Engine/Source/Programs/UnrealBuildAccelerator/Common/Private/UbaSharedMemoryAllocator.cpp
new file mode 100644
--- /dev/null
+++ b/Engine/Source/Programs/UnrealBuildAccelerator/Common/Private/UbaSharedMemoryAllocator.cpp
@@ -0,0 +1,596 @@
+// Copyright Epic Games, Inc. All Rights Reserved.
+
+#include "UbaSharedMemoryAllocator.h"
+#include "UbaBinaryReaderWriter.h"
+#include "UbaHash.h"
+#include "UbaTrace.h"
+
+#if PLATFORM_WINDOWS
+#include <windows.h>
+#else
+#define GetCurrentProcess() 0
+#define UnmapViewOfFile2(...) false
+#define VirtualFree(...) false
+#define VirtualAlloc(...) nullptr
+#define VirtualAlloc2(...) nullptr
+#define MapViewOfFile3(...) nullptr
+#define ZeroMemory(...) {}
+#define PAGE_READONLY 0x02
+#define PAGE_READWRITE 0x04
+#define PAGE_WRITECOPY 0x08
+#endif
+
+#define UBA_DEBUG_HINTS 1
+
+namespace uba
+{
+	////////////////////////////////////////////////////////////////////////////////////////////////////
+
+	u8 g_emptyMemory;
+
+	struct SharedMemoryAllocator::Slice
+	{
+		u64 offset;
+		u64 size;
+	};
+
+	struct SharedMemoryAllocator::InternalHandle
+	{
+		using Slices = Vector<Slice>;
+		Slices slices;
+		Atomic<u32> refCount = 1;
+
+		#if UBA_DEBUG_HINTS
+		TString hint;
+		const tchar* GetHint() { return hint.c_str(); }
+		#else
+		const tchar* GetHint() { return TC("EnableDebugHints"); }
+		#endif
+	};
+
+	#define UBA_REF_COUNT(f) UBA_ASSERTF(h.refCount, TC( #f " accessing released shared memory handle (%s)"), h.GetHint());
+
+	////////////////////////////////////////////////////////////////////////////////////////////////////
+
+	bool SharedMemoryAllocator::Init(u64 capacity, u64 commitStepSize)
+	{
+		m_commitStepSize = commitStepSize;
+		m_reserved = AlignUp(capacity, 64 * 1024);
+		m_fileMapping = CreateMemoryMappingW(m_logger, PAGE_READWRITE|SEC_RESERVE, m_reserved, nullptr, TC("SharedMemoryAllocator"));
+		if (!m_fileMapping.IsValid())
+			return false;
+
+		m_memory = MapViewOfFile(m_logger, m_fileMapping, FILE_MAP_READ | FILE_MAP_WRITE, 0, 0, m_reserved);
+		if (!m_memory)
+		{
+			CloseFileMapping(m_logger, m_fileMapping, TC("SharedMemoryAllocator"));
+			return false;
+		}
+		return true;
+	}
+
+	SharedMemoryAllocator::operator SharedMemoryAllocatorHandle() const
+	{
+		return { m_fileMapping.ToU64() };
+	}
+
+	u64 SharedMemoryAllocator::GetReserved()
+	{
+		return m_reserved;
+	}
+
+	SharedMemoryHandle SharedMemoryAllocator::CreateHandle(Logger& logger, const tchar* hint, bool priority)
+	{
+		++m_handleCount;
+		auto h = new InternalHandle();
+
+		#if UBA_DEBUG_HINTS
+		h->hint = hint;
+		#endif
+
+		if (priority)
+		{
+			u32 index = m_priorityHandleCounter++;
+			if (index < sizeof_array(m_priorityHandles))
+			{
+				m_priorityHandles[index] = h;
+				return { index + 1 };
+			}
+		}
+
+		//LoggerWithWriter(g_consoleLogWriter, TC("")).Info(TC("CREATED %s (%llu)"), hint, h);
+
+		return { u64(h) };
+	}
+
+	void SharedMemoryAllocator::ClearHandle(SharedMemoryHandle handle)
+	{
+		if (!handle.internalHandle)
+			return;
+		auto& h = GetHandle(handle);
+		UBA_REF_COUNT(ClearHandle)
+		ClearHandle(h);
+	}
+
+	void SharedMemoryAllocator::CloseHandle(Logger& logger, SharedMemoryHandle handle, const tchar* hint)
+	{
+		if (handle.internalHandle)
+			ReleaseHandle(handle);
+	}
+
+	void SharedMemoryAllocator::ExtendMemory(BinaryWriter& out, SharedMemoryHandle handle, u64 size, const tchar* hint, bool zeroMemory)
+	{
+		UBA_ASSERT(handle.internalHandle);
+		auto& h = GetHandle(handle);
+		UBA_REF_COUNT(ExtendMemory);
+
+		auto& slices = h.slices;
+		
+		u32 sliceCount = u32(slices.size());
+		ExtendHandle(h, size, hint, zeroMemory, 0);
+		u32 newCount = u32(slices.size());
+
+		out.WriteU32(newCount - sliceCount);
+		for (u32 i = sliceCount, e = newCount; i != e; ++i)
+		{
+			out.Write7BitEncoded(slices[i].offset);
+			out.Write7BitEncoded(slices[i].size);
+		}
+	}
+
+	void SharedMemoryAllocator::ExtendMemory(SharedMemoryHandle handle, u64 size, const tchar* hint, bool zeroMemory)
+	{
+		UBA_ASSERTF(m_fileMapping.IsValid(), TC("Allocator handle not initialized"));
+		UBA_ASSERT(handle.internalHandle);
+		auto& h = GetHandle(handle);
+		UBA_REF_COUNT(ExtendMemory);
+		ExtendHandle(h, size, hint, zeroMemory, 0);
+	}
+
+	void SharedMemoryAllocator::ExtendMemory(SharedMemoryView& view, SharedMemoryHandle handle, u64 size, const tchar* hint, bool zeroMemory, u64 preferredOffset)
+	{
+		UBA_ASSERT(handle.internalHandle);
+		auto& h = GetHandle(handle);
+		UBA_REF_COUNT(ExtendMemory);
+
+		auto& slices = h.slices;
+		u32 sliceCount = u32(slices.size());
+		ExtendHandle(h, size, hint, zeroMemory, preferredOffset);
+
+		for (u32 i=sliceCount, e= u32(slices.size()); i!=e; ++i)
+			view.ExtendMemory(*this, slices[i].offset, slices[i].size, SharedMemoryMapType_ReadWrite, hint);
+	}
+
+	u64 SharedMemoryAllocator::GetCommitted(SharedMemoryHandle handle)
+	{
+		auto& h = GetHandle(handle);
+		UBA_REF_COUNT(GetCommitted);
+		u64 committed = 0;
+		for (auto& slice : h.slices)
+			committed += slice.size;
+		return committed;
+	}
+
+	u64 SharedMemoryAllocator::GetRefCount(SharedMemoryHandle handle)
+	{
+		auto& h = GetHandle(handle);
+		UBA_REF_COUNT(GetRefCount);
+		return h.refCount;
+	}
+
+	void SharedMemoryAllocator::WriteViewSlices(BinaryWriter& out, SharedMemoryHandle handle, u64 offset, u64 size)
+	{
+		auto& h = GetHandle(handle);
+		UBA_REF_COUNT(WriteViewSlices);
+		auto& totCount = *(u32*)out.AllocWrite(4);
+		u32 count = 0;
+		for (auto& slice : h.slices)
+		{
+			u64 sliceOffset = slice.offset;
+			u64 sliceSize = slice.size;
+			if (offset >= sliceSize)
+			{
+				offset -= sliceSize;
+				continue;
+			}
+			else if (offset)
+			{
+				sliceOffset += offset;
+				sliceSize -= offset;
+				offset = 0;
+			}
+
+			++count;
+
+			out.Write7BitEncoded(sliceOffset);
+
+			if (size <= sliceSize)
+			{
+				out.Write7BitEncoded(AlignUp(size, 64*1024));
+				break;
+			}
+
+			out.Write7BitEncoded(sliceSize);
+			size -= sliceSize;
+			if (!size)
+				break;
+		}
+
+		//wprintf(L"WRITING %u slices for %s\n", count, h.GetHint());
+
+		totCount = count;
+	}
+
+	bool SharedMemoryAllocator::MapView(SharedMemoryView& view, SharedMemoryHandle handle, const tchar* hint, SharedMemoryMapType mapType, u64 offset, u64 size)
+	{
+		if (!handle.internalHandle)
+			return false;
+
+		auto& h = GetHandle(handle);
+		UBA_REF_COUNT(MapView);
+
+		if (size == 0)
+			return true;
+
+		u64 totalSize = size;
+		if (totalSize == ~0ull)
+		{
+			totalSize = 0;
+			for (auto& slice : h.slices)
+				totalSize += slice.size;
+		}
+		view.Init(*this, totalSize);
+
+		bool initialized = false;
+
+		for (auto& slice : h.slices)
+		{
+			u64 sliceOffset = slice.offset;
+			u64 sliceSize = slice.size;
+			if (offset >= sliceSize)
+			{
+				offset -= sliceSize;
+				continue;
+			}
+			else if (offset)
+			{
+				sliceOffset += offset;
+				sliceSize -= offset;
+				offset = 0;
+			}
+
+			if (totalSize <= sliceSize)
+			{
+				if (!initialized)
+					view.ReferenceExternalMemory(m_memory + sliceOffset, totalSize);
+				else
+					view.ExtendMemory(*this, sliceOffset, totalSize, mapType, hint);
+				break;
+			}
+
+			if (!initialized)
+				view.Init(*this, totalSize);
+			initialized = true;
+
+			view.ExtendMemory(*this, sliceOffset, sliceSize, mapType, hint);
+			totalSize -= sliceSize;
+			if (!totalSize)
+				break;
+		}
+		return true;
+	}
+
+	u8* SharedMemoryAllocator::MapView(SharedMemoryHandle handle, SharedMemoryMapType mapType)
+	{
+		if (!handle.internalHandle)
+			return nullptr;
+
+		auto& h = GetHandle(handle);
+		UBA_REF_COUNT(MapView);
+
+		if (h.slices.empty())
+			return &g_emptyMemory;
+
+		++h.refCount;
+
+		if (h.slices.size() == 1)
+			return m_memory + h.slices.front().offset;
+
+		u64 committed = 0;
+		for (auto& slice : h.slices)
+			committed += slice.size;
+
+		u32 protect = PAGE_READWRITE;
+		if (mapType == SharedMemoryMapType_ReadOnly)
+			protect = PAGE_READONLY;
+		else if (mapType == SharedMemoryMapType_CopyOnWrite)
+			protect = PAGE_WRITECOPY;
+
+		u8* memory = (u8*)VirtualAlloc2(GetCurrentProcess(), nullptr, committed, MEM_RESERVE | MEM_RESERVE_PLACEHOLDER, PAGE_NOACCESS, nullptr, 0);
+		UBA_ASSERT(memory);
+
+		auto& stats = KernelStats::GetCurrent().mapViewOfFile3;
+		u64 pos = 0;
+		for (Slice slice : h.slices)
+		{
+			u8* targetAddress = memory + pos;
+
+			if (pos + slice.size < committed)
+			{
+				bool res = VirtualFree(targetAddress, slice.size, MEM_RELEASE | MEM_PRESERVE_PLACEHOLDER);
+				UBA_ASSERTF(res, TC("VirtualFree failed for pos %llu with size %llu (%s)"), pos, slice.size, LastErrorToText().data); (void)res;
+			}
+
+			stats.bytes += slice.size;
+			TimerScope ts(stats);
+
+			void* viewA = MapViewOfFile3(m_fileMapping.mh, GetCurrentProcess(), targetAddress, slice.offset, slice.size, MEM_REPLACE_PLACEHOLDER, protect, nullptr, 0);
+			UBA_ASSERTF(viewA, TC("MapViewOfFile3 failed in MapView 0x%p %llu for % s (%s)"), memory, pos, h.GetHint(), LastErrorToText().data); (void)viewA;
+			pos += slice.size;
+		}
+
+		return memory;
+	}
+
+	void SharedMemoryAllocator::UnmapView(SharedMemoryHandle handle, u8* memory)
+	{
+		if (!memory || memory == &g_emptyMemory)
+			return;
+
+		UBA_ASSERT(handle.internalHandle);
+		auto& h = GetHandle(handle);
+		UBA_REF_COUNT(UnmapView);
+
+		if (h.slices.size() == 1)
+		{
+			UBA_ASSERT(memory >= m_memory && memory < m_memory + m_committed);
+			ReleaseHandle(handle);
+			return;
+		}
+		u64 pos = 0;
+		u32 res; (void)res;
+
+		auto& stats = KernelStats::GetCurrent().unmapViewOfFile2;
+
+		for (Slice slice : h.slices)
+		{
+			stats.bytes += slice.size;
+			TimerScope ts(stats);
+			res = UnmapViewOfFile2(GetCurrentProcess(), memory + pos, MEM_PRESERVE_PLACEHOLDER);
+			UBA_ASSERTF(res, TC("UnmapViewOfFile2 failed on pos %llu for %s (%s)"), pos, h.GetHint(), LastErrorToText().data);
+			pos += slice.size;
+		}
+		if (h.slices.size() > 1)
+		{
+			res = VirtualFree(memory, pos, MEM_RELEASE | MEM_COALESCE_PLACEHOLDERS);
+			UBA_ASSERTF(res, TC("VirtualFree failed (%s)"), LastErrorToText().data);
+		}
+
+		res = VirtualFree(memory, 0, MEM_RELEASE);
+		UBA_ASSERTF(res, TC("VirtualFree failed (%s)"), LastErrorToText().data);
+
+		ReleaseHandle(handle);
+	}
+
+	SharedMemoryAllocator::SharedMemoryAllocator(Logger& logger)
+	:	m_logger(logger)
+	{
+	}
+
+	SharedMemoryAllocator::~SharedMemoryAllocator()
+	{
+		if (m_memory)
+			UnmapViewOfFile(m_logger, m_memory, m_reserved, TC("SharedMemoryAllocator"));
+		if (m_fileMapping.IsValid())
+			CloseFileMapping(m_logger, m_fileMapping, TC("SharedMemoryAllocator"));
+	}
+
+	SharedMemoryAllocator::InternalHandle& SharedMemoryAllocator::GetHandle(SharedMemoryHandle h)
+	{
+		if (h.internalHandle <= sizeof_array(m_priorityHandles))
+			return *m_priorityHandles[h.internalHandle - 1];
+		return *(InternalHandle*)h.internalHandle;
+	}
+
+	void SharedMemoryAllocator::ReleaseHandle(SharedMemoryHandle h)
+	{
+		InternalHandle& handle = GetHandle(h);
+		UBA_ASSERT(handle.refCount);
+		if (--handle.refCount)
+			return;
+		ClearHandle(handle);
+		#if !UBA_DEBUG_HINTS
+		delete& handle;
+		#endif
+		if (h.internalHandle <= sizeof_array(m_priorityHandles))
+			m_priorityHandles[h.internalHandle - 1] = 0;
+		--m_handleCount;
+	}
+
+	void SharedMemoryAllocator::ExtendHandle(InternalHandle& handle, u64 size, const tchar* hint, bool zeroMemory, u64 preferredOffset)
+	{
+		size = AlignUp(size, 64 * 1024);
+		u64 left = size;
+
+		auto& slices = handle.slices;
+		u32 sliceCount = u32(slices.size());
+
+		SCOPED_FUTEX(m_lock, lock);
+		while (left)
+		{
+			Slice slice = AllocateNoLock(left, preferredOffset);
+			slices.push_back(slice);
+			left -= slice.size;
+		}
+		lock.Leave();
+
+		if (zeroMemory)
+			for (u32 i = sliceCount, e = u32(slices.size()); i != e; ++i)
+				ZeroMemory(m_memory + slices[i].offset, slices[i].size);
+	}
+
+	void SharedMemoryAllocator::ClearHandle(InternalHandle& handle)
+	{
+		if (handle.slices.empty())
+			return;
+		SCOPED_FUTEX(m_lock, lock);
+		for (auto& slice : handle.slices)
+			FreeNoLock(slice);
+		lock.Leave();
+		handle.slices.clear();
+	}
+
+	SharedMemoryAllocator::Slice SharedMemoryAllocator::AllocateNoLock(u64 desiredSize, u64 preferredOffset)
+	{
+		// Find best fit.
+		// Will return first iterator not less than desiredSize
+		auto it = m_freeMemoryBySize.lower_bound({ desiredSize, 0 });
+		if (it == m_freeMemoryBySize.end())
+		{
+			if (!m_freeMemory.empty())
+			{
+				--it;
+				u64 size = it->size;
+				if (size >= desiredSize / 2) // Half the size is ok
+				{
+					u64 offset = it->offset;
+					m_freeMemoryBySize.erase(it);
+					m_freeMemory.erase(offset);
+					m_used += size;
+					return { offset, size };
+				}
+			}
+
+			// No block is big enough, create new one
+			u64 toCommit = Min(AlignUp(desiredSize, m_commitStepSize), m_reserved);
+			UBA_SHIPPING_ASSERTF(m_committed + toCommit <= m_reserved, TC("Not enough space reserved for commit (Reserved: %llu, Committed: %llu Requested: %llu"), m_reserved, m_committed, toCommit);
+			void* res = VirtualAlloc(m_memory + m_committed, toCommit, MEM_COMMIT, PAGE_READWRITE);
+			UBA_SHIPPING_ASSERTF(res == m_memory + m_committed, TC("VirtualAlloc failed (%s)"), LastErrorToText().data);
+			u64 offset = m_committed;
+			m_committed += toCommit;
+			u64 size = toCommit;
+			if (toCommit > desiredSize)
+			{
+				u64 newOffset = offset + desiredSize;
+				u64 newSize = toCommit - desiredSize;
+				m_freeMemory.emplace(newOffset, newSize);
+				m_freeMemoryBySize.insert({ newSize, newOffset });
+				size = desiredSize;
+			}
+			m_used += size;
+			return { offset, size };
+		}
+		else
+		{
+			u64 offset = it->offset;
+			u64 size = it->size;
+			if (size == desiredSize) // Perfect fit
+			{
+				m_freeMemoryBySize.erase(it);
+				m_freeMemory.erase(offset);
+				m_used += desiredSize;
+				return { offset, desiredSize };
+			}
+			else // Slot is larger, take first part of it
+			{
+				m_freeMemory.erase(offset);
+				m_freeMemoryBySize.erase(it);
+
+				u64 newOffset = offset + desiredSize;
+				u64 newSize = size - desiredSize;
+
+				m_freeMemory.emplace(newOffset, newSize);
+				m_freeMemoryBySize.insert({ newSize, newOffset });
+
+				m_used += desiredSize;
+				return { offset, desiredSize };
+			}
+		}
+	}
+
+	void SharedMemoryAllocator::FreeNoLock(Slice slice)
+	{
+		// Returns it to first slice not earlier than offset
+		auto it = m_freeMemory.lower_bound(slice.offset);
+
+		UBA_ASSERT(slice.size);
+
+		m_used -= slice.size;
+
+		if (it == m_freeMemory.end()) // The slice to add is after last
+		{
+			if (m_freeMemory.empty()) // No slices found at all
+			{
+				m_freeMemory.try_emplace(slice.offset, slice.size);
+				m_freeMemoryBySize.insert({ slice.size, slice.offset });
+				return;
+			}
+			else // There are slots before, let's check if we can coalece
+			{
+				--it;
+				u64 end = it->first + it->second;
+				if (end == slice.offset) // We can combine the previous slot
+				{
+					m_freeMemoryBySize.erase({ it->second, it->first });
+					it->second += slice.size;
+					m_freeMemoryBySize.insert({ it->second, it->first });
+					return;
+				}
+				else // This is a new slot
+				{
+					m_freeMemory.try_emplace(slice.offset, slice.size);
+					m_freeMemoryBySize.insert({ slice.size, slice.offset });
+					return;
+				}
+			}
+		}
+		else
+		{
+			bool addedToPrev = false;
+
+			if (it != m_freeMemory.begin()) // There are slots both before and after
+			{
+				auto prevIt = it; --prevIt;
+				u64 end = prevIt->first + prevIt->second;
+				if (end == slice.offset) // This slice can be coaleced with previous slice (reuse the already added one
+				{
+					m_freeMemoryBySize.erase({ prevIt->second, prevIt->first });
+					prevIt->second += slice.size;
+					it = prevIt;
+					addedToPrev = true;
+					slice.offset = prevIt->first;
+					slice.size = prevIt->second;
+				}
+			}
+
+			if (!addedToPrev)
+			{
+				it = m_freeMemory.emplace(slice.offset, slice.size).first;
+			}
+
+			auto next = it; ++next;
+			if (it->first + it->second == next->first) // We can coalece with the one after
+			{
+				u64 nextSize = next->second;
+				m_freeMemoryBySize.erase({ nextSize, next->first });
+				it->second += nextSize;
+				m_freeMemory.erase(next);
+				slice.size = it->second;
+			}
+
+			m_freeMemoryBySize.insert({ slice.size, slice.offset });
+			return;
+		}
+	}
+
+	void SharedMemoryAllocator::TraceStats(Trace& trace, u32 startRow)
+	{
+		trace.StatusUpdate(startRow, 1, TCV("SharedMemory"), LogEntryType_Info);
+
+		StringBuffer<> str;
+		str.Appendf(TC("%s/%s (%llu)"), BytesToText(m_used).str, BytesToText(m_committed).str, m_handleCount.load());
+		trace.StatusUpdate(startRow, 7, str, LogEntryType_Info);
+	}
+
+	////////////////////////////////////////////////////////////////////////////////////////////////////
+}
diff --git a/Engine/Source/Programs/UnrealBuildAccelerator/Common/Public/UbaSharedMemoryAllocator.h b/Engine/Source/Programs/UnrealBuildAccelerator/Common/Public/UbaSharedMemoryAllocator.h
new file mode 100644
--- /dev/null
+++ b/Engine/Source/Programs/UnrealBuildAccelerator/Common/Public/UbaSharedMemoryAllocator.h
@@ -0,0 +1,87 @@
+// Copyright Epic Games, Inc. All Rights Reserved.
+
+#pragma once
+
+#include "UbaMemory.h"
+#include "UbaPlatform.h"
+#include "UbaSharedMemoryView.h"
+
+namespace uba
+{
+	struct BinaryWriter;
+	class Logger;
+	class Trace;
+
+	////////////////////////////////////////////////////////////////////////////////////////////////////
+
+	class SharedMemoryAllocator
+	{
+	public:
+		bool Init(u64 capacity, u64 commitStepSize = 16ull*1024*1024);
+
+		operator SharedMemoryAllocatorHandle() const;
+		u64 GetReserved();
+
+		SharedMemoryHandle CreateHandle(Logger& logger, const tchar* hint, bool priority = false); // Priority means it will have a compressable handle id
+		void ClearHandle(SharedMemoryHandle handle);
+		void CloseHandle(Logger& logger, SharedMemoryHandle handle, const tchar* hint);
+
+		void ExtendMemory(BinaryWriter& out, SharedMemoryHandle handle, u64 size, const tchar* hint = TC(""), bool zeroMemory = true);
+		void ExtendMemory(SharedMemoryHandle handle, u64 size, const tchar* hint = TC(""), bool zeroMemory = true);
+		void ExtendMemory(SharedMemoryView& view, SharedMemoryHandle handle, u64 size, const tchar* hint = TC(""), bool zeroMemory = true, u64 preferredOffset = 0);
+
+		u64 GetCommitted(SharedMemoryHandle handle);
+		u64 GetRefCount(SharedMemoryHandle handle);
+
+		void WriteViewSlices(BinaryWriter& out, SharedMemoryHandle handle, u64 offset = 0, u64 size = ~0llu);
+		bool MapView(SharedMemoryView& view, SharedMemoryHandle handle, const tchar* hint, SharedMemoryMapType mapType = SharedMemoryMapType_ReadOnly, u64 offset = 0, u64 size = ~0llu);
+
+
+		u8* MapView(SharedMemoryHandle handle, SharedMemoryMapType mapType = SharedMemoryMapType_ReadOnly);
+		void UnmapView(SharedMemoryHandle handle, u8* memory);
+
+
+		void TraceStats(Trace& trace, u32 startRow);
+
+		SharedMemoryAllocator(Logger& logger);
+		~SharedMemoryAllocator();
+
+	private:
+		struct InternalHandle;
+		struct Slice;
+
+		InternalHandle& GetHandle(SharedMemoryHandle h);
+		void ReleaseHandle(SharedMemoryHandle h);
+		void ExtendHandle(InternalHandle& handle, u64 size, const tchar* hint, bool zeroMemory, u64 preferredOffset);
+		void ClearHandle(InternalHandle& handle);
+		Slice AllocateNoLock(u64 desiredSize, u64 preferredOffset);
+		void FreeNoLock(Slice slice);
+
+		Logger& m_logger;
+
+		u64 m_commitStepSize;
+		FileMappingHandle m_fileMapping;
+		u64 m_reserved = 0;
+		u8* m_memory = nullptr;
+		u64 m_committed = 0;
+
+		u64 m_used = 0;
+
+		Futex m_lock;
+
+		Map<u64, u64> m_freeMemory; // memory slots. Key is offset, value is size
+
+		struct SizeAndOffset { u64 size; u64 offset; bool operator<(const SizeAndOffset& o) const { return size != o.size ? size < o.size : offset < o.offset; } };
+		Set<SizeAndOffset> m_freeMemoryBySize; // Sorted by size and then offset. To find best fit allocations
+
+		Atomic<u64> m_handleCount = 0;
+
+		InternalHandle* m_priorityHandles[8];
+		Atomic<u32> m_priorityHandleCounter = 0;
+
+		SharedMemoryAllocator(const SharedMemoryAllocator&) = delete;
+		SharedMemoryAllocator& operator=(const SharedMemoryAllocator&) = delete;
+	};
+
+	////////////////////////////////////////////////////////////////////////////////////////////////////
+}
diff --git a/Engine/Source/Programs/UnrealBuildAccelerator/Core/Private/UbaSharedMemoryView.cpp b/Engine/Source/Programs/UnrealBuildAccelerator/Core/Private/UbaSharedMemoryView.cpp
new file mode 100644
--- /dev/null
+++ b/Engine/Source/Programs/UnrealBuildAccelerator/Core/Private/UbaSharedMemoryView.cpp
@@ -0,0 +1,202 @@
+// Copyright Epic Games, Inc. All Rights Reserved.
+
+#include "UbaSharedMemoryView.h"
+#include "UbaBinaryReaderWriter.h"
+#include "UbaHash.h"
+#include "UbaProcessStats.h"
+#include "UbaTimer.h"
+
+#if PLATFORM_WINDOWS
+#include <windows.h>
+#pragma comment(lib, "onecore.lib")
+#else
+#define GetCurrentProcess() 0
+#define UnmapViewOfFile2(...) false
+#define VirtualFree(...) false
+#define VirtualAlloc2(...) nullptr
+#define MapViewOfFile3(...) nullptr
+#define PAGE_READONLY 0x02
+#define PAGE_READWRITE 0x04
+#define PAGE_WRITECOPY 0x08
+#endif
+
+namespace uba
+{
+	////////////////////////////////////////////////////////////////////////////////////////////////////
+
+	bool SharedMemoryView::Init(SharedMemoryAllocatorHandle handle, u64 capacity, void* baseAddress, const tchar* hint)
+	{
+		return InternalInit(handle, capacity, baseAddress, hint);
+	}
+
+	bool SharedMemoryView::Init(SharedMemoryAllocatorHandle handle, u8* sharedMemory, BinaryReader& reader, u64 capacity, const tchar* hint)
+	{
+		u32 sliceCount = reader.ReadU32();
+		if (sliceCount == 1 && sharedMemory)
+		{
+			u64 sliceOffset = reader.Read7BitEncoded();
+			m_mappedSize = capacity;
+			m_capacity = ~0ull;
+			m_memory = sharedMemory + sliceOffset;
+			return true;
+		}
+
+		if (!InternalInit(handle, capacity, nullptr, hint))
+			return false;
+		AddRequestedMemory(handle, reader, sliceCount, SharedMemoryMapType_ReadOnly, 0, hint, 0);
+		return true;
+	}
+
+	bool SharedMemoryView::AddRequestedMemory(SharedMemoryAllocatorHandle handle, BinaryReader& reader, SharedMemoryMapType mapType, u64 offset, const tchar* hint, u8 init)
+	{
+		u32 sliceCount = reader.ReadU32();
+		return AddRequestedMemory(handle, reader, sliceCount, mapType, offset, hint, init);
+	}
+
+	void SharedMemoryView::Reset()
+	{
+		if (!m_memory)
+			return;
+
+		if (m_capacity == ~0ull)
+		{
+			m_memory = nullptr;
+			m_capacity = 0;
+			m_mappedSize = 0;
+			return;
+		}
+
+		bool res; (void)res;
+
+		if (m_mappedSize)
+		{
+			auto& stats = KernelStats::GetCurrent().unmapViewOfFile2;
+			{
+				TimerScope ts(stats);
+				res = UnmapViewOfFile2(GetCurrentProcess(), m_memory, MEM_PRESERVE_PLACEHOLDER);
+				UBA_ASSERT(res);
+			}
+
+			for (u64 mapping : m_additionalMappings)
+			{
+				TimerScope ts(stats);
+				res = UnmapViewOfFile2(GetCurrentProcess(), m_memory + mapping, MEM_PRESERVE_PLACEHOLDER);
+				UBA_ASSERT(res);
+			}
+
+			if (m_mappedSize < m_capacity || !m_additionalMappings.empty())
+			{
+				res = VirtualFree(m_memory, m_capacity, MEM_RELEASE | MEM_COALESCE_PLACEHOLDERS);
+				UBA_ASSERTF(res, TC("VirtualFree failed (%u)"), GetLastError());
+			}
+
+			m_additionalMappings.clear();
+
+			m_mappedSize = 0;
+		}
+		res = VirtualFree(m_memory, 0, MEM_RELEASE);
+		UBA_ASSERTF(res, TC("VirtualFree failed (%u)"), GetLastError());
+
+		m_memory = nullptr;
+		m_capacity = 0;
+	}
+
+	SharedMemoryView::~SharedMemoryView()
+	{
+		Reset();
+	}
+
+	bool SharedMemoryView::InternalInit(SharedMemoryAllocatorHandle handle, u64 capacity, void* baseAddress, const tchar* hint)
+	{
+		m_capacity = AlignUp(capacity, 64 * 1024);
+		m_memory = (u8*)VirtualAlloc2(GetCurrentProcess(), baseAddress, m_capacity, MEM_RESERVE | MEM_RESERVE_PLACEHOLDER, PAGE_NOACCESS, nullptr, 0);
+		UBA_ASSERTF(m_memory, TC("VirtualAlloc2 failed (%s)"), LastErrorToText().data);
+		return m_memory != 0;
+	}
+
+	bool SharedMemoryView::AddRequestedMemory(SharedMemoryAllocatorHandle handle, BinaryReader& reader, u32 sliceCount, SharedMemoryMapType mapType, u64 offset, const tchar* hint, u8 init)
+	{
+		while (sliceCount--)
+		{
+			u64 sliceOffset = reader.Read7BitEncoded();
+			u64 sliceSize = reader.Read7BitEncoded();
+
+			if (offset)
+			{
+				if (offset >= sliceSize)
+				{
+					offset -= sliceSize;
+					continue;
+				}
+				else
+				{
+					sliceOffset += offset;
+					sliceSize -= offset;
+					offset = 0;
+				}
+			}
+
+			if (!ExtendMemory(handle, sliceOffset, sliceSize, mapType, hint, init))
+				return false;
+		}
+		return true;
+	}
+
+	void SharedMemoryView::ReferenceExternalMemory(u8* memory, u64 mappedSize)
+	{
+		m_memory = memory;
+		m_mappedSize = mappedSize;
+		m_capacity = ~0ull;
+	}
+
+	bool SharedMemoryView::ExtendMemory(SharedMemoryAllocatorHandle handle, u64 offset, u64 size, SharedMemoryMapType mapType, const tchar* hint, u8 init)
+	{
+		UBA_ASSERTF(handle.internalHandle, TC("Allocator handle not initialized (%s)"), hint);
+		UBA_ASSERTF(m_memory, TC("Init failed?"));
+		UBA_ASSERTF(m_capacity != ~0ull, TC("Can't extend memory on a view directly mapped on to allocator mapping (%s)"), hint);
+		UBA_ASSERTF(AlignUp(offset, 64 * 1024) == offset, TC("ExtendMemory got offset %llu parameter which is not aligned (%s)"), offset, hint);
+		UBA_ASSERTF(AlignUp(size, 64 * 1024) == size, TC("ExtendMemory got size %llu which is not aligned (%s)"), size, hint);
+
+		u8* targetAddress = m_memory + m_mappedSize;
+
+		u64 left = m_capacity - m_mappedSize;
+		UBA_ASSERTF(left >= size, TC("Want to extend memory more than capacity (Capacity: %llu, Committed: %llu, Requested: %llu)"), m_capacity, m_mappedSize, size);
+		if (left < size)
+			return false;
+
+		if (u64 toRelease = m_capacity - m_mappedSize - size)
+		{
+			bool res = VirtualFree(targetAddress + size, toRelease, MEM_RELEASE | MEM_PRESERVE_PLACEHOLDER);
+			UBA_ASSERTF(res, TC("VirtualFree for placeholder failed releasing %llu at %llu (Capacity: %llu MappedSize: %llu) - %s (%s)"), toRelease, size, m_capacity, m_mappedSize, LastErrorToText().data, hint);
+			if (!res)
+				return false;
+		}
+
+		u32 protect = PAGE_READWRITE;
+		if (mapType == SharedMemoryMapType_ReadOnly)
+			protect = PAGE_READONLY;
+		else if (mapType == SharedMemoryMapType_CopyOnWrite)
+			protect = PAGE_WRITECOPY;
+
+		auto& stats = KernelStats::GetCurrent().mapViewOfFile3;
+		stats.bytes += size;
+		TimerScope ts(stats);
+
+		void* viewA = MapViewOfFile3((HANDLE)handle.internalHandle, GetCurrentProcess(), targetAddress, offset, size, MEM_REPLACE_PLACEHOLDER, protect, nullptr, 0);
+		UBA_ASSERTF(viewA, TC("MapViewOfFile3 failed - %s (%s)"), LastErrorToText().data, hint);
+		if (!viewA)
+			return false;
+
+		if (init)
+			memset(viewA, init, size);
+
+		if (m_mappedSize)
+			m_additionalMappings.push_back(m_mappedSize);
+
+		m_mappedSize += size;
+
+		return true;
+	}
+
+	////////////////////////////////////////////////////////////////////////////////////////////////////
+}
diff --git a/Engine/Source/Programs/UnrealBuildAccelerator/Core/Public/UbaSharedMemoryView.h b/Engine/Source/Programs/UnrealBuildAccelerator/Core/Public/UbaSharedMemoryView.h
new file mode 100644
--- /dev/null
+++ b/Engine/Source/Programs/UnrealBuildAccelerator/Core/Public/UbaSharedMemoryView.h
@@ -0,0 +1,63 @@
+// Copyright Epic Games, Inc. All Rights Reserved.
+
+#pragma once
+
+#include "UbaMemory.h"
+
+namespace uba
+{
+	struct BinaryReader;
+
+	////////////////////////////////////////////////////////////////////////////////////////////////////
+
+	struct SharedMemoryAllocatorHandle { u64 internalHandle = 0; };
+	struct SharedMemoryHandle
+	{
+		inline bool IsValid() const { return internalHandle != 0; }
+		inline u64 ToU64() const { return internalHandle; }
+		inline bool operator==(const SharedMemoryHandle& o) const { return internalHandle == o.internalHandle; }
+		static SharedMemoryHandle FromU64(u64 v) { return { v }; }
+		u64 internalHandle = 0;
+	};
+
+	enum SharedMemoryMapType { SharedMemoryMapType_ReadOnly, SharedMemoryMapType_ReadWrite, SharedMemoryMapType_CopyOnWrite };
+
+	////////////////////////////////////////////////////////////////////////////////////////////////////
+
+	class SharedMemoryView
+	{
+	public:
+		SharedMemoryView() = default;
+		SharedMemoryView(SharedMemoryView&&) = default;
+
+		bool Init(SharedMemoryAllocatorHandle handle, u64 capacity, void* baseAddress = nullptr, const tchar* hint = TC(""));
+		bool Init(SharedMemoryAllocatorHandle handle, u8* sharedMemory, BinaryReader& reader, u64 capacity, const tchar* hint = TC(""));
+
+		bool AddRequestedMemory(SharedMemoryAllocatorHandle handle, BinaryReader& reader, SharedMemoryMapType mapType = SharedMemoryMapType_ReadWrite, u64 offset = 0, const tchar* hint = TC(""), u8 init = 0);
+		void Reset();
+
+		u8* GetMemory() { return m_memory; }
+		u64 GetMappedSize() { return m_mappedSize; }
+		u64 GetCapacity() { return m_capacity; }
+
+		~SharedMemoryView();
+
+	private:
+		bool InternalInit(SharedMemoryAllocatorHandle handle, u64 capacity, void* baseAddress, const tchar* hint);
+		bool AddRequestedMemory(SharedMemoryAllocatorHandle handle, BinaryReader& reader, u32 sliceCount, SharedMemoryMapType mapType, u64 offset, const tchar* hint, u8 init);
+		void ReferenceExternalMemory(u8* memory, u64 mappedSize);
+		bool ExtendMemory(SharedMemoryAllocatorHandle handle, u64 offset, u64 size, SharedMemoryMapType mapType, const tchar* hint, u8 init = 0);
+
+		u8* m_memory = nullptr;
+		u64 m_mappedSize = 0;
+		u64 m_capacity = 0;
+
+		List<u64> m_additionalMappings;
+
+		SharedMemoryView(const SharedMemoryView&) = delete;
+		SharedMemoryView& operator=(const SharedMemoryView&) = delete;
+		friend class SharedMemoryAllocator;
+	};
+
+	////////////////////////////////////////////////////////////////////////////////////////////////////
+}
diff --git a/Engine/Source/Programs/UnrealBuildAccelerator/Test/Private/UbaTestAll.h b/Engine/Source/Programs/UnrealBuildAccelerator/Test/Private/UbaTestAll.h
--- a/Engine/Source/Programs/UnrealBuildAccelerator/Test/Private/UbaTestAll.h
+++ b/Engine/Source/Programs/UnrealBuildAccelerator/Test/Private/UbaTestAll.h
@@ -28,6 +28,7 @@ namespace uba
 		UBA_TEST(TestMemoryBlock) \
 		UBA_TEST(TestParseArguments) \
 		UBA_TEST(TestBinaryWriter) \
+		UBA_TEST(TestFileMappingBuffer) \
 		UBA_TEST(TestCrypto) \
 		UBA_TEST(TestSockets) \
 		UBA_TEST(TestClientServer) \
@@ -38,7 +39,6 @@ namespace uba
 		UBA_TEST(TestRemoteStorageStore) \
 		UBA_TEST(TestRemoteStorageFetch) \
 		UBA_TEST(TestRemoteStorageStore2) \
-		UBA_TEST(TestFileMappingBuffer) \
 		UBA_TEST(TestSharedEvents) \
 		UBA_TEST(TestSharedEvents2) \
 		UBA_TEST(TestDetouredTestApp) \
@@ -84,6 +84,7 @@ namespace uba
 		UBA_TEST(TestRemoteVirtualFile) \
 		UBA_TEST(TestRemoteProcessSpecialCase1) \
 		UBA_TEST(TestSessionSpecialCopy) \
+		UBA_TEST(TestSharedMemory) \
 
 
 #define UBA_LINUX_TESTS \
diff --git a/Engine/Source/Programs/UnrealBuildAccelerator/Test/Private/UbaTestSharedMemory.cpp b/Engine/Source/Programs/UnrealBuildAccelerator/Test/Private/UbaTestSharedMemory.cpp
new file mode 100644
--- /dev/null
+++ b/Engine/Source/Programs/UnrealBuildAccelerator/Test/Private/UbaTestSharedMemory.cpp
@@ -0,0 +1,77 @@
+// Copyright Epic Games, Inc. All Rights Reserved.
+
+#include "UbaLogger.h"
+#include "UbaBinaryReaderWriter.h"
+#include "UbaSharedMemoryAllocator.h"
+#include "UbaSharedMemoryView.h"
+
+namespace uba
+{
+	bool TestSharedMemory(LoggerWithWriter& logger, const StringBufferBase& rootDir)
+	{
+		{
+			SharedMemoryAllocator mem(logger);
+			mem.Init(64ull * 1024 * 1024 * 1024);
+			SharedMemoryAllocatorHandle handle = mem;
+
+			SharedMemoryHandle viewHandle1 = mem.CreateHandle(logger, TC(""));
+			SharedMemoryHandle viewHandle2 = mem.CreateHandle(logger, TC(""));
+
+			SharedMemoryView view1;
+			SharedMemoryView view2;
+
+			view1.Init(handle, 1024 * 1024);
+			view2.Init(handle, 1024 * 1024);
+
+			for (u32 i = 0; i != 3; ++i)
+			{
+				{
+					StackBinaryWriter<1024> writer;
+					mem.ExtendMemory(writer, viewHandle1, 64 * 1024);
+					BinaryReader reader(writer.GetData(), 0, writer.GetPosition());
+					view1.AddRequestedMemory(handle, reader, SharedMemoryMapType_ReadWrite, 0, TC(""), 0xff);
+				}
+				{
+					StackBinaryWriter<1024> writer;
+					mem.ExtendMemory(writer, viewHandle2, 64 * 1024);
+					BinaryReader reader(writer.GetData(), 0, writer.GetPosition());
+					view2.AddRequestedMemory(handle, reader, SharedMemoryMapType_ReadWrite, 0, TC(""), 0xdd);
+				}
+			}
+
+			u8* ptr = mem.MapView(viewHandle1, SharedMemoryMapType_ReadOnly);
+			mem.UnmapView(viewHandle1, ptr);
+
+			SharedMemoryView view3;
+			mem.MapView(view3, viewHandle1, TC(""), SharedMemoryMapType_ReadOnly);
+
+			mem.CloseHandle(logger, viewHandle2, TC(""));
+			mem.CloseHandle(logger, viewHandle1, TC(""));
+
+
+			SharedMemoryHandle viewHandle4 = mem.CreateHandle(logger, TC(""));
+			SharedMemoryView view4;
+			view4.Init(handle, 1024 * 1024);
+			mem.ExtendMemory(view4, viewHandle4, 64 * 1024, TC(""), false, mem.GetReserved() / 2);
+			mem.ExtendMemory(view4, viewHandle4, 64 * 1024, TC(""), false, mem.GetReserved() / 2 + view4.GetMappedSize());
+		}
+
+		{
+			SharedMemoryAllocator mem(logger);
+			mem.Init(1024 * 1024);
+
+			SharedMemoryHandle viewHandle1 = mem.CreateHandle(logger, TC(""));
+			SharedMemoryHandle viewHandle2 = mem.CreateHandle(logger, TC(""));
+
+			StackBinaryWriter<1024> writer1;
+			mem.ExtendMemory(writer1, viewHandle1, 64 * 1024);
+
+			StackBinaryWriter<1024> writer2;
+			mem.ExtendMemory(writer2, viewHandle2, (1024-64) * 1024);
+
+			mem.CloseHandle(logger, viewHandle1, TC(""));
+			mem.CloseHandle(logger, viewHandle2, TC(""));
+		}
+		return true;
+	}
+}
diff --git a/Engine/Source/Programs/UnrealBuildAccelerator/UbaAgent.Target.cs b/Engine/Source/Programs/UnrealBuildAccelerator/UbaAgent.Target.cs
--- a/Engine/Source/Programs/UnrealBuildAccelerator/UbaAgent.Target.cs
+++ b/Engine/Source/Programs/UnrealBuildAccelerator/UbaAgent.Target.cs
@@ -46,6 +46,7 @@ public class UbaAgentTarget : TargetRules
 		Rules.bIsBuildingConsoleApplication = true;
 
 		Rules.WindowsPlatform.TargetWindowsVersion = 0x0A00;
+		Rules.WindowsPlatform.TargetWindowsMinorVersion = 0x0A000005;
 		Rules.bUseStaticCRT = true;
 
 		Rules.SolutionDirectory = "Programs/UnrealBuildAccelerator";
