[UBA]
* Added more process and kernel stats and bumped trace file version

diff --git a/Engine/Source/Programs/UnrealBuildAccelerator/Common/Public/UbaTrace.h b/Engine/Source/Programs/UnrealBuildAccelerator/Common/Public/UbaTrace.h
--- a/Engine/Source/Programs/UnrealBuildAccelerator/Common/Public/UbaTrace.h
+++ b/Engine/Source/Programs/UnrealBuildAccelerator/Common/Public/UbaTrace.h
@@ -81,7 +81,7 @@ namespace uba
 		#undef UBA_TRACE_TYPE
 	};
 
-	static constexpr u32 TraceVersion = 49;
+	static constexpr u32 TraceVersion = 50;
 	static constexpr u32 TraceReadCompatibilityVersion = 6;
 
 	////////////////////////////////////////////////////////////////////////////////////////////////////
diff --git a/Engine/Source/Programs/UnrealBuildAccelerator/Core/Public/UbaProcessStats.h b/Engine/Source/Programs/UnrealBuildAccelerator/Core/Public/UbaProcessStats.h
--- a/Engine/Source/Programs/UnrealBuildAccelerator/Core/Public/UbaProcessStats.h
+++ b/Engine/Source/Programs/UnrealBuildAccelerator/Core/Public/UbaProcessStats.h
@@ -44,6 +44,9 @@ namespace uba
 		UBA_PROCESS_STAT(fileTable, 30) \
 		UBA_PROCESS_STAT(dirTable, 30) \
 		UBA_PROCESS_STAT(longPathName, 31) \
+		UBA_PROCESS_STAT(createSharedMemory, 50) \
+		UBA_PROCESS_STAT(commitSharedMemory, 50) \
+		UBA_PROCESS_STAT(getSharedMemory, 50) \
 
 
 		#define UBA_PROCESS_STAT(T, ver) Timer T;
@@ -145,6 +148,8 @@ namespace uba
 			UBA_KERNEL_STAT(ExtendedTimer, copyFile, 43) \
 			UBA_KERNEL_STAT(ExtendedTimer, moveFile, 43) \
 			UBA_KERNEL_STAT(ExtendedTimer, deleteFile, 43) \
+			UBA_KERNEL_STAT(TimeAndBytes, mapViewOfFile3, 50) \
+			UBA_KERNEL_STAT(TimeAndBytes, unmapViewOfFile2, 50) \
 
 		#define UBA_KERNEL_STAT(type, var, ver) type var;
 		UBA_KERNEL_STATS
