[UBA]
* Fixed bug where some kernel stats were not tracked properly in detoured processes
* Added VirtualFree and VirtualAlloc2 to kernel stats

diff --git a/Engine/Source/Programs/UnrealBuildAccelerator/Common/Private/UbaStats.cpp b/Engine/Source/Programs/UnrealBuildAccelerator/Common/Private/UbaStats.cpp
--- a/Engine/Source/Programs/UnrealBuildAccelerator/Common/Private/UbaStats.cpp
+++ b/Engine/Source/Programs/UnrealBuildAccelerator/Common/Private/UbaStats.cpp
@@ -5,6 +5,9 @@
 
 namespace uba
 {
+	KernelStats g_kernelStatsInternal;
+	KernelStats& g_kernelStats = g_kernelStatsInternal;
+
 	inline void Write(BinaryWriter& writer, u64 v) { writer.Write7BitEncoded(v); }
 	inline void Write(BinaryWriter& writer, u32 v) { writer.Write7BitEncoded(v); }
 
diff --git a/Engine/Source/Programs/UnrealBuildAccelerator/Common/Private/UbaTrace.cpp b/Engine/Source/Programs/UnrealBuildAccelerator/Common/Private/UbaTrace.cpp
--- a/Engine/Source/Programs/UnrealBuildAccelerator/Common/Private/UbaTrace.cpp
+++ b/Engine/Source/Programs/UnrealBuildAccelerator/Common/Private/UbaTrace.cpp
@@ -599,6 +599,19 @@ namespace uba
 		writer.WriteString(statusLink);
 	}
 
+	u32 Trace::StringAdd(StringView string)
+	{
+		return AddString(string);
+	}
+
+	void Trace::StatusUpdate2(u32 statusRow, u32 statusColumn, u32 stringId)
+	{
+		BEGIN_TRACE_ENTRY(StatusUpdate2);
+		writer.Write7BitEncoded(statusRow);
+		writer.Write7BitEncoded(statusColumn);
+		writer.Write7BitEncoded(stringId);
+	}
+
 	void Trace::DriveUpdate(const tchar drive, u8 busyPercent, u32 readCount, u64 readBytes, u32 writeCount, u64 writeBytes)
 	{
 		BEGIN_TRACE_ENTRY_NOTIME(DriveUpdate);
diff --git a/Engine/Source/Programs/UnrealBuildAccelerator/Common/Public/UbaTrace.h b/Engine/Source/Programs/UnrealBuildAccelerator/Common/Public/UbaTrace.h
--- a/Engine/Source/Programs/UnrealBuildAccelerator/Common/Public/UbaTrace.h
+++ b/Engine/Source/Programs/UnrealBuildAccelerator/Common/Public/UbaTrace.h
@@ -73,6 +73,8 @@ namespace uba
 		UBA_TRACE_TYPE(SchedulerUpdate) \
 		UBA_TRACE_TYPE(SchedulerKillProcess) \
 		UBA_TRACE_TYPE(SessionInfo) \
+		UBA_TRACE_TYPE(StringAdd) \
+		UBA_TRACE_TYPE(StatusUpdate2) \
 
 	enum TraceType : u8
 	{
@@ -138,6 +140,8 @@ namespace uba
 		void WorkEnd(u32 workIndex);
 		void ProgressUpdate(u32 processesTotal, u32 processesDone, u32 errorCount);
 		void StatusUpdate(u32 statusRow, u32 statusColumn, StringView statusText, LogEntryType statusType, StringView statusLink = {});
+		u32  StringAdd(StringView string);
+		void StatusUpdate2(u32 statusRow, u32 statusColumn, u32 stringId);
 		void DriveUpdate(const tchar drive, u8 busyPercent, u32 readCount, u64 readBytes, u32 writeCount, u64 writeBytes);
 		void RemoteExecutionDisabled();
 		void SchedulerUpdate(bool spawnPaused);
diff --git a/Engine/Source/Programs/UnrealBuildAccelerator/Core/Private/UbaPlatform.cpp b/Engine/Source/Programs/UnrealBuildAccelerator/Core/Private/UbaPlatform.cpp
--- a/Engine/Source/Programs/UnrealBuildAccelerator/Core/Private/UbaPlatform.cpp
+++ b/Engine/Source/Programs/UnrealBuildAccelerator/Core/Private/UbaPlatform.cpp
@@ -23,18 +23,18 @@
 
 namespace uba
 {
-	KernelStats g_kernelStatsInternal;
+	extern KernelStats& g_kernelStats;
 	thread_local KernelStats* t_kernelStats;
 
 	KernelStats& KernelStats::GetCurrent()
 	{
 		KernelStats* stats = t_kernelStats;
-		return stats ? *stats : g_kernelStatsInternal;
+		return stats ? *stats : g_kernelStats;
 	}
 
 	KernelStats& KernelStats::GetGlobal()
 	{
-		return g_kernelStatsInternal;
+		return g_kernelStats;
 	}
 
 	KernelStatsScope::KernelStatsScope(KernelStats& s) : stats(s)
diff --git a/Engine/Source/Programs/UnrealBuildAccelerator/Core/Public/UbaProcessStats.h b/Engine/Source/Programs/UnrealBuildAccelerator/Core/Public/UbaProcessStats.h
--- a/Engine/Source/Programs/UnrealBuildAccelerator/Core/Public/UbaProcessStats.h
+++ b/Engine/Source/Programs/UnrealBuildAccelerator/Core/Public/UbaProcessStats.h
@@ -148,8 +148,10 @@ namespace uba
 			UBA_KERNEL_STAT(ExtendedTimer, copyFile, 43) \
 			UBA_KERNEL_STAT(ExtendedTimer, moveFile, 43) \
 			UBA_KERNEL_STAT(ExtendedTimer, deleteFile, 43) \
+			UBA_KERNEL_STAT(ExtendedTimer, virtualFree, 50) \
+			UBA_KERNEL_STAT(ExtendedTimer, virtualAlloc2, 50) \
 			UBA_KERNEL_STAT(TimeAndBytes, mapViewOfFile3, 50) \
-			UBA_KERNEL_STAT(TimeAndBytes, unmapViewOfFile2, 50) \
+			UBA_KERNEL_STAT(ExtendedTimer, unmapViewOfFile2, 50) \
 
 		#define UBA_KERNEL_STAT(type, var, ver) type var;
 		UBA_KERNEL_STATS
diff --git a/Engine/Source/Programs/UnrealBuildAccelerator/Visualizer/Private/UbaTraceReader.cpp b/Engine/Source/Programs/UnrealBuildAccelerator/Visualizer/Private/UbaTraceReader.cpp
--- a/Engine/Source/Programs/UnrealBuildAccelerator/Visualizer/Private/UbaTraceReader.cpp
+++ b/Engine/Source/Programs/UnrealBuildAccelerator/Visualizer/Private/UbaTraceReader.cpp
@@ -1262,6 +1262,16 @@ namespace uba
 
 			break;
 		}
+		case TraceType_StatusUpdate2:
+		{
+			u64 row = reader.Read7BitEncoded();
+			u64 column = reader.Read7BitEncoded();
+			u64 key = (row << 32) | column;
+			auto& status = out.statusMap[key];
+			u64 stringIndex = reader.Read7BitEncoded();
+			status.text = out.strings[stringIndex];
+			break;
+		}
 		case TraceType_ProcessBreadcrumbs:
 		{
 			u32 processId = reader.ReadU32();
