[UBA]
* Changed so file retry on ERROR_INSUFFICIENT_RESOURCES

diff --git a/Engine/Source/Programs/UnrealBuildAccelerator/Common/Private/UbaFileAccessor.cpp b/Engine/Source/Programs/UnrealBuildAccelerator/Common/Private/UbaFileAccessor.cpp
--- a/Engine/Source/Programs/UnrealBuildAccelerator/Common/Private/UbaFileAccessor.cpp
+++ b/Engine/Source/Programs/UnrealBuildAccelerator/Common/Private/UbaFileAccessor.cpp
@@ -109,7 +109,8 @@ namespace uba
 			}
 			u32 lastError = GetLastError();
 			#if PLATFORM_WINDOWS
-			if (lastError == ERROR_SHARING_VIOLATION || lastError == ERROR_USER_MAPPED_FILE || lastError == ERROR_ACCESS_DENIED)
+			constexpr u32 ERROR_INSUFFICIENT_RESOURCES = 1450;
+			if (lastError == ERROR_SHARING_VIOLATION || lastError == ERROR_USER_MAPPED_FILE || lastError == ERROR_ACCESS_DENIED || lastError == ERROR_INSUFFICIENT_RESOURCES)
 			{
 				// Since we have unmap of files in jobs it might be that the file is not unmapped yet, therefor we should retry
 				if (retryCount == 1)
