[UBA]
* Changed so system managed page file caps at 64gb since it seems like it never becomes bigger even if I have tons of space left

diff --git a/Engine/Source/Programs/UnrealBuildAccelerator/Common/Private/UbaEnvironment.cpp b/Engine/Source/Programs/UnrealBuildAccelerator/Common/Private/UbaEnvironment.cpp
--- a/Engine/Source/Programs/UnrealBuildAccelerator/Common/Private/UbaEnvironment.cpp
+++ b/Engine/Source/Programs/UnrealBuildAccelerator/Common/Private/UbaEnvironment.cpp
@@ -206,7 +206,11 @@ namespace uba
 				}
 
 				// Max possible system-managed page file
-				maxSizeMb = Max(u64(memStatus.ullTotalPhys) * 3, 4ull * 1024 * 1024 * 1024);
+				u64 maxSize = Max(u64(memStatus.ullTotalPhys) * 3, 4ull * 1024 * 1024 * 1024);
+
+				// There seems to be something inside windows that limits the system managed pagefile..
+				// .. machines with 256gb ram and 2tb os drives still hit oom at ~326gb
+				maxSize = Min(maxSize, 64ull * 1024 * 1024 * 1024);
 
 				// Check if disk is limiting factor of system-managed page file
 				// Page file can be max 1/8 of volume size and ofc not more than free space
@@ -216,7 +220,7 @@ namespace uba
 					return logger.Error(TC("GetDiskFreeSpaceExW failed to get information about %s (%s)"), drive.data, LastErrorToText().data);
 
 				u64 maxDiskPageFileSize = Min(totalNumberOfBytes.QuadPart / 8, totalNumberOfFreeBytes.QuadPart);
-				maxPageSize += Min(maxDiskPageFileSize, maxSizeMb);
+				maxPageSize += Min(maxDiskPageFileSize, maxSize);
 			}
 		}
 		return maxPageSize;
