[UBA]
* Added method to clear event cache

diff --git a/Engine/Source/Programs/UnrealBuildAccelerator/Core/Private/UbaEvent.cpp b/Engine/Source/Programs/UnrealBuildAccelerator/Core/Private/UbaEvent.cpp
--- a/Engine/Source/Programs/UnrealBuildAccelerator/Core/Private/UbaEvent.cpp
+++ b/Engine/Source/Programs/UnrealBuildAccelerator/Core/Private/UbaEvent.cpp
@@ -178,6 +178,30 @@ namespace uba
 		#endif
 	}
 
+	void Event::ClearCache()
+	{
+		for (u32 i = 0; i != 2; ++i)
+		{
+			auto it = g_firstEvent[i];
+			g_firstEvent[i] = nullptr;;
+			while (it)
+			{
+				auto next = it->m_next;
+				if (g_useUseSemaphore)
+				{
+					((LocalEventSemaphoreImpl*)it)->Destroy();
+					delete (LocalEventSemaphoreImpl*)it;
+				}
+				else
+				{
+					((LocalEventDefaultImpl*)it)->Destroy();
+					delete (LocalEventDefaultImpl*)it;
+				}
+				it = next;
+			}
+		}
+	}
+
 	////////////////////////////////////////////////////////////////////////////////////////////////////
 
 	SharedEvent::SharedEvent()
diff --git a/Engine/Source/Programs/UnrealBuildAccelerator/Core/Public/UbaEvent.h b/Engine/Source/Programs/UnrealBuildAccelerator/Core/Public/UbaEvent.h
--- a/Engine/Source/Programs/UnrealBuildAccelerator/Core/Public/UbaEvent.h
+++ b/Engine/Source/Programs/UnrealBuildAccelerator/Core/Public/UbaEvent.h
@@ -35,7 +35,7 @@ namespace uba
 		bool IsSet(u32 timeOutMs = ~0u);
 		void* GetHandle();
 		static bool IsNative();
-
+		static void ClearCache();
 	private:
 		void* m_impl = nullptr;
 		Event(const Event&) = delete;
