[UBA]
* Changed StringBuffer AppendValue to not use sprintf since code can be called in code where sprintf can blow up on windows

diff --git a/Engine/Source/Programs/UnrealBuildAccelerator/Core/Private/UbaStringBuffer.cpp b/Engine/Source/Programs/UnrealBuildAccelerator/Core/Private/UbaStringBuffer.cpp
--- a/Engine/Source/Programs/UnrealBuildAccelerator/Core/Private/UbaStringBuffer.cpp
+++ b/Engine/Source/Programs/UnrealBuildAccelerator/Core/Private/UbaStringBuffer.cpp
@@ -232,9 +232,20 @@ namespace uba
 
 	StringBufferBase& StringBufferBase::AppendValue(u64 v)
 	{
-		tchar buf[256];
-		TSprintf_s(buf, sizeof_array(buf), TC("%llu"), v);
-		return Append(buf);
+		tchar inverse[32];
+		int i = 0;
+		do
+		{
+			inverse[i++] = '0' + (v % 10);
+			v /= 10;
+		}
+		while (v);
+
+		for (int j = 0; j < i; ++j)
+			data[count++] = inverse[i - j - 1];
+		data[count] = 0;
+
+		return *this;
 	}
 
 	StringBufferBase& ReplaceEnd(StringBufferBase& sb, const tchar* str)
