[UBA]
* Added wip support for tables of text in visualizer. Used to visualize shared memory allocator

diff --git a/Engine/Source/Programs/UnrealBuildAccelerator/Common/Private/UbaTrace.cpp b/Engine/Source/Programs/UnrealBuildAccelerator/Common/Private/UbaTrace.cpp
--- a/Engine/Source/Programs/UnrealBuildAccelerator/Common/Private/UbaTrace.cpp
+++ b/Engine/Source/Programs/UnrealBuildAccelerator/Common/Private/UbaTrace.cpp
@@ -698,6 +698,17 @@ namespace uba
 		writer.WriteBool(success);
 	}
 
+	void Trace::TableUpdate(u32 name, u32 rowCharWidth, u32 rowMaxCount, const u32* items, u32 itemCount)
+	{
+		BEGIN_TRACE_ENTRY(TableUpdate);
+		writer.Write7BitEncoded(name);
+		writer.Write7BitEncoded(rowCharWidth);
+		writer.Write7BitEncoded(rowMaxCount);
+		writer.Write7BitEncoded(itemCount);
+		for (u32 i=0; i!=itemCount; ++i)
+			writer.Write7BitEncoded(items[i]);
+	}
+
 	TraceChannel::TraceChannel(Logger& logger) : m_logger(logger)
 	{
 	}
diff --git a/Engine/Source/Programs/UnrealBuildAccelerator/Common/Public/UbaTrace.h b/Engine/Source/Programs/UnrealBuildAccelerator/Common/Public/UbaTrace.h
--- a/Engine/Source/Programs/UnrealBuildAccelerator/Common/Public/UbaTrace.h
+++ b/Engine/Source/Programs/UnrealBuildAccelerator/Common/Public/UbaTrace.h
@@ -75,6 +75,7 @@ namespace uba
 		UBA_TRACE_TYPE(SessionInfo) \
 		UBA_TRACE_TYPE(StringAdd) \
 		UBA_TRACE_TYPE(StatusUpdate2) \
+		UBA_TRACE_TYPE(TableUpdate) \
 
 	enum TraceType : u8
 	{
@@ -157,6 +158,8 @@ namespace uba
 		void TaskHint(u32 processId, StringView hint);
 		void TaskEnd(u32 processId, bool success);
 
+		void TableUpdate(u32 name, u32 rowCharWidth, u32 rowMaxCount, const u32* items, u32 itemCount);
+
 		virtual u32 TrackWorkStart(const StringView& desc, const Color& color) final override;
 		virtual void TrackWorkHint(u32 id, const StringView& hint, u64 startTime = 0) final override;
 		virtual void TrackWorkEnd(u32 id) final override;
diff --git a/Engine/Source/Programs/UnrealBuildAccelerator/Visualizer/Private/UbaTraceReader.cpp b/Engine/Source/Programs/UnrealBuildAccelerator/Visualizer/Private/UbaTraceReader.cpp
--- a/Engine/Source/Programs/UnrealBuildAccelerator/Visualizer/Private/UbaTraceReader.cpp
+++ b/Engine/Source/Programs/UnrealBuildAccelerator/Visualizer/Private/UbaTraceReader.cpp
@@ -65,6 +65,7 @@ namespace uba
 		strings.clear();
 		statusMap.clear();
 		cacheWrites.clear();
+		tables.clear();
 		activeProcessCounts.clear();
 		maxActiveProcessCount = 0;
 		startTime = 0;
@@ -1270,6 +1271,19 @@ namespace uba
 			auto& status = out.statusMap[key];
 			u64 stringIndex = reader.Read7BitEncoded();
 			status.text = out.strings[stringIndex];
+			status.type = LogEntryType_Info;
+			break;
+		}
+		case TraceType_TableUpdate:
+		{
+			u32 name = u32(reader.Read7BitEncoded());
+			TraceView::Table& table = out.tables[name];
+			table.rowCharWidth = u32(reader.Read7BitEncoded());
+			table.rowMaxCount = u32(reader.Read7BitEncoded());
+			u32 itemCount = u32(reader.Read7BitEncoded());
+			table.items.resize(itemCount);
+			for (u32 i = 0; i != itemCount; ++i)
+				table.items[i] = out.strings[reader.Read7BitEncoded()];
 			break;
 		}
 		case TraceType_ProcessBreadcrumbs:
diff --git a/Engine/Source/Programs/UnrealBuildAccelerator/Visualizer/Private/UbaTraceReader.h b/Engine/Source/Programs/UnrealBuildAccelerator/Visualizer/Private/UbaTraceReader.h
--- a/Engine/Source/Programs/UnrealBuildAccelerator/Visualizer/Private/UbaTraceReader.h
+++ b/Engine/Source/Programs/UnrealBuildAccelerator/Visualizer/Private/UbaTraceReader.h
@@ -182,6 +182,13 @@ namespace uba
 			Vector<TString> lines;
 		};
 
+		struct Table
+		{
+			Vector<const tchar*> items;
+			u32 rowCharWidth;
+			u32 rowMaxCount;
+		};
+
 		Vector<Session> sessions;
 		Vector<CacheSummary> cacheSummaries;
 		Vector<WorkTrack> workTracks;
@@ -189,6 +196,7 @@ namespace uba
 		Vector<ActiveProcessCount> activeProcessCounts;
 		Map<u64, StatusUpdate> statusMap;
 		Map<u32, CacheWrite> cacheWrites;
+		Map<u32, Table> tables;
 		u64 realStartTime = 0;
 		u64 traceSystemStartTimeUs = 0;
 		u64 startTime = 0;
diff --git a/Engine/Source/Programs/UnrealBuildAccelerator/Visualizer/Private/UbaVisualizer.cpp b/Engine/Source/Programs/UnrealBuildAccelerator/Visualizer/Private/UbaVisualizer.cpp
--- a/Engine/Source/Programs/UnrealBuildAccelerator/Visualizer/Private/UbaVisualizer.cpp
+++ b/Engine/Source/Programs/UnrealBuildAccelerator/Visualizer/Private/UbaVisualizer.cpp
@@ -1964,6 +1964,38 @@ namespace uba
 			SetActiveFont(m_defaultFont);
 		}
 
+		if (true)
+		{
+			for (auto& kv : m_traceView.tables)
+			{
+				auto& table = kv.second;
+				u32 tableOffsetX = m_activeFont.height*40;
+				u32 tableOffsetY = 30;
+
+				u32 rowOffset = m_activeFont.height + 2;
+				u32 columnOffset = table.rowCharWidth * m_activeFont.height;
+
+				u32 columnIndex = 0;
+				u32 rowIndex = 0;
+				for (const tchar* item : table.items)
+				{
+					u32 y = tableOffsetY + rowIndex*rowOffset;
+					u32 x = tableOffsetX + columnIndex*columnOffset;
+					if (int(x) >= progressRect.right)
+						break;
+
+					ExtTextOutW(hdc, x, y, ETO_CLIPPED, nullptr, item, TStrlen(item), NULL);
+
+					++rowIndex;
+					if (rowIndex > table.rowMaxCount || int(y) >= progressRect.bottom)
+					{
+						rowIndex = 0;
+						++columnIndex;
+					}
+				}
+			}
+		}
+
 		SelectObject(textDC, oldBmp);
 		DeleteObject(nullBmp);
 		DeleteDC(textDC);
