[UBA]
* Added libs with absolute paths to link dependency crawler. This is to prevent .lib files from being evicted from memory after being written

diff --git a/Engine/Source/Programs/UnrealBuildAccelerator/Common/Public/UbaDependencyCrawler.h b/Engine/Source/Programs/UnrealBuildAccelerator/Common/Public/UbaDependencyCrawler.h
--- a/Engine/Source/Programs/UnrealBuildAccelerator/Common/Public/UbaDependencyCrawler.h
+++ b/Engine/Source/Programs/UnrealBuildAccelerator/Common/Public/UbaDependencyCrawler.h
@@ -566,6 +566,8 @@ namespace uba
 					{
 						if (sb.EndsWith(TCV(".obj")))
 							AddCodeFile(sb, false, false);
+						if (sb.Contains('/') && sb.EndsWith(TCV(".lib")))
+							AddCodeFile(sb, false, false);
 					}
 				});
 		}
