[UBA]
* Added TrackSystemUsage to test code to be able to track if new windows handles have leaked

diff --git a/Engine/Source/Programs/UnrealBuildAccelerator/Test/Private/UbaTest.cpp b/Engine/Source/Programs/UnrealBuildAccelerator/Test/Private/UbaTest.cpp
--- a/Engine/Source/Programs/UnrealBuildAccelerator/Test/Private/UbaTest.cpp
+++ b/Engine/Source/Programs/UnrealBuildAccelerator/Test/Private/UbaTest.cpp
@@ -49,6 +49,25 @@ namespace uba
 		AddExceptionHandler();
 		return RunTests(argc, argv);
 	}
+
+	TrackSystemUsage::TrackSystemUsage(Logger& l) : logger(l)
+	{
+		Event::ClearCache();
+		#if PLATFORM_WINDOWS
+		GetProcessHandleCount(GetCurrentProcess(), (DWORD*)&startHandleCount);
+		#endif
+	}
+
+	TrackSystemUsage::~TrackSystemUsage()
+	{
+		Event::ClearCache();
+		#if PLATFORM_WINDOWS
+		u32 endHandleCount = 0;
+		GetProcessHandleCount(GetCurrentProcess(), (DWORD*)&endHandleCount);
+		if (startHandleCount != endHandleCount)
+			logger.Error(TC("Handles were leaked while running test (%u)"), endHandleCount - startHandleCount);
+		#endif
+	}
 }
 
 #if PLATFORM_WINDOWS
diff --git a/Engine/Source/Programs/UnrealBuildAccelerator/Test/Private/UbaTest.h b/Engine/Source/Programs/UnrealBuildAccelerator/Test/Private/UbaTest.h
--- a/Engine/Source/Programs/UnrealBuildAccelerator/Test/Private/UbaTest.h
+++ b/Engine/Source/Programs/UnrealBuildAccelerator/Test/Private/UbaTest.h
@@ -30,4 +30,12 @@ namespace uba
 	void GetTestAppPath(LoggerWithWriter& logger, StringBufferBase& out);
 	const tchar* GetSystemApplication();
 	const tchar* GetSystemArguments();
+
+	struct TrackSystemUsage
+	{
+		TrackSystemUsage(Logger& logger);
+		~TrackSystemUsage();
+		Logger& logger;
+		u32 startHandleCount;
+	};
 }
\ No newline at end of file
