[UBA]
* Improved assert description
* Moved PAGE_READONLY posix define to UbaPlatform.h

diff --git a/Engine/Source/Programs/UnrealBuildAccelerator/Common/Private/UbaFile.cpp b/Engine/Source/Programs/UnrealBuildAccelerator/Common/Private/UbaFile.cpp
--- a/Engine/Source/Programs/UnrealBuildAccelerator/Common/Private/UbaFile.cpp
+++ b/Engine/Source/Programs/UnrealBuildAccelerator/Common/Private/UbaFile.cpp
@@ -116,7 +116,7 @@ namespace uba
 			DWORD wasRead = 0;
 			if (!::ReadFile(asHANDLE(fileHandle), buffer, toRead, &wasRead, NULL))
 				if (GetLastError() != ERROR_IO_PENDING)
-					return logger.Error(TC("ERROR reading %llu bytes from file %s (error: %s)"), toRead, fileName, LastErrorToText().data);
+					return logger.Error(TC("ERROR reading %llu bytes from file %s with handle %llu (error: %s)"), toRead, fileName, asHANDLE(fileHandle), LastErrorToText().data);
 #else
 			ssize_t wasRead = read(asFileDescriptor(fileHandle), buffer, toRead);
 			if (wasRead == -1)
diff --git a/Engine/Source/Programs/UnrealBuildAccelerator/Common/Public/UbaFile.h b/Engine/Source/Programs/UnrealBuildAccelerator/Common/Public/UbaFile.h
--- a/Engine/Source/Programs/UnrealBuildAccelerator/Common/Public/UbaFile.h
+++ b/Engine/Source/Programs/UnrealBuildAccelerator/Common/Public/UbaFile.h
@@ -33,7 +33,6 @@ namespace uba
 
 	inline constexpr u32 CREATE_ALWAYS = 2;
 	inline constexpr u32 OPEN_EXISTING = 3;
-	inline constexpr u32 PAGE_READONLY = 0x02;
 	#endif
 	inline constexpr u32 FILE_SHARE_ALL = 0x00000007;
 
diff --git a/Engine/Source/Programs/UnrealBuildAccelerator/Core/Public/UbaPlatform.h b/Engine/Source/Programs/UnrealBuildAccelerator/Core/Public/UbaPlatform.h
--- a/Engine/Source/Programs/UnrealBuildAccelerator/Core/Public/UbaPlatform.h
+++ b/Engine/Source/Programs/UnrealBuildAccelerator/Core/Public/UbaPlatform.h
@@ -116,6 +116,8 @@ namespace uba
 	bool UseSemaphoreForLocalEvents();
 	bool UseSemaphoreForSharedEvents();
 
+	inline constexpr u64 PageSize = 64 * 1024;
+
 #if PLATFORM_WINDOWS
 	inline constexpr bool CaseInsensitiveFs = true;
 	inline constexpr tchar PathSeparator = '\\';
@@ -183,6 +185,8 @@ namespace uba
 	#define UBA_NOINLINE __attribute__((noinline))
     #define ERROR_SUCCESS 0
 	#define DUPLICATE_SAME_ACCESS 0
+	inline constexpr u32 PAGE_READONLY = 0x02;
+
 	struct FileMappingHandle
 	{
 		FileMappingHandle(int shmFd_ = -1, int lockFd_ = -1, u64 uid_ = ~u64(0)) : shmFd(shmFd_), lockFd(lockFd_), uid(uid_) {};
diff --git a/Engine/Source/Programs/UnrealBuildAccelerator/Core/Public/UbaProcessUtils.h b/Engine/Source/Programs/UnrealBuildAccelerator/Core/Public/UbaProcessUtils.h
--- a/Engine/Source/Programs/UnrealBuildAccelerator/Core/Public/UbaProcessUtils.h
+++ b/Engine/Source/Programs/UnrealBuildAccelerator/Core/Public/UbaProcessUtils.h
@@ -70,8 +70,8 @@ namespace uba
 					u64 writtenCount = argWriteIt - argWriteBuffer;
 					if (writtenCount + 1 == argBufferSize)
 					{
-						UBA_ASSERT(argWriteBuffer == temp);
-						argBufferSize = argumentsLen + 1;
+						UBA_ASSERTF(argWriteBuffer == temp, TC("argWriteBuffer already replaced for size %llu, and now we need %llu"), argBufferSize, writtenCount + 1);
+						argBufferSize = argumentsLen + 16;
 						argWriteBuffer = new CharType[argBufferSize];
 						argWriteIt = argWriteBuffer + writtenCount;
 						memcpy(argWriteBuffer, temp, writtenCount *sizeof(CharType));
