[UBA - LARGE CHANGE]
* Changed entire uba to use shared memory allocator for all file communication. Session process now owns all memory and detoured processes requests memory. Host keeps memory mapped and can easily access all written pages without having to map/unmap all the time. This improves build times a lot for powerful computers such as Threadrippers but also affects the build farm noticeably. It also removes all stutters that has been seen on the new Threadripper pro 7985WX64. This change should reduce the amount of kernel calls related to I/O and page file

diff --git a/Engine/Source/Programs/UnrealBuildAccelerator/Common/Private/UbaCacheClient.cpp b/Engine/Source/Programs/UnrealBuildAccelerator/Common/Private/UbaCacheClient.cpp
--- a/Engine/Source/Programs/UnrealBuildAccelerator/Common/Private/UbaCacheClient.cpp
+++ b/Engine/Source/Programs/UnrealBuildAccelerator/Common/Private/UbaCacheClient.cpp
@@ -1353,13 +1353,13 @@ namespace uba
 
 				FileAccessor file(m_logger, casKeyFileName.data);
 
-				auto mh = casEntry->mappingHandle;
+				auto mh = casEntry->memoryHandle;
 				if (mh.IsValid()) // If file was created by helper it will be in the transient mapped memory
 				{
-					u64 mappingOffset = casEntry->mappingOffset;
-					u64 mappingSize = casEntry->mappingSize;
+					u64 memoryOffset = casEntry->memoryOffset;
+					u64 memorySize = casEntry->memorySize;
 
-					mappedView = m_storage.m_casDataBuffer.MapView(mh, mappingOffset, mappingSize, path.data, casEntry->mappingIsIndependent);
+					mappedView = m_storage.m_casDataBuffer.MapView(mh, memoryOffset, memorySize, path.data);
 					fileData = mappedView.memory;
 					fileSize = mappedView.size;
 				}
@@ -1390,7 +1390,7 @@ namespace uba
 				FileAccessor file(m_logger, path.data);
 				if (!file.OpenMemoryRead())
 					return false;
-				MemoryBlock block(AlignUp(file.GetSize() + 16, 64*1024));
+				MemoryBlock block(AlignUp(file.GetSize() + 16, PageSize));
 				u32& rootOffsetsStart = *(u32*)block.Allocate(sizeof(u32), 1, TC(""));
 				rootOffsetsStart = 0;
 				Vector<u32> rootOffsets;
diff --git a/Engine/Source/Programs/UnrealBuildAccelerator/Common/Private/UbaFileMapping.cpp b/Engine/Source/Programs/UnrealBuildAccelerator/Common/Private/UbaFileMapping.cpp
--- a/Engine/Source/Programs/UnrealBuildAccelerator/Common/Private/UbaFileMapping.cpp
+++ b/Engine/Source/Programs/UnrealBuildAccelerator/Common/Private/UbaFileMapping.cpp
@@ -3,18 +3,12 @@
 #include "UbaFileMapping.h"
 #include "UbaBottleneck.h"
 #include "UbaDirectoryIterator.h"
-#include "UbaFile.h"
-#include "UbaPlatform.h"
 #include "UbaProcessStats.h"
-#include "UbaTimer.h"
-#include "UbaWorkManager.h"
 
 #if !PLATFORM_WINDOWS
 #include <sys/file.h>
 #endif
 
-#define UBA_DEBUG_FILE_MAPPING 0
-
 ////////////////////////////////////////////////////////////////////////////////////////////////////
 
 namespace uba
@@ -371,7 +365,7 @@ namespace uba
 		if (h.uid == ~u64(0))
 			return true;
 		if (close(h.shmFd) != 0)
-			UBA_ASSERT(false);
+			UBA_ASSERTF(false, "Failed to close file mapping (%s)", strerror(errno));
 
 		StringBuffer<64> uidName;
 		GetMappingHandleName(uidName, h.uid);
@@ -460,448 +454,30 @@ namespace uba
 
 	////////////////////////////////////////////////////////////////////////////////////////////////////
 
-	FileMappingBuffer::FileMappingBuffer(Logger& logger, WorkManager* workManager)
-	:	m_logger(logger)
-	,	m_workManager(workManager)
-	{
-		m_pageSize = 64*1024;
-	}
-
-	FileMappingBuffer::~FileMappingBuffer()
-	{
-		CloseMappingStorage(m_storage[MappedView_Transient]);
-		CloseMappingStorage(m_storage[MappedView_Persistent]);
-	}
-
-	u64 FileMappingBuffer::GetFileMappingCapacity()
-	{
-		return (IsWindows ? 32ull : 8ull) * 1024 * 1024 * 1024; // Linux can't have larger than 8gb
-	}
-
-	bool FileMappingBuffer::AddTransient(const tchar* name, bool keepMapped)
-	{
-		MappingStorage& storage = m_storage[MappedView_Transient];
-		File** prev = &storage.availableFile;
-		for (u32 i = 0; i != sizeof_array(storage.files); ++i)
-		{
-			File& file = storage.files[i];
-			file.name = name;
-			file.keepMapped = keepMapped;
-			*prev = &file;
-			prev = &file.next;
-		}
-		storage.fileCount = sizeof_array(storage.files);
-		return true;
-	}
-
-	bool FileMappingBuffer::AddPersistent(const tchar* name, FileHandle fileHandle, u64 size, u64 capacity)
-	{
-		FileMappingHandle sparseMemoryHandle = uba::CreateFileMappingW(m_logger, fileHandle, PAGE_READWRITE, capacity, name);
-		if (!sparseMemoryHandle.IsValid())
-		{
-			m_logger.Error(TC("Failed to create file mapping (%s)"), LastErrorToText().data);
-			return false;
-		}
-
-		MappingStorage& storage = m_storage[MappedView_Persistent];
-		File& f = storage.files[storage.fileCount++];
-		f.name = name;
-		f.file = fileHandle;
-		f.handle = sparseMemoryHandle;
-		f.size = size;
-		f.capacity = capacity;
-
-		PushFile(storage, &f);
-		return true;
-	}
-
-	void FileMappingBuffer::CloseDatabase()
-	{
-		CloseMappingStorage(m_storage[MappedView_Persistent]);
-	}
-
-	MappedView FileMappingBuffer::AllocAndMapView(FileMappingType type, u64 size, u64 alignment, const tchar* hint, bool allowShrink, bool isIndependent)
-	{
-		if (isIndependent)
-		{
-			UBA_ASSERT(!allowShrink);
-			FileMappingHandle mapping = CreateMemoryMappingW(m_logger, PAGE_READWRITE, size, nullptr, hint);
-			if (!mapping.IsValid())
-				return {};
-			u8* memory = MapViewOfFile(m_logger, mapping, FILE_MAP_WRITE, 0, size, true);
-			if (!memory)
-			{
-				CloseFileMapping(m_logger, mapping, hint);
-				return {};
-			}
-			MappedView res;
-			res.handle = mapping;
-			res.memory = memory;
-			res.size = size;
-			res.isIndependent = true;
-			return res;
-		}
-
-		MappingStorage& storage = m_storage[type];
-		File* file = PopFile(storage, size, alignment);
-		if (!file)
-			return {};
-
-		if (allowShrink)
-			return AllocAndMapViewNoLock(*file, size, alignment, hint);
-
-		MappedView res = AllocAndMapViewNoLock(*file, size, alignment, hint);
-
-		PushFile(storage, file);
-		return res;
-	}
-
-	MappedView FileMappingBuffer::AllocAndMapViewNoLock(File& f, u64 size, u64 alignment, const tchar* hint)
-	{
-		MappedView res;
-
-		u64 offset = AlignUp(f.size, alignment);
-		u64 alignedOffsetStart = AlignUp(offset - (m_pageSize - 1), m_pageSize);
-
-		u64 newOffset = offset + size;
-		u64 alignedOffsetEnd = AlignUp(newOffset, m_pageSize);
-		
-		if (alignedOffsetEnd > f.capacity)
-		{
-			m_logger.Error(TC("%s - AllocAndMapView has reached max capacity %llu trying to allocate %llu for %s"), f.name, f.capacity, size, hint);
-			return res;
-		}
-
-		u64 mapSize = alignedOffsetEnd - alignedOffsetStart;
-		u8* data;
-		if (f.mappedMemory)
-		{
-			data = f.mappedMemory + alignedOffsetStart;
-		}
-		else
-		{
-			data = MapViewOfFile(m_logger, f.handle, FILE_MAP_WRITE, alignedOffsetStart, mapSize);
-			if (!data)
-			{
-				m_logger.Error(TC("%s - AllocAndMapView failed to map view of file for %s with size %llu and offset %llu (%s)"), f.name, hint, size, alignment, LastErrorToText().data);
-				return res;
-			}
-		}
-
-		u64 committedBefore = AlignUp(offset, m_pageSize);
-		u64 committedAfter = AlignUp(newOffset, m_pageSize);
-
-		if (f.commitOnAlloc && committedBefore != committedAfter)
-		{
-			u64 commitStart = committedBefore - alignedOffsetStart;
-			u64 commitSize = committedAfter - committedBefore;
-			if (!MapViewCommit(m_logger, data + commitStart, commitSize, true))
-			{
-				UnmapViewOfFile(m_logger, data, mapSize, hint);
-				m_logger.Error(TC("%s - Failed to allocate/map %s for %s (%s)"), f.name, BytesToText(commitSize).str, hint, LastErrorToText().data);
-				return res;
-			}
-
-			PrefetchVirtualMemory(data + commitStart, commitSize);
-		}
-
-		f.size = newOffset;
-		//++f.activeMapCount;
-
-		res.handle = f.handle;
-		res.offset = offset;
-		res.size = size;
-		res.memory = data + (offset - alignedOffsetStart);
-		return res;
-	}
-
-	FileMappingBuffer::File& FileMappingBuffer::GetFile(FileMappingHandle handle, u8& outStorageIndex)
-	{
-		for (u8 storageI = 0; storageI != 2; ++storageI)
-		{
-			MappingStorage& storage = m_storage[storageI];
-			for (u32 i = 0; i != storage.fileCount; ++i)
-				if (storage.files[i].handle == handle)
-				{
-					outStorageIndex = storageI;
-					return storage.files[i];
-				}
-		}
-		UBA_ASSERT(false);
-		static FileMappingBuffer::File error;
-		return error;
-	}
-
-	FileMappingBuffer::File* FileMappingBuffer::PopFile(MappingStorage& storage, u64 size, u64 alignment)
-	{
-		while (true)
-		{
-			SCOPED_FUTEX(storage.availableFilesLock, lock);
-			File* file = storage.availableFile;
-			if (!file)
-			{
-				if (storage.fullFileCount == storage.fileCount)
-				{
-					m_logger.Error(TC("All files are full!"));
-					return nullptr;
-				}
-				lock.Leave();
-				storage.availableFilesEvent.IsSet();
-				continue;
-			}
-			storage.availableFile = file->next;
-			lock.Leave();
-
-			if (!file->handle.IsValid())
-			{
-				UBA_ASSERT(&storage == &m_storage[MappedView_Transient]);
-				u64 capacity = GetFileMappingCapacity();
-				file->handle = uba::CreateMemoryMappingW(m_logger, PAGE_READWRITE | SEC_RESERVE, capacity, nullptr, TC("FileMappingBuffer"));
-				if (!file->handle.IsValid())
-				{
-					m_logger.Error(TC("%s - Failed to create memory map (%s)"), file->name, LastErrorToText().data);
-					return nullptr;
-				}
-				file->commitOnAlloc = true;
-				file->capacity = capacity;
-
-				if (file->keepMapped)
-					file->mappedMemory = MapViewOfFile(m_logger, file->handle, FILE_MAP_WRITE, 0, capacity);
-
-			}
-			else
-			{
-				u64 newSize = AlignUp(file->size, alignment) + size;
-				if (newSize > file->capacity)
-				{
-					lock.Enter();
-					++storage.fullFileCount; // We will never add this file back to the available file list
-					continue;
-				}
-			}
-			return file;
-		}
-		return nullptr;
-	}
-
-	void FileMappingBuffer::PushFile(MappingStorage& storage, File* file)
-	{
-		SCOPED_FUTEX(storage.availableFilesLock, lock);
-		file->next = storage.availableFile;
-		storage.availableFile = file;
-		storage.availableFilesEvent.Set();
-	}
-
-	void FileMappingBuffer::CloseMappingStorage(MappingStorage& storage)
-	{
-		SCOPED_FUTEX(storage.availableFilesLock, lock);
-		u32 filesTaken = storage.fullFileCount;
-		while (true)
-		{
-			while (storage.availableFile)
-			{
-				storage.availableFile = storage.availableFile->next;
-				++filesTaken;
-			}
-			lock.Leave();
-			if (filesTaken == storage.fileCount)
-				break;
-			storage.availableFilesEvent.IsSet();
-			lock.Enter();
-		}
-
-		for (u32 i=0; i!=storage.fileCount; ++i)
-		{
-			File& file = storage.files[i];
-			if (file.mappedMemory)
-			{
-				UnmapViewOfFile(m_logger, file.mappedMemory, file.capacity, TC("FileMappingBuffer"));
-				file.mappedMemory = nullptr;
-			}
-			//UBA_ASSERT(!file.activeMapCount);
-			CloseFileMapping(m_logger, file.handle, TC("FileMappingBuffer"));
-			CloseFile(nullptr, file.file);
-		}
-		storage.fileCount = 0;
-		storage.availableFile = nullptr;
-	}
-
-	MappedView FileMappingBuffer::MapView(FileMappingHandle handle, u64 offset, u64 size, const tchar* hint, bool isIndependent)
-	{
-		return MapView(handle, offset, size, hint, isIndependent, FILE_MAP_READ);
-	}
-
-	MappedView FileMappingBuffer::MapViewWrite(FileMappingHandle handle, u64 offset, u64 size, const tchar* hint, bool isIndependent)
-	{
-		return MapView(handle, offset, size, hint, isIndependent, FILE_MAP_WRITE);
-	}
-
-	MappedView FileMappingBuffer::MapView(FileMappingHandle handle, u64 offset, u64 size, const tchar* hint, bool isIndependent, u32 desiredAccess)
-	{
-		UBA_ASSERT(handle.IsValid());// && (handle == m_files[0].handle || handle == m_files[1].handle));
-
-		MappedView res;
-
-		if (isIndependent)
-		{
-			u8* data = MapViewOfFile(m_logger, handle, desiredAccess, offset, size);
-			if (!data)
-				return {};
-			res.handle = handle;
-			res.offset = offset;
-			res.size = size;
-			res.memory = data;
-			res.isIndependent = true;
-			return res;
-		}
-
-		u8 storageIndex = 255;
-		File& file = GetFile(handle, storageIndex);
-
-		u64 alignedOffsetStart = 0;
-
-		u8* data;
-		if (file.mappedMemory)
-		{
-			data = file.mappedMemory;
-		}
-		else
-		{
-			alignedOffsetStart = AlignUp(offset - (m_pageSize - 1), m_pageSize);
-			u64 newOffset = offset + size;
-			u64 alignedOffsetEnd = AlignUp(newOffset, m_pageSize);
-			data = MapViewOfFile(m_logger, handle, desiredAccess, alignedOffsetStart, alignedOffsetEnd - alignedOffsetStart);
-			if (!data)
-			{
-				m_logger.Error(TC("%s - MapView failed to map view of file for %s with size %llu and offset %llu (%s)"), file.name, hint, size, offset, LastErrorToText().data);
-				return res;
-			}
-		}
-
-		//++file.activeMapCount;
-
-		res.handle = handle;
-		res.offset = offset;
-		res.size = size;
-		res.memory = data + (offset - alignedOffsetStart);
-		return res;
-	}
-
-	MappedView FileMappingBuffer::MapView(const StringView& str, u64 size, const tchar* hint)
-	{
-#if PLATFORM_WINDOWS
-		const tchar* handleStr = str.data + 1;
-		const tchar* handleStrEnd = TStrchr(handleStr, '-');
-		u64 mappingOffset = 0;
-		if (handleStrEnd)
-		{
-			const tchar* mappingOffsetStr = handleStrEnd + 1;
-			mappingOffset = StringToValueBase62(mappingOffsetStr, wcslen(mappingOffsetStr));
-		}
-		else
-			handleStrEnd = handleStr + wcslen(handleStr);
-
-		FileMappingHandle h = FileMappingHandle::FromU64(StringToValueBase62(handleStr, handleStrEnd - handleStr));
-		
-		return MapView(h, mappingOffset, size, hint);
-#else
-		return {};
-#endif
-	}
-
-	void FileMappingBuffer::UnmapView(MappedView view, const tchar* hint_, u64 newSize, bool allowDeferredUnmap)
-	{
-		if (!view.handle.IsValid())
-			return;
-
-		auto unmap = [=, this](const tchar* hint)
-			{
-				if (view.isIndependent)
-				{
-					if (!UnmapViewOfFile(m_logger, view.memory, view.size, hint))
-						m_logger.Error(TC("Failed to unmap view on address %llx (offset %llu) - %s (%s)"), u64(view.memory), view.offset, hint, LastErrorToText().data);
-					return;
-				}
-
-				u8 storageIndex = 255;
-				File& file = GetFile(view.handle, storageIndex);
-
-				u64 alignedOffsetStart = AlignUp(view.offset - (m_pageSize - 1), m_pageSize);
-				u64 alignedOffsetEnd = AlignUp(view.offset + view.size, m_pageSize);
-
-				u8* memory = view.memory - (view.offset - alignedOffsetStart);
-				u64 mapSize = alignedOffsetEnd - alignedOffsetStart;
-				if (!file.mappedMemory)
-					if (!UnmapViewOfFile(m_logger, memory, mapSize, hint))
-						m_logger.Error(TC("%s - Failed to unmap view on address %llx (offset %llu) - %s (%s)"), file.name, u64(memory), view.offset, hint, LastErrorToText().data);
-
-				if (newSize != InvalidValue)
-				{
-					if (newSize != view.size)
-					{
-						UBA_ASSERT(!file.commitOnAlloc);
-						UBA_ASSERTF(newSize < view.size, TC("%s - Reserved too little memory. Reserved %llu, needed %llu for %s"), file.name, view.size, newSize, hint);
-						file.size -= view.size - newSize;
-					}
-
-					MappingStorage& storage = m_storage[storageIndex];
-					PushFile(storage, &file);
-				}
-
-				//--file.activeMapCount;
-			};
-
-		#if UBA_DEBUG_FILE_MAPPING
-		allowDeferredUnmap = false;
-		#endif
-
-		if (allowDeferredUnmap && m_workManager && newSize == InvalidValue)
-			m_workManager->AddWork([=, h = TString(hint_)](const WorkContext&) { unmap(h.c_str()); }, 1, TC("UnmapView"));
-		else
-			unmap(hint_);
-	}
-
-	void FileMappingBuffer::GetSizeAndCount(FileMappingType type, u64& outSize, u32& outCount)
-	{
-		MappingStorage& storage = m_storage[type];
-		outSize = 0;
-		outCount = 0;
-		for (u32 i = 0; i != storage.fileCount; ++i)
-		{
-			if (storage.files[i].handle.IsValid())
-				++outCount;
-			outSize += storage.files[i].size;
-		}
-	}
-
-
-	////////////////////////////////////////////////////////////////////////////////////////////////////
-
-	FileMappingAllocator::FileMappingAllocator(Logger& logger, const tchar* name)
+	FixedSizeFileMappingAllocator::FixedSizeFileMappingAllocator(Logger& logger, const tchar* name)
 	:	m_logger(logger)
 	,	m_name(name)
 	{
 	}
 
-	FileMappingAllocator::~FileMappingAllocator()
+	FixedSizeFileMappingAllocator::~FixedSizeFileMappingAllocator()
 	{
 		if (m_mappingHandle.IsValid())
-			CloseFileMapping(m_logger, m_mappingHandle, TC("FileMappingAllocator"));
+			CloseFileMapping(m_logger, m_mappingHandle, TC("FixedSizeFileMappingAllocator"));
 	}
 
-	bool FileMappingAllocator::Init(u64 blockSize, u64 capacity)
+	bool FixedSizeFileMappingAllocator::Init(u64 blockSize, u64 capacity)
 	{
-		m_mappingHandle = uba::CreateMemoryMappingW(m_logger, PAGE_READWRITE|SEC_RESERVE, capacity, nullptr, TC("FileMappingAllocator"));
+		m_mappingHandle = uba::CreateMemoryMappingW(m_logger, PAGE_READWRITE|SEC_RESERVE, capacity, nullptr, TC("FixedSizeFileMappingAllocator"));
 		if (!m_mappingHandle.IsValid())
 			return m_logger.Error(TC("%s - Failed to create memory map with capacity %llu (%s)"), m_name, capacity, LastErrorToText().data);
 
 		m_blockSize = blockSize;
-		m_pageSize = 64*1024;
 		m_capacity = capacity;
 		return true;
 	}
 
-	FileMappingAllocator::Allocation FileMappingAllocator::Alloc(const tchar* hint)
+	FileMappingAllocation FixedSizeFileMappingAllocator::Alloc(const tchar* hint)
 	{
 		SCOPED_FUTEX(m_mappingLock, lock);
 
@@ -944,7 +520,7 @@ namespace uba
 		return {m_mappingHandle, offset, data};
 	}
 
-	void FileMappingAllocator::Free(Allocation allocation)
+	void FixedSizeFileMappingAllocator::Free(FileMappingAllocation allocation)
 	{
 		UBA_ASSERT(allocation.handle == m_mappingHandle);
 		if (!UnmapViewOfFile(m_logger, allocation.memory, m_blockSize, m_name, true))
@@ -953,4 +529,6 @@ namespace uba
 		SCOPED_FUTEX(m_mappingLock, lock);
 		m_availableBlocks.insert(index);
 	}
+
+	////////////////////////////////////////////////////////////////////////////////////////////////////
 }
diff --git a/Engine/Source/Programs/UnrealBuildAccelerator/Common/Private/UbaFileMappingBuffer.cpp b/Engine/Source/Programs/UnrealBuildAccelerator/Common/Private/UbaFileMappingBuffer.cpp
new file mode 100644
--- /dev/null
+++ b/Engine/Source/Programs/UnrealBuildAccelerator/Common/Private/UbaFileMappingBuffer.cpp
@@ -0,0 +1,177 @@
+// Copyright Epic Games, Inc. All Rights Reserved.
+
+#include "UbaFileMappingBuffer.h"
+#include "UbaSharedMemoryAllocator.h"
+#include "UbaTrace.h"
+
+namespace uba
+{
+	////////////////////////////////////////////////////////////////////////////////////////////////////
+
+	FileMappingBuffer::FileMappingBuffer(Logger& logger, SharedMemoryAllocator& allocator)
+		: m_logger(logger)
+		, m_allocator(allocator)
+	{
+	}
+
+	FileMappingBuffer::~FileMappingBuffer()
+	{
+		if (!m_memoryHandle.IsValid())
+			return;
+		m_allocator.DecommitMemory(m_memoryHandle, 0, m_committed);
+		m_allocator.UnreserveMemory(m_memoryHandle);
+		m_allocator.CloseHandle(m_logger, m_memoryHandle, m_name.c_str());
+	}
+
+	void FileMappingBuffer::TraceStatus(Trace& trace, u32 startRow, const tchar* hint, u32 traceLevel)
+	{
+		if (!traceLevel)
+			return;
+
+		trace.StatusUpdate(startRow, 1, m_name, LogEntryType_Info);
+		StringBuffer<> str;
+
+		str.Appendf(TC("Mem: %s/%s"), BytesToText(m_used).str, BytesToText(m_committed).str);
+		str.Appendf(TC(", Independents: %s (%llu)"), BytesToText(m_independentSize).str, m_independentMappingCount.load());
+
+		trace.StatusUpdate(startRow, 9, str, LogEntryType_Info);
+	}
+
+	bool FileMappingBuffer::Init(const tchar* name, u64 capacity, u64 commitStepSize)
+	{
+		m_name = name;
+		m_commitStepSize = commitStepSize;
+		m_capacity = capacity;
+		m_memoryHandle = m_allocator.CreateHandle(m_logger, m_name.c_str(), true);
+		if (!m_memoryHandle.IsValid())
+			return false;
+		m_memory = m_allocator.ReserveMemory(m_memoryHandle, m_capacity); // We manually commit to memory
+		return true;
+	}
+
+	MappedView FileMappingBuffer::AllocAndMapView(u64 size, u64 alignment, const tchar* hint, bool isIndependent)
+	{
+		if (isIndependent)
+		{
+			MappedView res;
+			// TODO: This can come from network with tiny allocations.. kind of sucks to take a 64k causing fragmentation.
+			// Should categorize in small and big blocks.
+			SharedMemoryHandle handle = m_allocator.CreateHandle(m_logger, hint);
+			m_allocator.ExtendMemory(handle, size);
+
+			m_independentSize += size;
+
+			u8* memory = m_allocator.MapView(handle, hint, SharedMemoryMapType_ReadWrite);
+			res.handle = handle;
+
+			res.memory = memory;
+			res.size = size;
+
+			++m_independentMappingCount;
+			return res;
+		}
+
+		MappedView res;
+
+		SCOPED_FUTEX(m_usedLock, lock);
+		u64 offset = AlignUp(m_used, alignment);
+		u64 requiredCommit = AlignUp(offset + size, PageSize);
+		if (requiredCommit > m_committed)
+		{
+			u64 commitSize = requiredCommit - m_committed;
+			commitSize = AlignUp(commitSize, m_commitStepSize);
+			commitSize = Min(commitSize, m_capacity - m_committed);
+			m_allocator.CommitMemory(m_memoryHandle, m_committed, commitSize);
+			m_committed += commitSize;
+		}
+		m_used = offset + size;
+		lock.Leave();
+
+		res.handle = m_memoryHandle;
+		res.offset = offset;
+		res.size = size;
+		res.memory = m_memory + offset;
+		return res;
+	}
+
+	MappedView FileMappingBuffer::MapView(SharedMemoryHandle handle, u64 offset, u64 size, const tchar* hint)
+	{
+		return InternalMapView(handle, offset, size, hint);
+	}
+
+	MappedView FileMappingBuffer::MapView(StringView handleName, u64 size, const tchar* hint)
+	{
+		const tchar* handleStr = handleName.data + 1;
+		const tchar* handleStrEnd = TStrchr(handleStr, '-');
+		u64 mappingOffset = 0;
+		if (handleStrEnd)
+		{
+			const tchar* mappingOffsetStr = handleStrEnd + 1;
+			mappingOffset = StringToValueBase62(mappingOffsetStr, TStrlen(mappingOffsetStr));
+		}
+		else
+			handleStrEnd = handleStr + TStrlen(handleStr);
+		SharedMemoryHandle h = SharedMemoryHandle::FromU64(StringToValueBase62(handleStr, handleStrEnd - handleStr));
+		return InternalMapView(h, mappingOffset, size, hint);
+	}
+
+	MappedView FileMappingBuffer::InternalMapView(SharedMemoryHandle handle, u64 offset, u64 size, const tchar* hint)
+	{
+		UBA_ASSERT(handle.IsValid());
+
+		MappedView res;
+
+		if (handle == m_memoryHandle)
+		{
+			res.handle = handle;
+			res.offset = offset;
+			res.size = size;
+			res.memory = m_memory + offset;
+			return res;
+		}
+
+		u8* data = m_allocator.MapView(handle, hint, SharedMemoryMapType_ReadOnly);
+		if (!data)
+			return {};
+		UBA_ASSERT(offset == 0);
+		res.handle = handle;
+		res.offset = offset;
+		res.size = size;
+		res.memory = data;
+		++m_independentMappingCount;
+		return res;
+	}
+
+	void FileMappingBuffer::UnmapView(MappedView view, const tchar* hint_, u64 newSize)
+	{
+		if (!view.handle.IsValid())
+			return;
+		if (view.handle == m_memoryHandle)
+			return;
+		m_allocator.UnmapView(view.handle, view.memory);
+		--m_independentMappingCount;
+	}
+
+	bool FileMappingBuffer::IsIndependent(SharedMemoryHandle handle)
+	{
+		return handle != m_memoryHandle;
+	}
+
+	void FileMappingBuffer::CloseIndependent(SharedMemoryHandle handle, u64 size, const tchar* hint)
+	{
+		m_independentSize -= size;
+		m_allocator.CloseHandle(m_logger, handle, hint);
+	}
+
+	SharedMemoryHandle FileMappingBuffer::GetMemoryHandle()
+	{
+		return m_memoryHandle;
+	}
+
+	u64 FileMappingBuffer::GetUsed()
+	{
+		return m_used;
+	}
+
+	////////////////////////////////////////////////////////////////////////////////////////////////////
+}
diff --git a/Engine/Source/Programs/UnrealBuildAccelerator/Common/Private/UbaProcess.cpp b/Engine/Source/Programs/UnrealBuildAccelerator/Common/Private/UbaProcess.cpp
--- a/Engine/Source/Programs/UnrealBuildAccelerator/Common/Private/UbaProcess.cpp
+++ b/Engine/Source/Programs/UnrealBuildAccelerator/Common/Private/UbaProcess.cpp
@@ -150,7 +150,7 @@ namespace uba
 	:	m_session(session)
 	,	m_parentProcess(parent)
 	,	m_id(id)
-	,	m_comMemory(detourEnabled ? m_session.m_processCommunicationAllocator.Alloc(TC("")) : FileMappingAllocator::Allocation())
+	,	m_comMemory(detourEnabled ? m_session.m_processCommunicationAllocator.Alloc(TC("")) : FileMappingAllocation())
 		#if !PLATFORM_WINDOWS
 	,	m_cancelEvent(CreateSharedEvent(m_comMemory.memory, 0))
 	,	m_writeEvent(CreateSharedEvent(m_comMemory.memory, 1))
@@ -206,8 +206,8 @@ namespace uba
 		if (!m_parentProcess)
 		{
 			for (auto& pair : m_shared.writtenFiles)
-				if (pair.second.mappingHandle.IsValid())
-					CloseFileMapping(m_session.m_logger, pair.second.mappingHandle, pair.second.name.c_str());
+				if (pair.second.memoryHandle.IsValid())
+					m_session.m_sharedMemory.CloseHandle(m_session.m_logger, pair.second.memoryHandle, pair.second.name.c_str());
 			delete &m_shared;
 		}
 	}
@@ -722,6 +722,13 @@ namespace uba
 
 		CallExitedFunc();
 
+		if (!m_parentProcess)
+		{
+			for (auto& memoryHandle : m_shared.sharedMemoryHandles)
+				m_session.m_sharedMemory.CloseHandle(m_session.m_logger, memoryHandle, TC(""));
+			m_shared.sharedMemoryHandles.clear();
+		}
+
 		UBA_ASSERT(m_refCount);
 
 		if (m_processStats.exitTime)
@@ -902,6 +909,8 @@ namespace uba
 		writer.WriteBool(m_parentProcess != nullptr);
 		writer.WriteString(m_startInfo.application);
 		writer.WriteString(m_startInfo.workingDir);
+		writer.WriteU64(response.sharedMemoryAllocatorHandle.ToU64());
+		writer.WriteU64(response.permanentFilesHandle.ToU64());
 		writer.WriteU64(response.directoryTableHandle);
 		writer.WriteU32(response.directoryTableSize);
 		writer.WriteU32(response.directoryTableCount);
@@ -976,8 +985,8 @@ namespace uba
 		msg.attributes = DefaultAttributes(); // reader.ReadFileAttributes(); TODO
 		msg.deleteOnClose = reader.ReadBool();
 		msg.success = reader.ReadBool();
-		msg.mappingHandle = FileMappingHandle::FromU64(reader.ReadU64());
-		msg.mappingWritten = reader.ReadU64();
+		msg.memoryHandle = SharedMemoryHandle::FromU64(reader.ReadU64());
+		msg.memoryWritten = reader.ReadU64();
 		msg.newNameKey = reader.ReadStringKey();
 		if (msg.newNameKey != StringKeyZero)
 			reader.ReadString(msg.newName);
@@ -1143,8 +1152,8 @@ namespace uba
 			}
 
 			u64 fileSize;
-			if (wf.mappingHandle.IsValid())
-				fileSize = wf.mappingWritten;
+			if (wf.memoryHandle.IsValid())
+				fileSize = wf.memoryWritten;
 			else if (!FileExists(m_session.m_logger, wf.backedName.c_str(), &fileSize))
 				continue;
 			else if (wf.name == wf.backedName)
@@ -1189,7 +1198,7 @@ namespace uba
 			w->WriteBool(isInTemp);
 			w->WriteString(fileName);
 			w->WriteString(wf.backedName);
-			w->Write7BitEncoded(wf.mappingHandle.ToU64());
+			w->Write7BitEncoded(wf.memoryHandle.ToU64());
 			w->Write7BitEncoded(fileSize);
 			++count;
 		}
@@ -1544,6 +1553,43 @@ namespace uba
 		return true;
 	}
 
+	bool ProcessImpl::HandleCreateSharedMemory(BinaryReader& reader, BinaryWriter& writer)
+	{
+		u64 initialSize = reader.ReadU64();
+		StringBuffer<> hint;
+		reader.ReadString(hint);
+		SharedMemoryHandle h = m_session.m_sharedMemory.CreateHandle(m_session.m_logger, hint.data);
+		writer.Write7BitEncoded(h.ToU64());
+		if (initialSize)
+			m_session.m_sharedMemory.ExtendMemory(writer, h, initialSize);
+
+		SCOPED_FUTEX(m_shared.writtenFilesLock, lock);
+		m_shared.sharedMemoryHandles.insert(h);
+		return true;
+	}
+
+	bool ProcessImpl::HandleCommitSharedMemory(BinaryReader& reader, BinaryWriter& writer)
+	{
+		SharedMemoryHandle h = SharedMemoryHandle::FromU64(reader.ReadU64());
+		u64 size = reader.ReadU64();
+		m_session.m_sharedMemory.ExtendMemory(writer, h, size);
+		return true;
+	}
+
+	bool ProcessImpl::HandleGetSharedMemory(BinaryReader& reader, BinaryWriter& writer)
+	{
+		SharedMemoryHandle h = SharedMemoryHandle::FromU64(reader.ReadU64());
+		u64 offset = reader.ReadU64();
+		u64 size = reader.ReadU64();
+		if (!size)
+		{
+			size = m_session.m_sharedMemory.GetCommitted(h);
+			writer.Write7BitEncoded(size);
+		}
+		m_session.m_sharedMemory.WriteViewSlices(writer, h, offset, size);
+		return true;
+	}
+
 	bool ProcessImpl::WriteFilesToDisk(bool isExiting)
 	{
 		Vector<WrittenFile*> files;
@@ -1556,7 +1602,7 @@ namespace uba
 			if (m_parentProcess)
 				continue;
 			kv.second.owner = nullptr;
-			if (kv.second.mappingHandle.IsValid())
+			if (kv.second.memoryHandle.IsValid())
 				files.push_back(&kv.second);
 		}
 
@@ -2455,6 +2501,6 @@ namespace uba
 
 	bool ProcessImpl::IsCalledFromThis() const
 	{
-		return m_messageThread->IsInside();
+		return m_messageThread && m_messageThread->IsInside();
 	}
 }
diff --git a/Engine/Source/Programs/UnrealBuildAccelerator/Common/Private/UbaScheduler.cpp b/Engine/Source/Programs/UnrealBuildAccelerator/Common/Private/UbaScheduler.cpp
--- a/Engine/Source/Programs/UnrealBuildAccelerator/Common/Private/UbaScheduler.cpp
+++ b/Engine/Source/Programs/UnrealBuildAccelerator/Common/Private/UbaScheduler.cpp
@@ -479,6 +479,8 @@ namespace uba
 		{
 			if (!m_updateThreadWakeup.IsSet())
 				break;
+			
+			m_session.FlushDeadProcesses();
 
 			while (RunCacheQuery())
 				;
@@ -1285,7 +1287,7 @@ namespace uba
 					auto& writtenFile = kv.second;
 					if (StringView(writtenFile.name).EndsWith(TC(".pch")))
 					{
-						stats.baseline = writtenFile.mappingWritten;
+						stats.baseline = writtenFile.memoryWritten;
 						processPeakMemory = stats.baseline;
 					}
 				}
diff --git a/Engine/Source/Programs/UnrealBuildAccelerator/Common/Private/UbaSession.cpp b/Engine/Source/Programs/UnrealBuildAccelerator/Common/Private/UbaSession.cpp
--- a/Engine/Source/Programs/UnrealBuildAccelerator/Common/Private/UbaSession.cpp
+++ b/Engine/Source/Programs/UnrealBuildAccelerator/Common/Private/UbaSession.cpp
@@ -226,7 +226,6 @@ namespace uba
 		table.GetValueAsBool(allowLocalDetour, TC("AllowLocalDetour"));
 		table.GetValueAsBool(extractObjFilesSymbols, TC("ExtractObjFilesSymbols"));
 		table.GetValueAsBool(useFakeVolumeSerial, TC("UseFakeVolumeSerial"));
-		table.GetValueAsBool(keepTransientDataMapped, TC("KeepTransientDataMapped"));
 		table.GetValueAsBool(allowLinkDependencyCrawler, TC("AllowLinkDependencyCrawler"));
 		table.GetValueAsU32(traceReserveSizeMb, TC("TraceReserveSizeMb"));
 		table.GetValueAsU32(writeFilesBottleneck, TC("WriteFilesBottleneck"));
@@ -242,26 +241,29 @@ namespace uba
 
 		if (StringView(file.name).StartsWith(m_tempPath) && !rules.KeepTempOutputFile(StringView(file.name)))
 		{
-			CloseFileMapping(m_logger, file.mappingHandle, file.backedName.c_str());
-			file.mappingHandle = {};
+			m_sharedMemory.CloseHandle(m_logger, file.memoryHandle, file.backedName.c_str());
+			file.memoryHandle = {};
 			#if UBA_DEBUG
 			m_logger.Info(TC("Skipping writing temp file %s to disk"), file.name.c_str());
 			#endif
 			return true;
 		}
 
-		bool shouldEvictFromMemory = IsRarelyReadAfterWritten(process, file.name) || file.mappingWritten > m_keepOutputFileMemoryMapsThreshold;
-		bool tryEvictFromMemory = false;
+		bool isMaybeReadAfterWritten = IsMaybeReadAfterWritten(process, file.name);
+		bool isBigEnoughToKeep = file.memoryWritten >= m_keepOutputFileMemoryMapsThreshold;
+
+		bool shouldKeepInMemory = isMaybeReadAfterWritten && isBigEnoughToKeep;
+		bool tryKeepInMemory = true;
 
 		u64 writtenSize = 0;
 		bool shouldWriteToDisk = ShouldWriteToDisk(file.name);
 
-		u64 fileSize = file.mappingWritten;
+		u64 fileSize = file.memoryWritten;
 		u8* mem = nullptr;
-		auto memClose = MakeGuard([&]()
+		auto memClose = MakeGuard([&, memoryHandle = file.memoryHandle]()
 			{
 				if (mem)
-					UnmapViewOfFile(m_logger, mem, fileSize, file.name.c_str(), shouldEvictFromMemory && m_allowDiscardVirtualMemory);
+					m_sharedMemory.UnmapView(memoryHandle, mem);
 			});
 
 		if (shouldWriteToDisk)
@@ -274,70 +276,15 @@ namespace uba
 			if (!storeCompressed)
 				shouldDevirtualize = HasVfs(rootsHandle) && rules.ShouldDevirtualizeFile(file.name, escapeSpaces);
 
-#if UBA_ENABLE_ON_DISK_FILE_MAPPINGS
-			if (!storeCompressed && !shouldDevirtualize && file.mappingHandle.fh)
-			{
-				bool success = true;
-				SetFilePointerEx(file.mappingHandle.fh, ToLargeInteger(file.mappingWritten), NULL, FILE_BEGIN);
-				::SetEndOfFile(file.mappingHandle.fh);
-
-				if (false)
-				{
-					TrackWorkScope work(m_workManager, AsView(TC("Flush")));
-					work.AddHint(file.name);
-					FlushFileBuffers(file.mappingHandle.fh);
-				}
-
-				if (true)
-				{
-					FILE_DISPOSITION_INFO info;
-					info.DeleteFile = false;
-					if (!::SetFileInformationByHandle(file.mappingHandle.fh, FileDispositionInfo, &info, sizeof(info)))
-						m_logger.Warning(TC("Failed to remove delete-on-close"));
-				}
-
-				if (false)
-				{
-					StringBuffer<> name2;
-					name2.Append(TCV("\\??\\")).Append(file.name);
-
-					u8 buffer[2048];
-					auto& info = *(FILE_RENAME_INFO*)buffer;;
-					info.ReplaceIfExists = true;
-					info.RootDirectory = 0;
-					info.FileNameLength = u32(name2.count*sizeof(tchar));
-					info.Flags = FILE_RENAME_FLAG_POSIX_SEMANTICS;
-					memcpy(info.FileName, name2.data, info.FileNameLength+sizeof(tchar));
-					TrackWorkScope work(m_workManager, AsView(TC("Rename")));
-					work.AddHint(file.name);
-					success = SetFileInformationByHandle(file.mappingHandle.fh, FileRenameInfo, &info, sizeof(buffer));
-				}
-
-				if (success)
-				{
-					{
-						TrackWorkScope work(m_workManager, AsView(TC("CloseFileMapping")));
-						work.AddHint(file.name);
-						CloseFileMapping(m_logger, file.mappingHandle, file.name.c_str());
-					}
-					file.mappingHandle = {};
-					file.originalMappingHandle = {};
-					m_storage.InvalidateCachedFileInfo(file.key);
-					return true;
-				}
-				m_logger.Warning(TC("SetFileInformationByHandle failed %s (%s)"), file.name.c_str(), LastErrorToText().data);
-			}
-#endif
-
 			#if UBA_DEBUG_TRACK_MAPPING
-			m_debugLogger->Info(TC("Writing written file with mapping 0x%llx for %s"), u64(file.mappingHandle.mh), file.name.c_str());
+			m_debugLogger->Info(TC("Writing written file with mapping 0x%llx for %s"), u64(file.memoryHandle.ToU64()), file.name.c_str());
 			#endif
 
 			u32 desiredAccess = FILE_MAP_READ;
 			if (m_allowDiscardVirtualMemory)
 				desiredAccess |= FILE_MAP_WRITE;
 
-			mem = MapViewOfFile(m_logger, file.mappingHandle, desiredAccess, 0, fileSize);
+			mem = m_sharedMemory.MapView(file.memoryHandle, file.name.c_str(), SharedMemoryMapType_ReadOnly);
 			if (!mem)
 				return m_logger.Error(TC("Failed to map view of filehandle for read %s (%s)"), file.name.c_str(), LastErrorToText().data);
 
@@ -354,8 +301,8 @@ namespace uba
 				writtenSize = res.storedSize;
 
 				// Can't evict without properly update filemappingtable.. the file on disk does now not match what was registered for write
-				tryEvictFromMemory = shouldEvictFromMemory;
-				shouldEvictFromMemory = false;
+				tryKeepInMemory = shouldKeepInMemory;
+				shouldKeepInMemory = true;
 			}
 			else
 			{
@@ -401,17 +348,17 @@ namespace uba
 			uba::DeleteFileW(file.name.c_str());
 
 			// If shouldWriteToDisk is false we can't evict from memory because we can't re-read from disk if needed
-			shouldEvictFromMemory = false;
+			shouldKeepInMemory = true;
 		}
 
-		FileMappingHandle mh = file.mappingHandle;
-		file.mappingHandle = {};
-		file.originalMappingHandle = {};
+		memClose.Execute();
 
-		if (!shouldEvictFromMemory)
+		SharedMemoryHandle memoryHandle = file.memoryHandle;
+		file.memoryHandle = {};
+		file.originalMemoryHandle = {};
+
+		if (shouldKeepInMemory)
 		{
-			StringBuffer<> name;
-			Storage::GetMappingString(name, mh, 0);
 			SCOPED_FUTEX(m_fileMappingTableLookupLock, lookupLock);
 			auto insres = m_fileMappingTableLookup.try_emplace(file.key);
 			FileMappingEntry& entry = insres.first->second;
@@ -420,18 +367,21 @@ namespace uba
 			SCOPED_FUTEX(entry.lock, entryCs);
 			if (!insres.second && entry.canBeFreed) // It could be that this file has been read as input and that is fine.
 				m_logger.Error(TC("Trying to write the same file twice (%s)"), file.name.c_str());
-			UBA_ASSERT(mh.IsValid());
+			UBA_ASSERT(memoryHandle.IsValid());
 
-			if (!entry.handled && tryEvictFromMemory)
+			if (!entry.handled && !tryKeepInMemory)
 			{
-				shouldEvictFromMemory = true;
+				shouldKeepInMemory = false;
 			}
 			else
 			{
+				// Free up unsude memory since we're keeping it around
+				m_sharedMemory.TrimMemory(memoryHandle, file.memoryWritten);
+
 				entry.handled = true;
-				entry.mapping = mh;
-				entry.mappingOffset = 0;
-				entry.contentSize = file.mappingWritten;
+				entry.memoryHandle = memoryHandle;
+				entry.memoryOffset = 0;
+				entry.contentSize = file.memoryWritten;
 				entry.lastWriteTime = file.lastWriteTime;
 				entry.isDir = false;
 				entry.success = true;
@@ -441,40 +391,46 @@ namespace uba
 				entry.storedSize = writtenSize;
 				entry.isInvisible = rules.IsInvisible(file.name);
 
+				//m_logger.Info(TC("KEEPING %s"), file.name.c_str());
+				++m_independentMappingCreated;
+				++m_independentMappingActive;
+
+				if (shouldWriteToDisk) // If writing to disk we only keep file in memory for 3 seconds until we evict it
+				{
+					SCOPED_FUTEX(m_fileMappingsThatCanBeFreedLock, lock);
+					m_fileMappingsThatCanBeFreed[file.key] = GetTime() + MsToTime(3*1000);
+				}
+
 				#if UBA_DEBUG_TRACK_MAPPING
 				entry.name = file.name;
-				m_debugLogger->Info(TC("Mapping kept 0x%llx (%s) from detoured process (UsedCountBeforeFree: %u)"), u64(mh.mh), entry.name.c_str(), u32(entry.usedCountBeforeFree));
+				m_debugLogger->Info(TC("Mapping kept 0x%llx (%s) from detoured process (UsedCountBeforeFree: %u)"), u64(memoryHandle.ToU64()), entry.name.c_str(), u32(entry.usedCountBeforeFree));
 				#endif
 
 				if (!entry.isInvisible)
 				{
+					StringBuffer<> name;
+					Storage::GetMappingString(name, memoryHandle, 0, entry.canBeFreed);
+
 					SCOPED_WRITE_LOCK(m_fileMappingTableMemLock, lock);
 					BinaryWriter writer(m_fileMappingTableMem, m_fileMappingTableSize);
 					writer.WriteStringKey(file.key);
 					writer.WriteString(name);
-					writer.Write7BitEncoded(file.mappingWritten);
+					writer.Write7BitEncoded(file.memoryWritten);
 					u32 newSize = (u32)writer.GetPosition();
 					m_fileMappingTableSize = (u32)newSize;
 				}
 			}
 		}
 
-		memClose.Execute();
-
-		if (shouldEvictFromMemory)
+		if (!shouldKeepInMemory)
 		{
+			//m_logger.Info(TC("EVICTING %s"), file.name.c_str());
+
 			#if UBA_DEBUG_TRACK_MAPPING
-			m_debugLogger->Info(TC("Mapping eviction queued 0x%llx (%s)"), u64(mh.mh), file.name.c_str());
+			m_debugLogger->Info(TC("Mapping evicted 0x%llx (%s)"), u64(memoryHandle.ToU64()), file.name.c_str());
 			#endif
 
-			m_workManager.AddWork([mh, n = file.name, this](const WorkContext&)
-				{
-					#if UBA_DEBUG_TRACK_MAPPING
-					m_debugLogger->Info(TC("Mapping evicted 0x%llx (%s)"), u64(mh.mh), n.c_str());
-					#endif
-
-					CloseFileMapping(m_logger, mh, n.c_str());
-				}, 1, TC("CloseFileMapping"));
+			m_sharedMemory.CloseHandle(m_logger, memoryHandle, file.name.c_str());
 		}
 
 		if (shouldWriteToDisk)
@@ -513,7 +469,7 @@ namespace uba
 
 	bool Session::GetFileMemory(const Function<bool(const void*, u64)>& func, const StringKey& fileNameKey, StringView filePath, bool deleteInternalMapping)
 	{
-		FileMappingHandle mapping;
+		SharedMemoryHandle mapping;
 		u64 size;
 		{
 			SCOPED_FUTEX(m_fileMappingTableLookupLock, lookupLock);
@@ -523,7 +479,7 @@ namespace uba
 			FileMappingEntry& entry = findIt->second;
 			if (deleteInternalMapping && !entry.isInvisible)
 				return m_logger.Error(TC("Trying to delete mapping that is not invisible (%s)"), filePath.data);
-			mapping = entry.mapping;
+			mapping = entry.memoryHandle;
 			size = entry.contentSize;
 			if (deleteInternalMapping)
 			{
@@ -533,13 +489,13 @@ namespace uba
 			}
 		}
 
-		u8* mem = MapViewOfFile(m_logger, mapping, FILE_MAP_READ, 0, size);
+		u8* mem = m_sharedMemory.MapView(mapping, TC(""));
 		if (!mem)
 			return false;
 		bool res = func(mem, size);
-		UnmapViewOfFile(m_logger, mem, size, filePath.data, deleteInternalMapping && m_allowDiscardVirtualMemory);
+		m_sharedMemory.UnmapView(mapping, mem);
 		if (deleteInternalMapping)
-			CloseFileMapping(m_logger, mapping, TC(""));
+			m_sharedMemory.CloseHandle(m_logger, mapping, TC(""));
 		return res;
 	}
 
@@ -821,7 +777,7 @@ namespace uba
 		entry.contentSize = fileSize;
 		entry.isDir = false;
 		entry.success = true;
-		entry.mapping = {};
+		entry.memoryHandle = {};
 		entry.handled = true;
 
 		#if UBA_DEBUG_TRACK_MAPPING
@@ -871,9 +827,9 @@ namespace uba
 			if (!entry.success)
 				return false;
 			out.size = entry.contentSize;
-			if (entry.mapping.IsValid())
+			if (entry.memoryHandle.IsValid())
 			{
-				Storage::GetMappingString(out.name, entry.mapping, entry.mappingOffset);
+				Storage::GetMappingString(out.name, entry.memoryHandle, entry.memoryOffset, entry.canBeFreed);
 				updateRequestingProcess();
 			}
 			else
@@ -965,51 +921,55 @@ namespace uba
 		}
 		else
 		{
-			if (contentSize > m_fileMappingBuffer.GetFileMappingCapacity())
-				return m_logger.Error(TC("File %s has a size (%llu) that is too large to fit in mapping buffer (%s)"), fileName, contentSize, hint);
-
-			FileMappingHandle mapping;
-			u64 mappingOffset = 0;
-			u8* mappingMemory = nullptr;
+			SharedMemoryHandle memoryHandle;
+			u64 memoryOffset = 0;
+			u8* memory = nullptr;
 
 			MappedView mappedView;
 
-			auto ownedGuard1 = MakeGuard([&]() { CloseFileMapping(m_logger, mapping, fileName); });
-			auto ownedGuard2 = MakeGuard([&]() { m_workManager.AddWork([this, mappingMemory, contentSize, fn = fileNameView.ToString()](const WorkContext&) { UnmapViewOfFile(m_logger, mappingMemory, contentSize, fn.c_str()); }, 1, TC("UnmapViewOfFile")); });
+			auto ownedGuard1 = MakeGuard([&]() { m_sharedMemory.CloseHandle(m_logger, memoryHandle, fileName); });
+			auto ownedGuard2 = MakeGuard([&]() { m_sharedMemory.UnmapView(memoryHandle, memory); });
 			auto viewGuard = MakeGuard([&]() { m_fileMappingBuffer.UnmapView(mappedView, fileName); });
 
 			if (canBeFreed)
 			{
 				viewGuard.Cancel();
-				mapping = CreateMemoryMappingW(m_logger, PAGE_READWRITE, contentSize, nullptr, fileName);
-				if (!mapping.IsValid())
+
+				memoryHandle = m_sharedMemory.CreateHandle(m_logger, fileName);
+				if (!memoryHandle.IsValid())
 					return false;
-				mappingMemory = MapViewOfFile(m_logger, mapping, FILE_MAP_WRITE, 0, contentSize, m_allowDiscardVirtualMemory);
-				if (!mappingMemory)
+				m_sharedMemory.ExtendMemory(memoryHandle, contentSize);
+				memory = m_sharedMemory.MapView(memoryHandle, fileName, SharedMemoryMapType_ReadWrite);
+				if (!memory)
 					return false;
+
+				// TODO: THIS IS THE ALLOCATED obj files etc
+				//m_logger.Info(TC("MAPPING %s"), fileName);
+				++m_independentMappingCreated;
+				++m_independentMappingActive;
 			}
 			else
 			{
 				UBA_ASSERTF(alignment, TC("No alignment set when creating memory map for %s (%s)"), fileName, hint);
-				UBA_ASSERT(!entry.canBeFreed);
+				//UBA_ASSERTF(!entry.canBeFreed, TC("Found entry %s which is set as canBeFreed but is not opened as canBeFreed (%s)"), fileName, hint);
 				ownedGuard1.Cancel();
 				ownedGuard2.Cancel();
-				mappedView = m_fileMappingBuffer.AllocAndMapView(MappedView_Transient, contentSize, alignment, fileName, false);
+				mappedView = m_fileMappingBuffer.AllocAndMapView(contentSize, alignment, fileName);
 				if (!mappedView.memory)
 					return false;
-				mapping = mappedView.handle;
-				mappingOffset = mappedView.offset;
-				mappingMemory = mappedView.memory;
+				memoryHandle = mappedView.handle;
+				memoryOffset = mappedView.offset;
+				memory = mappedView.memory;
 			}
 
 			if (isCompressed)
 			{
-				if (!m_storage.DecompressFileToMemory(fileName, fileHandle, mappingMemory, contentSize, TC("GetOrCreateMemoryMapFromFile"), fileStartOffset))
+				if (!m_storage.DecompressFileToMemory(fileName, fileHandle, memory, contentSize, TC("GetOrCreateMemoryMapFromFile"), fileStartOffset))
 					return false;
 			}
 			else
 			{
-				if (!ReadFile(m_logger, fileName, fileHandle, mappingMemory, contentSize))
+				if (!ReadFile(m_logger, fileName, fileHandle, memory, contentSize))
 					return false;
 			}
 
@@ -1018,12 +978,14 @@ namespace uba
 			viewGuard.Execute();
 
 			entry.canBeFreed = canBeFreed;
-			entry.mappingOffset = mappingOffset;
-			Storage::GetMappingString(out.name, mapping, mappingOffset);
-			entry.mapping = mapping;
+			entry.memoryOffset = memoryOffset;
+			entry.memoryHandle = memoryHandle;
+			Storage::GetMappingString(out.name, memoryHandle, memoryOffset, entry.canBeFreed);
 
 			if (canBeFreed)
 			{
+				UBA_ASSERT(!memoryHandle.IsValid() || m_sharedMemory.GetRefCount(memoryHandle));
+
 				entry.usedCount = 0;
 				entry.usedCountBeforeFree = g_globalRules.GetUsedCountBeforeFree(fileNameView);
 			}
@@ -1044,7 +1006,7 @@ namespace uba
 
 		#if UBA_DEBUG_TRACK_MAPPING
 		entry.name = fileName;
-		m_debugLogger->Info(TC("Mapping created 0x%llx (%s) from file (%s) - %s"), u64(entry.mapping.mh), entry.name.c_str(), hint, TimeToText(GetTime() - ts2.start).str);
+		m_debugLogger->Info(TC("Mapping created 0x%llx (%s) from file (%s) - %s"), u64(entry.memoryHandle.ToU64()), entry.name.c_str(), hint, TimeToText(GetTime() - ts2.start).str);
 		#endif
 
 		entry.isDir = isDir;
@@ -1076,8 +1038,8 @@ namespace uba
 			if (!entry.success)
 				return false;
 			out.size = entry.contentSize;
-			if (entry.mapping.IsValid())
-				Storage::GetMappingString(out.name, entry.mapping, entry.mappingOffset);
+			if (entry.memoryHandle.IsValid())
+				Storage::GetMappingString(out.name, entry.memoryHandle, entry.memoryOffset, entry.canBeFreed);
 			else
 				out.name.Append(entry.isDir ? TC("$d") : TC("$f"));
 			return true;
@@ -1087,7 +1049,7 @@ namespace uba
 
 		entry.handled = true;
 
-		MappedView mappedViewRead = m_storage.MapView(casKey, fileName);
+		MappedView2 mappedViewRead = m_storage.MapView(casKey, fileName);
 		if (!mappedViewRead.handle.IsValid())
 			return false;
 
@@ -1106,10 +1068,7 @@ namespace uba
 			}
 			else
 			{
-				if (contentSize > m_fileMappingBuffer.GetFileMappingCapacity())
-					return m_logger.Error(TC("File %s has a size (%llu) that is too large to fit in mapping buffer (GetOrCreateMemoryMapFromView)"), fileName, contentSize);
-
-				auto mappedViewWrite = m_fileMappingBuffer.AllocAndMapView(MappedView_Transient, contentSize, alignment, fileName);
+				auto mappedViewWrite = m_fileMappingBuffer.AllocAndMapView(contentSize, alignment, fileName);
 				if (!mappedViewWrite.memory)
 					return false;
 				auto unmapGuard = MakeGuard([&](){ m_fileMappingBuffer.UnmapView(mappedViewWrite, fileName); });
@@ -1118,18 +1077,18 @@ namespace uba
 					return false;
 				unmapGuard.Execute();
 
-				entry.mappingOffset = mappedViewWrite.offset;
-				Storage::GetMappingString(out.name, mappedViewWrite.handle, mappedViewWrite.offset);
-				entry.mapping = mappedViewWrite.handle;
+				entry.memoryOffset = mappedViewWrite.offset;
+				entry.memoryHandle = mappedViewWrite.handle;
+				Storage::GetMappingString(out.name, mappedViewWrite.handle, mappedViewWrite.offset, entry.canBeFreed);
 			}
 			mvrg.Execute();
 		}
 		else
 		{
 			UBA_ASSERT(mappedViewRead.memory == nullptr);
-			entry.mappingOffset = mappedViewRead.offset;
-			Storage::GetMappingString(out.name, mappedViewRead.handle, mappedViewRead.offset);
-			entry.mapping = mappedViewRead.handle;
+			entry.memoryOffset = mappedViewRead.offset;
+			Storage::GetMappingString(out.name, mappedViewRead.handle, mappedViewRead.offset, entry.canBeFreed);
+			entry.memoryHandle = mappedViewRead.handle;
 			contentSize = mappedViewRead.size;
 		}
 		entry.success = true;
@@ -1150,12 +1109,55 @@ namespace uba
 
 		#if UBA_DEBUG_TRACK_MAPPING
 		entry.name = fileName;
-		m_debugLogger->Info(TC("Mapping created 0x%llx (%s) from view"), u64(entry.mapping.mh), entry.name.c_str());
+		m_debugLogger->Info(TC("Mapping created 0x%llx (%s) from view"), u64(entry.memoryHandle.ToU64()), entry.name.c_str());
 		#endif
 
 		return true;
 	}
 
+	void Session::ThreadGcFileMappings()
+	{
+		while (!m_fileMappingGcDone.IsSet(1000))
+		{
+			u64 now = GetTime();
+			SCOPED_FUTEX(m_fileMappingsThatCanBeFreedLock, lock);
+
+			for (auto i=m_fileMappingsThatCanBeFreed.begin(), e=m_fileMappingsThatCanBeFreed.end(); i!=e;)
+			{
+				auto inc = MakeGuard([&]() { ++i;  });
+				if (i->second < now)
+					continue;
+
+				SCOPED_FUTEX(m_fileMappingTableLookupLock, lookupLock);
+				auto findIt = m_fileMappingTableLookup.find(i->first);
+				UBA_ASSERT(findIt != m_fileMappingTableLookup.end());
+				FileMappingEntry& entry = findIt->second;
+				lookupLock.Leave();
+
+				SCOPED_FUTEX(entry.lock, entryLock);
+				if (entry.refCount)
+					continue;
+
+				if (entry.memoryHandle.IsValid() && entry.canBeFreed)
+				{
+					// There is a chance something might be evicted and re-added in between and when reloaded, loaded into non-freeable memory
+					if (entry.canBeFreed)
+					{
+						m_sharedMemory.CloseHandle(m_logger, entry.memoryHandle, TC("GarbageCollect"));
+						--m_independentMappingActive;
+						entry.handled = false;
+						entry.memoryHandle = {};
+					}
+				}
+
+				inc.Cancel();
+
+				i = m_fileMappingsThatCanBeFreed.erase(i);
+				e = m_fileMappingsThatCanBeFreed.end();
+			}
+		}
+	}
+
 	bool GetDirKey(StringKey& outDirKey, StringBufferBase& outDirName, const tchar*& outLastSlash, const StringView& fileName)
 	{
 		outLastSlash = TStrrchr(fileName.data, PathSeparator);
@@ -1553,11 +1555,12 @@ namespace uba
 
 	Session::Session(const SessionCreateInfo& info, const tchar* logPrefix, bool runningRemote, WorkManager& workManager)
 	:	m_storage(info.storage)
+	,	m_sharedMemory(info.storage.GetSharedMemory())
 	,	m_logger(info.logWriter, logPrefix)
 	,	m_workManager(workManager)
 	,	m_ownsTrace(info.trace == nullptr)
 	,	m_directoryTable(m_directoryTableMemory)
-	,	m_fileMappingBuffer(m_logger, &workManager)
+	,	m_fileMappingBuffer(m_logger, m_sharedMemory)
 	,	m_processCommunicationAllocator(m_logger, TC("CommunicationAllocator"))
 	,	m_processMessageThreads(*new WorkManagerImpl(~0u, TC("UbaProcMsg")))
 	,	m_trace(info.trace ? *info.trace : *new Trace(info.logWriter))
@@ -1721,7 +1724,8 @@ namespace uba
 		m_directoryTable.m_lookup.reserve(30000);
 		m_fileMappingTableLookup.reserve(70000);
 
-		m_fileMappingBuffer.AddTransient(TC("FileMappings"), info.keepTransientDataMapped);
+		u64 capacity = 92ull*1024*1024*1024;
+		m_fileMappingBuffer.Init(TC("FileDataNoDelete"), capacity, 1ull*1024*1024); // 32 slots to be able to add multiple files in parallel
 
 		if (!m_processCommunicationAllocator.Init(CommunicationMemSize, CommunicationMemSize * 512))
 		{
@@ -1815,10 +1819,14 @@ namespace uba
 				FileMappingEntry& entry = findIt->second;
 				lookupLock.Leave();
 				SCOPED_FUTEX_READ(entry.lock, entryLock);
-				if (!entry.handled || !entry.success || !entry.mapping.IsValid())
+				if (!entry.handled || !entry.success || !entry.memoryHandle.IsValid())
 					return false;
-				out.handle = entry.mapping;
-				out.offset = entry.mappingOffset;
+				
+				if (entry.canBeFreed) // TODO: Hand over ref counting?
+					return false;
+
+				out.handle = entry.memoryHandle;
+				out.offset = entry.memoryOffset;
 				out.size = entry.contentSize;
 				out.lastWriteTime = entry.lastWriteTime;
 				out.storedSize = entry.storedSize;
@@ -1831,11 +1839,17 @@ namespace uba
 			return false;
 		#endif
 
+
+		m_fileMappingGcDone.Create(EventResetType_Manual);
+		m_fileMappingGcThread.Start([this]() { ThreadGcFileMappings();  return 0; }, TC("UbaMemGc"));
+
 		return true;
 	}
 
 	Session::~Session()
 	{
+		m_fileMappingGcDone.Set();
+
 		if (m_ownsTrace)
 		{
 			m_trace.StopThread();
@@ -1851,12 +1865,14 @@ namespace uba
 
 		delete &m_processMessageThreads;
 
+		m_fileMappingGcThread.Wait();
+
 		for (auto& kv : m_virtualSourceFiles)
-			CloseFileMapping(m_logger, kv.second.mappingHandle, TC("VirtualFile"));
+			m_sharedMemory.CloseHandle(m_logger, kv.second.memoryHandle, TC("VirtualFile"));
 
 		for (auto& i : m_fileMappingTableLookup)
 			if (i.second.canBeFreed)
-				CloseFileMapping(m_logger, i.second.mapping, TC("FileMappingKeptFromOutput"));
+				m_sharedMemory.CloseHandle(m_logger, i.second.memoryHandle, TC("FileMappingKeptFromOutput"));
 
 		UnmapViewOfFile(m_logger, m_fileMappingTableMem, FileMappingTableMemSize, TC("FileMappingTable"), m_allowDiscardVirtualMemory);
 		CloseFileMapping(m_logger, m_fileMappingTableHandle, TC("FileMappingTable"));
@@ -2112,7 +2128,7 @@ namespace uba
 		VirtualSourceFile vsf = findIt->second;
 		m_virtualSourceFiles.erase(findIt);
 		virtualSourceFilesLock.Leave();
-		CloseFileMapping(m_logger, vsf.mappingHandle, TC("VirtualFile"));
+		m_sharedMemory.CloseHandle(m_logger, vsf.memoryHandle, TC("VirtualFile"));
 		return true;
 	}
 
@@ -2512,17 +2528,18 @@ namespace uba
 			UBA_ASSERT(entry.canBeFreed);
 
 			#if UBA_DEBUG_TRACK_MAPPING
-			m_debugLogger->Info(TC("Mapping freed 0x%llx (%s)"), u64(entry.mapping.mh), entry.name.c_str());
+			m_debugLogger->Info(TC("Mapping freed 0x%llx (%s)"), u64(entry.memoryHandle.ToU64()), entry.name.c_str());
 			#endif
 
-			m_workManager.AddWork([mh = entry.mapping, fileSize = entry.contentSize, this](const WorkContext&)
+			m_workManager.AddWork([mh = entry.memoryHandle, fileSize = entry.contentSize, this](const WorkContext&)
 				{
-					if (m_allowDiscardVirtualMemory)
-						DiscardFileMapping(m_logger, mh, fileSize, TC("UsedFileMapping"));
-					CloseFileMapping(m_logger, mh, TC("UsedFileMapping"));
+					//m_logger.Info(TC("RELEASING %s"), name.c_str());
+					m_sharedMemory.CloseHandle(m_logger, mh, TC("UsedFileMapping"));
+					--m_independentMappingActive;
+
 				}, 1, TC("CloseFileMapping"));
 			entry.handled = false;
-			entry.mapping = {};
+			entry.memoryHandle = {};
 		}
 
 
@@ -2578,6 +2595,9 @@ namespace uba
 
 	bool Session::GetInitResponse(InitResponse& out, const InitMessage& msg)
 	{
+		out.sharedMemoryAllocatorHandle = m_sharedMemory;
+		out.permanentFilesHandle = m_fileMappingBuffer.GetMemoryHandle();
+
 		out.directoryTableHandle = m_directoryTableHandle.ToU64();
 		{
 			SCOPED_READ_LOCK(m_directoryTable.m_memoryLock, l);
@@ -2688,9 +2708,9 @@ namespace uba
 		return process.m_startInfo.rules->IsRarelyRead(fileName);
 	}
 
-	bool Session::IsRarelyReadAfterWritten(ProcessImpl& process, const StringView& fileName) const
+	bool Session::IsMaybeReadAfterWritten(ProcessImpl& process, const StringView& fileName) const
 	{
-		return process.m_startInfo.rules->IsRarelyReadAfterWritten(fileName);
+		return process.m_startInfo.rules->IsMaybeReadAfterWritten(fileName);
 	}
 
 	bool Session::IsKnownSystemFile(const tchar* applicationName)
@@ -2980,11 +3000,11 @@ namespace uba
 				FileMappingEntry& entry = insres.first->second;
 				lookupLock.Leave();
 				SCOPED_FUTEX(entry.lock, entryLock);
-				if (entry.handled && entry.success && entry.mapping.IsValid())
+				if (entry.handled && entry.success && entry.memoryHandle.IsValid())
 				{
 					UBA_ASSERT(entry.isInvisible);
 					out.size = entry.contentSize;
-					Storage::GetMappingString(out.fileName, entry.mapping, entry.mappingOffset);
+					Storage::GetMappingString(out.fileName, entry.memoryHandle, entry.memoryOffset, entry.canBeFreed);
 					return true;
 				}
 			}
@@ -3009,7 +3029,7 @@ namespace uba
 		if (findIt == writtenFiles.end())
 			return;
 
-		FileMappingHandle h = findIt->second.mappingHandle;
+		SharedMemoryHandle h = findIt->second.memoryHandle;
 		TString name = findIt->second.name;
 		writtenFiles.erase(findIt);
 		writtenLock.Leave();
@@ -3018,10 +3038,10 @@ namespace uba
 			return;
 
 		#if UBA_DEBUG_TRACK_MAPPING
-		m_debugLogger->Info(TC("Removed %s with handle 0x%llx"), name.c_str(), h.mh);
+		m_debugLogger->Info(TC("Removed %s with handle 0x%llx"), name.c_str(), h.ToU64());
 		#endif
 
-		CloseFileMapping(process.m_session.GetLogger(), h, name.c_str());
+		m_sharedMemory.CloseHandle(process.m_session.GetLogger(), h, name.c_str());
 	}
 
 	bool Session::CloseFile(CloseFileResponse& out, const CloseFileMessage& msg)
@@ -3090,28 +3110,29 @@ namespace uba
 			{
 				if (writtenFile.backedName != msgName.data)
 				{
-					UBA_ASSERT(!msg.mappingHandle.IsValid() && !writtenFile.mappingHandle.IsValid());
+					UBA_ASSERT(!msg.memoryHandle.IsValid() && !writtenFile.memoryHandle.IsValid());
 					writtenFile.backedName = msgName.ToString();
 				}
 
-				if (!msg.mappingHandle.IsValid() || (msg.mappingHandle == writtenFile.originalMappingHandle && writtenFile.owner == &msg.process))
+				if (!msg.memoryHandle.IsValid() || (msg.memoryHandle == writtenFile.originalMemoryHandle && writtenFile.owner == &msg.process))
 				{
-					if (msg.mappingWritten)
+					if (msg.memoryWritten)
 					{
-						writtenFile.mappingWritten = msg.mappingWritten;
+						writtenFile.memoryWritten = msg.memoryWritten;
 						writtenFile.lastWriteTime = GetSystemTimeAsFileTime();
 					}
 					addMapping = false;
 				}
-				else if (writtenFile.mappingHandle.IsValid())
+				else if (writtenFile.memoryHandle.IsValid())
 				{
 					#if UBA_DEBUG_TRACK_MAPPING
-					m_debugLogger->Info(TC("Closing old mapping 0x%llx for %s"), u64(writtenFile.mappingHandle.mh), writtenFile.name.c_str());
+					m_debugLogger->Info(TC("Closing old mapping 0x%llx for %s"), u64(writtenFile.memoryHandle.ToU64()), writtenFile.name.c_str());
 					#endif
-
-					CloseFileMapping(m_logger, writtenFile.mappingHandle, msg.fileName.data);
-					writtenFile.mappingHandle = {};
-					writtenFile.originalMappingHandle = {};
+					
+					if (writtenFile.memoryHandle != msg.memoryHandle)
+						m_sharedMemory.CloseHandle(m_logger, writtenFile.memoryHandle, msg.fileName.data);
+					writtenFile.memoryHandle = {};
+					writtenFile.originalMemoryHandle = {};
 				}
 
 				writtenFile.owner = &msg.process;
@@ -3120,7 +3141,7 @@ namespace uba
 			if (!m_runningRemote && HasVfs(startInfo.rootsHandle)) // For posix we write the dependency file directly to disk so we need to update it if vfs is enabled
 			{
 				bool escapeSpaces;
-				if (!msg.mappingHandle.IsValid() && rules.ShouldDevirtualizeFile(activeFile.name, escapeSpaces))
+				if (!msg.memoryHandle.IsValid() && rules.ShouldDevirtualizeFile(activeFile.name, escapeSpaces))
 				{
 					// TODO: On linux we don't use file mappings for outputs yet.. so we have to open the file and change it
 					FileAccessor readFile(m_logger, name.data);
@@ -3147,25 +3168,21 @@ namespace uba
 
 			if (addMapping)
 			{
-				FileMappingHandle mappingHandle;
-				if (msg.mappingHandle.IsValid())
-					if (!DuplicateFileMapping(m_logger, msg.process.m_nativeProcessHandle, msg.mappingHandle, GetCurrentProcessHandle(), mappingHandle, 0, false, DUPLICATE_SAME_ACCESS, msgName.data))
-						return m_logger.Error(TC("Failed to duplicate file mapping handle for %s"), name.data);
-
-				writtenFile.mappingHandle = mappingHandle;
-				writtenFile.mappingWritten = msg.mappingWritten;
-				writtenFile.originalMappingHandle = msg.mappingHandle;
+				msg.process.m_shared.sharedMemoryHandles.erase(msg.memoryHandle); // memory handle transferred to written files
+				writtenFile.memoryHandle = msg.memoryHandle;
+				writtenFile.originalMemoryHandle = msg.memoryHandle;
+				writtenFile.memoryWritten = msg.memoryWritten;
 				writtenFile.lastWriteTime = GetSystemTimeAsFileTime();
 
 				#if UBA_DEBUG_TRACK_MAPPING
-				m_debugLogger->Info(TC("Adding written file with mapping 0x%llx (from 0x%llx) for %s"), u64(writtenFile.mappingHandle.mh), u64(msg.mappingHandle.mh), writtenFile.name.c_str());
+				m_debugLogger->Info(TC("Adding written file with mapping 0x%llx (from 0x%llx) for %s"), u64(writtenFile.memoryHandle.ToU64()), u64(msg.memoryHandle.ToU64()), writtenFile.name.c_str());
 				#endif
 			}
 
-			if (writtenFile.mappingHandle.IsValid())
+			if (writtenFile.memoryHandle.IsValid())
 			{
 				registerRealFile = false;
-				fileSize = writtenFile.mappingWritten;
+				fileSize = writtenFile.memoryWritten;
 				lastWriteTime = writtenFile.lastWriteTime;
 			}
 
@@ -3218,9 +3235,9 @@ namespace uba
 			if (findIt != shared.writtenFiles.end())
 			{
 				WrittenFile& writtenFile = findIt->second;
-				if (writtenFile.mappingHandle.IsValid())
+				if (writtenFile.memoryHandle.IsValid())
 				{
-					CloseFileMapping(m_logger, writtenFile.mappingHandle, writtenFile.name.c_str());
+					m_sharedMemory.CloseHandle(m_logger, writtenFile.memoryHandle, writtenFile.name.c_str());
 					deleteRealFile = false;
 				}
 				shared.writtenFiles.erase(findIt);
@@ -3263,7 +3280,7 @@ namespace uba
 			if (findIt != fs.writtenFiles.end())
 			{
 				auto& oldFile = findIt->second;
-				bool isMapping = oldFile.mappingHandle.IsValid();
+				bool isMapping = oldFile.memoryHandle.IsValid();
 				if (!isMapping)
 				{
 					out.result = MoveFileExW(msg.fromName.data, msg.toName.data, msg.flags);
@@ -3531,7 +3548,8 @@ namespace uba
 
 		MemoryMap out;
 		u64 alignment = GetMemoryMapAlignment(fileName);
-		if (!GetOrCreateMemoryMapFromFile(out, fileNameKey, fileName.data, false, alignment, TC("Remap"), &process, true))
+		bool canBeFreed = true;
+		if (!GetOrCreateMemoryMapFromFile(out, fileNameKey, fileName.data, false, alignment, TC("Remap"), &process, canBeFreed))
 			return m_logger.Error(TC("Failed to remap %s"), fileName.data);
 		writer.WriteU32(GetFileMappingSize());
 		return true;
@@ -3580,12 +3598,9 @@ namespace uba
 
 	void Session::PrintSessionStats(Logger& logger)
 	{
-		u64 mappingBufferSize;
-		u32 mappingBufferCount;
-		m_fileMappingBuffer.GetSizeAndCount(MappedView_Transient, mappingBufferSize, mappingBufferCount);
 		logger.Info(TC("  DirectoryTable      %7u %9s"), u32(m_directoryTable.m_lookup.size()), BytesToText(GetDirectoryTableSize()).str);
 		logger.Info(TC("  MappingTable        %7u %9s"), u32(m_fileMappingTableLookup.size()), BytesToText(GetFileMappingSize()).str);
-		logger.Info(TC("  MappingBuffer       %7u %9s"), mappingBufferCount, BytesToText(mappingBufferSize).str);
+		logger.Info(TC("  MappingBuffer       %7u %9s"), 1, BytesToText(m_fileMappingBuffer.GetUsed()).str);
 		m_stats.Print(logger);
 		logger.Info(TC(""));
 	}
@@ -3620,7 +3635,9 @@ namespace uba
 		if (sourceSize + sourceOffset > virtualFile.size)
 			return m_logger.Error(TC("Virtual file offset(%llu)+size(%llu) outside source file size(%llu)"), sourceOffset, sourceSize, virtualFile.size, filePath.data);
 
-		return RegisterVirtualFileInternal(fileNameKey, filePath, virtualFile.mappingHandle, sourceSize, sourceOffset, false);
+		SharedMemoryHandle memoryHandle = m_sharedMemory.RegisterExternalMapping(virtualFile.mappingHandle);
+
+		return RegisterVirtualFileInternal(fileNameKey, filePath, memoryHandle, sourceSize, sourceOffset, false);
 	}
 
 	bool Session::CreateVirtualFileInternal(const StringKey& fileNameKey, const StringView& filePath, const void* memory, u64 memorySize, bool transient)
@@ -3632,21 +3649,23 @@ namespace uba
 		if (!insres.second)
 			return false;
 		VirtualSourceFile& virtualFile = insres.first->second;
-		FileMappingHandle mapping = CreateMemoryMappingW(m_logger, PAGE_READWRITE, memorySize, nullptr, filePath.data);
+
+		SharedMemoryHandle mapping = m_sharedMemory.CreateHandle(m_logger, filePath.data);
 		if (!mapping.IsValid())
 			return false;
-		u8* mem2 = MapViewOfFile(m_logger, mapping, FILE_MAP_WRITE, 0, memorySize);
+		m_sharedMemory.ExtendMemory(mapping, memorySize);
+		u8* mem2 = m_sharedMemory.MapView(mapping, filePath.data, SharedMemoryMapType_ReadWrite);
 		if (!mem2)
 			return false;
 		MapMemoryCopy(mem2, memory, memorySize);
-		UnmapViewOfFile(m_logger, mem2, memorySize, filePath.data);
+		m_sharedMemory.UnmapView(mapping, mem2);
 
-		virtualFile.mappingHandle = mapping;
+		virtualFile.memoryHandle = mapping;
 		virtualFile.size = memorySize;
-		return RegisterVirtualFileInternal(fileNameKey, filePath, virtualFile.mappingHandle, memorySize, 0, transient);
+		return RegisterVirtualFileInternal(fileNameKey, filePath, virtualFile.memoryHandle, memorySize, 0, transient);
 	}
 
-	bool Session::RegisterVirtualFileInternal(const StringKey& fileNameKey, const StringView& filePath, FileMappingHandle mappingHandle, u64 mappingSize, u64 mappingOffset, bool transient)
+	bool Session::RegisterVirtualFileInternal(const StringKey& fileNameKey, const StringView& filePath, SharedMemoryHandle memoryHandle, u64 mappingSize, u64 mappingOffset, bool transient)
 	{
 		SCOPED_FUTEX(m_fileMappingTableLookupLock, lookupLock);
 		auto insres = m_fileMappingTableLookup.try_emplace(fileNameKey);
@@ -3658,8 +3677,8 @@ namespace uba
 		if (entry.handled)
 			return m_logger.Error(TC("Virtual file %s has already been registered"), filePath.data);
 
-		entry.mapping = mappingHandle;
-		entry.mappingOffset = mappingOffset;
+		entry.memoryHandle = memoryHandle;
+		entry.memoryOffset = mappingOffset;
 		entry.contentSize = mappingSize;
 		entry.handled = true;
 		entry.lastWriteTime = 0; // TODO: Take lastwritetime of source file?
@@ -3667,12 +3686,13 @@ namespace uba
 		entry.isInvisible = transient;
 
 		StringBuffer<> mappingName;
-		Storage::GetMappingString(mappingName, mappingHandle, mappingOffset);
+		Storage::GetMappingString(mappingName, memoryHandle, mappingOffset, entry.canBeFreed);
+
 		entry.success = true;
 
 		#if UBA_DEBUG_TRACK_MAPPING
 		entry.name = filePath.data;
-		m_debugLogger->Info(TC("Mapping created 0x%llx (%s) from virtual file"), u64(entry.mapping.mh), entry.name.c_str());
+		m_debugLogger->Info(TC("Mapping created 0x%llx (%s) from virtual file"), u64(entry.memoryHandle.ToU64()), entry.name.c_str());
 		#endif
 
 		if (!transient)
@@ -3819,19 +3839,11 @@ namespace uba
 
 	bool Session::ExtractSymbolsFromObjectFile(const CloseFileMessage& msg, const tchar* fileName, u64 fileSize)
 	{
-		if (!msg.mappingHandle.IsValid())
+		if (!msg.memoryHandle.IsValid())
 			return m_logger.Error(TC("Can't extract symbols from obj file that is written directly to disk (%s writing %s)"), msg.process.m_startInfo.application, fileName);
 
-		FileMappingHandle objectFileMappingHandle;
-
-		if (!DuplicateFileMapping(m_logger, msg.process.m_nativeProcessHandle, msg.mappingHandle, GetCurrentProcessHandle(), objectFileMappingHandle, FILE_MAP_ALL_ACCESS, false, 0, fileName))
-			return m_logger.Error(TC("Failed to duplicate file mapping handle for %s"), fileName);
-		auto ofmh = MakeGuard([&]() { CloseFileMapping(m_logger, objectFileMappingHandle, fileName); });
-
-		u8* mem = MapViewOfFile(m_logger, objectFileMappingHandle, FILE_MAP_ALL_ACCESS, 0, fileSize);
-		if (!mem)
-			return m_logger.Error(TC("Failed to map view of filehandle for read %s (%s)"), fileName, LastErrorToText().data);
-		auto memClose = MakeGuard([&](){ UnmapViewOfFile(m_logger, mem, fileSize, fileName); });
+		u8* mem = m_sharedMemory.MapView(msg.memoryHandle, fileName);
+		auto memClose = MakeGuard([&]() { m_sharedMemory.UnmapView(msg.memoryHandle, mem); });
 
 		ObjectFile* objectFile = ObjectFile::Parse(m_logger, ObjectFileParseMode_All, mem, fileSize, fileName);
 		if (!objectFile)
@@ -3851,16 +3863,16 @@ namespace uba
 		if (!objectFile->WriteImportsAndExports(m_logger, memoryBlock, verbose))
 			return false;
 
-		FileMappingHandle symHandle = CreateMemoryMappingW(m_logger, PAGE_READWRITE, memoryBlock.writtenSize, nullptr, TC("SymHandle"));
-		if (!symHandle.IsValid())
-			return false;
-		auto mg = MakeGuard([&]() { CloseFileMapping(m_logger, symHandle, TC("SymHandle")); });
-		u8* mem2 = MapViewOfFile(m_logger, symHandle, FILE_MAP_ALL_ACCESS, 0, memoryBlock.writtenSize);
+		SharedMemoryHandle symHandle = m_sharedMemory.CreateHandle(m_logger, TC("SymHandle"));
+		auto mg = MakeGuard([&]() { m_sharedMemory.CloseHandle(m_logger, symHandle, TC("SymHandle")); });
+		m_sharedMemory.ExtendMemory(symHandle, memoryBlock.writtenSize);
+		u8* mem2 = m_sharedMemory.MapView(symHandle, TC("SymHandle"), SharedMemoryMapType_ReadWrite);
 		if (!mem2)
 			return false;
 
 		MapMemoryCopy(mem2, memoryBlock.memory, memoryBlock.writtenSize);
-		UnmapViewOfFile(m_logger, mem2, memoryBlock.writtenSize, TC("SymHandle"));
+
+		m_sharedMemory.UnmapView(symHandle, mem2);
 
 		StringKey symFileKey = CaseInsensitiveFs ? ToStringKeyLower(exportsFile) : ToStringKey(exportsFile);
 		u64 lastWriteTime = GetSystemTimeAsFileTime();
@@ -3877,8 +3889,8 @@ namespace uba
 		writtenFile.key = symFileKey;
 		writtenFile.owner = &msg.process;
 		writtenFile.attributes = msg.attributes;
-		writtenFile.mappingHandle = symHandle;
-		writtenFile.mappingWritten = memoryBlock.writtenSize;
+		writtenFile.memoryHandle = symHandle;
+		writtenFile.memoryWritten = memoryBlock.writtenSize;
 		writtenFile.lastWriteTime = lastWriteTime;
 		writtenFile.name = exportsFile.data;
 
@@ -4000,7 +4012,7 @@ namespace uba
 				if (!func)
 					return true;
 
-				if (out.fileName[0] == '^')
+				if (IsMemoryHandle(out.fileName.data))
 				{
 					MappedView view = m_fileMappingBuffer.MapView(out.fileName, out.size, fileName.data);
 					if (view.memory)
diff --git a/Engine/Source/Programs/UnrealBuildAccelerator/Common/Private/UbaSessionClient.cpp b/Engine/Source/Programs/UnrealBuildAccelerator/Common/Private/UbaSessionClient.cpp
--- a/Engine/Source/Programs/UnrealBuildAccelerator/Common/Private/UbaSessionClient.cpp
+++ b/Engine/Source/Programs/UnrealBuildAccelerator/Common/Private/UbaSessionClient.cpp
@@ -659,13 +659,13 @@ namespace uba
 							return m_logger.Error(TC("Error retrieving cas entry %s (%s) (2)"), CasKeyString(casKey).str, fileName.data);
 						entry.success = true;
 						entry.contentSize = result.size;
-						entry.mapping = result.view.handle;
-						entry.mappingOffset = result.view.offset;
+						entry.memoryHandle = result.view.handle;
+						entry.memoryOffset = result.view.offset;
 					}
 
 					fileSize = entry.contentSize;
-					if (entry.mapping.IsValid())
-						Storage::GetMappingString(newName, entry.mapping, entry.mappingOffset);
+					if (entry.memoryHandle.IsValid())
+						Storage::GetMappingString(newName, entry.memoryHandle, entry.memoryOffset, entry.canBeFreed);
 					else
 						newName.Append(entry.isDir ? TC("$d") : TC("$f"));
 				}
@@ -673,7 +673,7 @@ namespace uba
 
 			UBA_ASSERTF(!newName.IsEmpty(), TC("No casfile available for %s using %s"), fileName.data, CasKeyString(casKey).str);
 
-			if (newName[0] != '^')
+			if (!IsMemoryHandle(newName.data))
 			{
 				TrackHintScope ths(tws, AsView(TC("CreateMemoryMap")));
 				if (!isDir && memoryMapAlignment)
@@ -682,7 +682,8 @@ namespace uba
 						retryEntryGuard.Execute();
 
 					MemoryMap map;
-					if (!GetOrCreateMemoryMapFromFile(map, fileNameKey, newName.data, IsCompressed(newCasKey), memoryMapAlignment, fileName.data, nullptr, false))
+					bool canBeFreed = false;
+					if (!GetOrCreateMemoryMapFromFile(map, fileNameKey, newName.data, IsCompressed(newCasKey), memoryMapAlignment, fileName.data, nullptr, canBeFreed))
 					{
 						if (!shouldRetry)
 							return false;
@@ -730,8 +731,8 @@ namespace uba
 
 			if (StringView(file.name).StartsWith(m_tempPath) && !rules.KeepTempOutputFile(StringView(file.name)))
 			{
-				CloseFileMapping(m_logger, file.mappingHandle, file.backedName.c_str());
-				file.mappingHandle = {};
+				m_sharedMemory.CloseHandle(m_logger, file.memoryHandle, file.backedName.c_str());
+				file.memoryHandle = {};
 				#if UBA_DEBUG
 				m_logger.Info(TC("Skipping sending back temp file %s"), file.name.c_str());
 				#endif
@@ -740,10 +741,10 @@ namespace uba
 
 			TimerScope timer(sendFiles);
 #ifdef _DEBUG
-			if (!pair.second.mappingHandle.IsValid())
+			if (!pair.second.memoryHandle.IsValid())
 				m_logger.Warning(TC("%s is not using file mapping"), file.name.c_str());
 #endif
-			bool keepMappingInMemory = IsWindows && !IsRarelyReadAfterWritten(process, file.name) && file.mappingWritten < m_keepOutputFileMemoryMapsThreshold;
+			bool keepMappingInMemory = IsWindows && IsMaybeReadAfterWritten(process, file.name) && file.memoryWritten < m_keepOutputFileMemoryMapsThreshold;
 			bool compressed = rules.SendFileCompressedFromClient(file.name);
 			if (!SendFile(file, process.GetId(), keepMappingInMemory, compressed))
 				return false;
@@ -756,14 +757,14 @@ namespace uba
 		CasKey casKey;
 		{
 			TimerScope ts(m_stats.storageSend);
-			if (!m_storage.StoreCasFileClient(casKey, file.key, file.backedName.c_str(), file.mappingHandle, 0, file.mappingWritten, file.name.c_str(), keepMappingInMemory, compressed))
+			if (!m_storage.StoreCasFileClient(casKey, file.key, file.backedName.c_str(), file.memoryHandle, 0, file.memoryWritten, file.name.c_str(), keepMappingInMemory, compressed))
 				return false;
 		}
 		if (casKey == CasKeyZero)
-			return m_logger.Warning(TC("Failed to store cas on server for local file %s (size %llu destination %s)"), file.backedName.c_str(), file.mappingWritten, file.name.c_str());
+			return m_logger.Warning(TC("Failed to store cas on server for local file %s (size %llu destination %s)"), file.backedName.c_str(), file.memoryWritten, file.name.c_str());
 	
-		CloseFileMapping(m_logger, file.mappingHandle, file.backedName.c_str());
-		file.mappingHandle = {};
+		m_sharedMemory.CloseHandle(m_logger, file.memoryHandle, file.backedName.c_str());
+		file.memoryHandle = {};
 
 		StackBinaryReader<128> reader;
 		{
@@ -2721,5 +2722,9 @@ namespace uba
 
 		// send and recv are swapped on purpose because that is how visualizer is visualizing 
 		m_trace.SessionUpdate(0, 0, send, recv, m_lastPing, m_memAvail, m_memTotal, m_cpuUsage);
+
+		m_sharedMemory.TraceStats(m_trace, 30);
+		m_fileMappingBuffer.TraceStatus(m_trace, 20, TC("Session"), 1);
+		m_storage.TraceUpdate(m_trace, 300);
 	}
 }
diff --git a/Engine/Source/Programs/UnrealBuildAccelerator/Common/Private/UbaSessionServer.cpp b/Engine/Source/Programs/UnrealBuildAccelerator/Common/Private/UbaSessionServer.cpp
--- a/Engine/Source/Programs/UnrealBuildAccelerator/Common/Private/UbaSessionServer.cpp
+++ b/Engine/Source/Programs/UnrealBuildAccelerator/Common/Private/UbaSessionServer.cpp
@@ -261,6 +261,8 @@ namespace uba
 
 		if (m_trace.IsWriting())
 		{
+			m_sharedMemory.TraceStats(m_trace, 30);
+
 			m_trace.WriteSessionSummary([&](Logger& logger)
 				{
 					PrintSummary(logger);
@@ -2302,6 +2304,21 @@ namespace uba
 		}
 
 		m_trace.SessionUpdate(0, serverConnectionCount, serverSend, serverRecv, 0, memAvail, memTotal, cpuLoad);
+
+#if 0
+		m_sharedMemory.TraceStats(m_trace, 30);
+
+		m_fileMappingBuffer.TraceStatus(m_trace, 20, TC("Session"), 1);
+
+		m_trace.StatusUpdate(40, 1, TCV("Independents"), LogEntryType_Info);
+		StringBuffer<> str;
+		u64 created = m_independentMappingCreated;
+		u64 count = m_independentMappingActive;
+		str.Appendf(TC("Active: %llu, Created: %llu"), count, created);
+		m_trace.StatusUpdate(40, 7, str, LogEntryType_Info);
+
+		m_storage.TraceUpdate(m_trace, 300);
+#endif
 	}
 
 	void SessionServer::WriteRemoteEnvironmentVariables(BinaryWriter& writer)
@@ -2388,7 +2405,7 @@ namespace uba
 
 	bool SessionServer::HandleDebugFileNotFoundError(const ConnectionInfo& connectionInfo, const WorkContext& workContext, BinaryReader& reader, BinaryWriter& writer)
 	{
-#if PLATFORM_WINDOWS
+#if 0//PLATFORM_WINDOWS
 		StringBuffer<> errorPath;
 		reader.ReadString(errorPath);
 		StringBuffer<> workDir;
diff --git a/Engine/Source/Programs/UnrealBuildAccelerator/Common/Private/UbaSharedMemoryAllocator.cpp b/Engine/Source/Programs/UnrealBuildAccelerator/Common/Private/UbaSharedMemoryAllocator.cpp
--- a/Engine/Source/Programs/UnrealBuildAccelerator/Common/Private/UbaSharedMemoryAllocator.cpp
+++ b/Engine/Source/Programs/UnrealBuildAccelerator/Common/Private/UbaSharedMemoryAllocator.cpp
@@ -8,19 +8,12 @@
 #if PLATFORM_WINDOWS
 #include <windows.h>
 #else
-#define GetCurrentProcess() 0
-#define UnmapViewOfFile2(...) false
-#define VirtualFree(...) false
-#define VirtualAlloc(...) nullptr
-#define VirtualAlloc2(...) nullptr
-#define MapViewOfFile3(...) nullptr
-#define ZeroMemory(...) {}
-#define PAGE_READONLY 0x02
-#define PAGE_READWRITE 0x04
+#define ZeroMemory(mem, size) memset(mem, 0, size)
+#define DiscardVirtualMemory(...)
 #define PAGE_WRITECOPY 0x08
 #endif
 
-#define UBA_DEBUG_HINTS 1
+#define UBA_DEBUG_HINTS (UBA_DEBUG || UBA_TRACK_SHARED_MEMORY_ALLOCATIONS)
 
 namespace uba
 {
@@ -46,21 +39,30 @@ namespace uba
 		#else
 		const tchar* GetHint() { return TC("EnableDebugHints"); }
 		#endif
+
+		#if UBA_TRACK_SHARED_MEMORY_ALLOCATIONS
+		Atomic<u64> committed = 0;
+		InternalHandle* next;
+		InternalHandle* prev;
+		#endif
 	};
 
+
 	#define UBA_REF_COUNT(f) UBA_ASSERTF(h.refCount, TC( #f " accessing released shared memory handle (%s)"), h.GetHint());
 
 	////////////////////////////////////////////////////////////////////////////////////////////////////
 
 	bool SharedMemoryAllocator::Init(u64 capacity, u64 commitStepSize)
 	{
-		m_commitStepSize = commitStepSize;
-		m_reserved = AlignUp(capacity, 64 * 1024);
-		m_fileMapping = CreateMemoryMappingW(m_logger, PAGE_READWRITE|SEC_RESERVE, m_reserved, nullptr, TC("SharedMemoryAllocator"));
+		m_commitStepSize = Min(commitStepSize, capacity);
+		m_capacity = AlignUp(capacity, PageSize);
+		m_fileMapping = CreateMemoryMappingW(m_logger, PAGE_READWRITE|SEC_RESERVE, m_capacity, nullptr, TC("SharedMemoryAllocator"));
 		if (!m_fileMapping.IsValid())
 			return false;
 
-		m_memory = MapViewOfFile(m_logger, m_fileMapping, FILE_MAP_READ | FILE_MAP_WRITE, 0, 0, m_reserved);
+		m_memory = MapViewOfFile(m_logger, m_fileMapping, FILE_MAP_READ|FILE_MAP_WRITE, 0, m_capacity);
+		UBA_ASSERTF(m_memory, TC("MapViewOfFile failed trying to map %llu bytes (%s)"), m_capacity, LastErrorToText().data);
+
 		if (!m_memory)
 		{
 			CloseFileMapping(m_logger, m_fileMapping, TC("SharedMemoryAllocator"));
@@ -69,14 +71,25 @@ namespace uba
 		return true;
 	}
 
+	void SharedMemoryAllocator::Deinit(bool keepFileMapping)
+	{
+		if (m_memory)
+			UnmapViewOfFile(m_logger, m_memory, m_capacity, TC("SharedMemoryAllocator"));
+		m_memory = nullptr;
+
+		if (m_fileMapping.IsValid() && !keepFileMapping)
+			CloseFileMapping(m_logger, m_fileMapping, TC("SharedMemoryAllocator"));
+		m_fileMapping = {};
+	}
+
 	SharedMemoryAllocator::operator SharedMemoryAllocatorHandle() const
 	{
-		return { m_fileMapping.ToU64() };
+		return m_fileMapping;
 	}
 
-	u64 SharedMemoryAllocator::GetReserved()
+	u64 SharedMemoryAllocator::GetCapacity()
 	{
-		return m_reserved;
+		return m_capacity;
 	}
 
 	SharedMemoryHandle SharedMemoryAllocator::CreateHandle(Logger& logger, const tchar* hint, bool priority)
@@ -88,6 +101,16 @@ namespace uba
 		h->hint = hint;
 		#endif
 
+		#if UBA_TRACK_SHARED_MEMORY_ALLOCATIONS
+		{
+			SCOPED_FUTEX(m_handleLock, l);
+			if (m_firstHandle)
+				m_firstHandle->prev = h;
+			h->next = m_firstHandle;
+			m_firstHandle = h;
+		}
+		#endif
+
 		if (priority)
 		{
 			u32 index = m_priorityHandleCounter++;
@@ -98,17 +121,25 @@ namespace uba
 			}
 		}
 
-		//LoggerWithWriter(g_consoleLogWriter, TC("")).Info(TC("CREATED %s (%llu)"), hint, h);
-
 		return { u64(h) };
 	}
 
+	SharedMemoryHandle SharedMemoryAllocator::DuplicateHandle(SharedMemoryHandle handle, const tchar* hint)
+	{
+		if (!handle.internalHandle)
+			return {};
+		auto& h = GetHandle(handle);
+		UBA_REF_COUNT(DuplicateHandle);
+		++h.refCount;
+		return handle;
+	}
+
 	void SharedMemoryAllocator::ClearHandle(SharedMemoryHandle handle)
 	{
 		if (!handle.internalHandle)
 			return;
 		auto& h = GetHandle(handle);
-		UBA_REF_COUNT(ClearHandle)
+		UBA_REF_COUNT(ClearHandle);
 		ClearHandle(h);
 	}
 
@@ -127,7 +158,7 @@ namespace uba
 		auto& slices = h.slices;
 		
 		u32 sliceCount = u32(slices.size());
-		ExtendHandle(h, size, hint, zeroMemory, 0);
+		ExtendHandle(h, size, hint, zeroMemory);
 		u32 newCount = u32(slices.size());
 
 		out.WriteU32(newCount - sliceCount);
@@ -144,10 +175,10 @@ namespace uba
 		UBA_ASSERT(handle.internalHandle);
 		auto& h = GetHandle(handle);
 		UBA_REF_COUNT(ExtendMemory);
-		ExtendHandle(h, size, hint, zeroMemory, 0);
+		ExtendHandle(h, size, hint, zeroMemory);
 	}
 
-	void SharedMemoryAllocator::ExtendMemory(SharedMemoryView& view, SharedMemoryHandle handle, u64 size, const tchar* hint, bool zeroMemory, u64 preferredOffset)
+	void SharedMemoryAllocator::ExtendMemory(SharedMemoryView& view, SharedMemoryHandle handle, u64 size, const tchar* hint, bool zeroMemory)
 	{
 		UBA_ASSERT(handle.internalHandle);
 		auto& h = GetHandle(handle);
@@ -155,12 +186,121 @@ namespace uba
 
 		auto& slices = h.slices;
 		u32 sliceCount = u32(slices.size());
-		ExtendHandle(h, size, hint, zeroMemory, preferredOffset);
+		ExtendHandle(h, size, hint, zeroMemory);
 
 		for (u32 i=sliceCount, e= u32(slices.size()); i!=e; ++i)
 			view.ExtendMemory(*this, slices[i].offset, slices[i].size, SharedMemoryMapType_ReadWrite, hint);
 	}
 
+	void SharedMemoryAllocator::TrimMemory(SharedMemoryHandle handle, u64 usedSize)
+	{
+		UBA_ASSERT(handle.internalHandle);
+		auto& h = GetHandle(handle);
+		UBA_REF_COUNT(TrimMemory);
+		UBA_ASSERT(h.refCount == 1);
+
+		usedSize = AlignUp(usedSize, PageSize);
+
+		#if UBA_TRACK_SHARED_MEMORY_ALLOCATIONS
+		h.committed = usedSize;
+		#endif
+
+		auto& slices = h.slices;
+		for (u32 i = 0, e = u32(slices.size()); i != e; ++i)
+		{
+			u64 sliceSize = slices[i].size;
+			if (usedSize <= sliceSize)
+			{
+				SCOPED_FUTEX(m_lock, lock);
+				if (u64 left = sliceSize - usedSize)
+				{
+					u64 offset = slices[i].offset + usedSize;
+					FreeNoLock({ offset, left });
+					slices[i].size = usedSize;
+				}
+				for (u32 j = ++i; j != e; ++j)
+					FreeNoLock(slices[j]);
+				slices.resize(i);
+				return;
+			}
+
+			usedSize -= sliceSize;
+		}
+	}
+
+	u8* SharedMemoryAllocator::ReserveMemory(SharedMemoryHandle handle, u64 size)
+	{
+		UBA_ASSERT(AlignUp(size, PageSize) == size);
+		UBA_ASSERT(handle.internalHandle);
+		auto& h = GetHandle(handle);
+		UBA_REF_COUNT(ReserveMemory);
+
+		UBA_ASSERT(h.slices.empty());
+
+		auto& slice = h.slices.emplace_back();
+		SCOPED_FUTEX(m_lock, lock);
+		slice.offset = m_reserved;
+		m_reserved += size;
+		slice.size = size;
+		lock.Leave();
+
+		return m_memory + slice.offset;
+	}
+
+	void SharedMemoryAllocator::CommitMemory(SharedMemoryHandle handle, u64 offset, u64 size)
+	{
+		UBA_ASSERT(AlignUp(offset, PageSize) == offset);
+		UBA_ASSERT(AlignUp(size, PageSize) == size);
+		UBA_ASSERT(handle.internalHandle);
+		auto& h = GetHandle(handle);
+		UBA_REF_COUNT(CommitMemory);
+
+		UBA_ASSERT(h.slices.size() == 1);
+		auto& slice = h.slices.back();
+		MapViewCommit(m_logger, m_memory + slice.offset + offset, size);
+
+		SCOPED_FUTEX(m_lock, lock);
+		m_committed += size;
+		m_used += size;
+
+		#if UBA_TRACK_SHARED_MEMORY_ALLOCATIONS
+		h.committed += size;
+		#endif
+	}
+
+	void SharedMemoryAllocator::DecommitMemory(SharedMemoryHandle handle, u64 offset, u64 size)
+	{
+		if (!handle.internalHandle)
+			return;
+		if (!size)
+			return;
+		UBA_ASSERT(AlignUp(offset, PageSize) == offset);
+		UBA_ASSERT(AlignUp(size, PageSize) == size);
+		UBA_ASSERT(handle.internalHandle);
+		auto& h = GetHandle(handle);
+		UBA_REF_COUNT(DecommitMemory);
+		UBA_ASSERT(h.slices.size() == 1);
+		auto& slice = h.slices.back();
+
+		#if UBA_TRACK_SHARED_MEMORY_ALLOCATIONS
+		h.committed -= size;
+		#endif
+
+		// Unfortunately windows has no way of decommitting memory from anonymous file mappings.
+		SCOPED_FUTEX(m_lock, lock);
+		FreeNoLock({ slice.offset + offset, size });
+	}
+
+	void SharedMemoryAllocator::UnreserveMemory(SharedMemoryHandle handle)
+	{
+		if (!handle.internalHandle)
+			return;
+		UBA_ASSERT(handle.internalHandle);
+		auto& h = GetHandle(handle);
+		UBA_REF_COUNT(UnreserveMemory);
+		h.slices.clear();
+	}
+
 	u64 SharedMemoryAllocator::GetCommitted(SharedMemoryHandle handle)
 	{
 		auto& h = GetHandle(handle);
@@ -206,18 +346,14 @@ namespace uba
 
 			if (size <= sliceSize)
 			{
-				out.Write7BitEncoded(AlignUp(size, 64*1024));
+				out.Write7BitEncoded(AlignUp(size, PageSize));
 				break;
 			}
 
 			out.Write7BitEncoded(sliceSize);
 			size -= sliceSize;
-			if (!size)
-				break;
 		}
 
-		//wprintf(L"WRITING %u slices for %s\n", count, h.GetHint());
-
 		totCount = count;
 	}
 
@@ -226,6 +362,16 @@ namespace uba
 		if (!handle.internalHandle)
 			return false;
 
+		if (handle.internalHandle & 0x8000'0000'0000'0000)
+		{
+			// TODO: Right now we leak this mapping. Not good..
+			// .. but only used in very special situations
+			FileMappingHandle file = FileMappingHandle::FromU64(handle.internalHandle & ~0x8000'0000'0000'0000);
+			u8* mem = MapViewOfFile(m_logger, file, FILE_MAP_READ, offset, size);
+			view.ReferenceExternalMemory(mem, size);
+			return true;
+		}
+
 		auto& h = GetHandle(handle);
 		UBA_REF_COUNT(MapView);
 
@@ -280,7 +426,7 @@ namespace uba
 		return true;
 	}
 
-	u8* SharedMemoryAllocator::MapView(SharedMemoryHandle handle, SharedMemoryMapType mapType)
+	u8* SharedMemoryAllocator::MapView(SharedMemoryHandle handle, const tchar* hint, SharedMemoryMapType mapType)
 	{
 		if (!handle.internalHandle)
 			return nullptr;
@@ -306,25 +452,42 @@ namespace uba
 		else if (mapType == SharedMemoryMapType_CopyOnWrite)
 			protect = PAGE_WRITECOPY;
 
-		u8* memory = (u8*)VirtualAlloc2(GetCurrentProcess(), nullptr, committed, MEM_RESERVE | MEM_RESERVE_PLACEHOLDER, PAGE_NOACCESS, nullptr, 0);
+		auto& kernelStats = KernelStats::GetCurrent();
+
+		u8* memory;
+		{
+			TimerScope ts(kernelStats.virtualAlloc2);
+			#if PLATFORM_WINDOWS
+			memory = (u8*)VirtualAlloc2(GetCurrentProcess(), nullptr, committed, MEM_RESERVE | MEM_RESERVE_PLACEHOLDER, PAGE_NOACCESS, nullptr, 0);
+			#else
+			memory = (u8*)mmap(NULL, committed, PROT_NONE, MAP_PRIVATE | MAP_ANONYMOUS, -1, 0);
+			#endif
+		}
 		UBA_ASSERT(memory);
 
-		auto& stats = KernelStats::GetCurrent().mapViewOfFile3;
 		u64 pos = 0;
 		for (Slice slice : h.slices)
 		{
 			u8* targetAddress = memory + pos;
 
+			#if PLATFORM_WINDOWS
 			if (pos + slice.size < committed)
 			{
+				TimerScope ts(kernelStats.virtualFree);
 				bool res = VirtualFree(targetAddress, slice.size, MEM_RELEASE | MEM_PRESERVE_PLACEHOLDER);
 				UBA_ASSERTF(res, TC("VirtualFree failed for pos %llu with size %llu (%s)"), pos, slice.size, LastErrorToText().data); (void)res;
 			}
+			#endif
 
-			stats.bytes += slice.size;
-			TimerScope ts(stats);
+			kernelStats.mapViewOfFile3.bytes += slice.size;
+			TimerScope ts(kernelStats.mapViewOfFile3);
 
+			#if PLATFORM_WINDOWS
 			void* viewA = MapViewOfFile3(m_fileMapping.mh, GetCurrentProcess(), targetAddress, slice.offset, slice.size, MEM_REPLACE_PLACEHOLDER, protect, nullptr, 0);
+			#else
+			void* viewA = mmap(targetAddress, slice.size, protect, MAP_SHARED | MAP_FIXED, m_fileMapping.shmFd, slice.offset);
+			#endif
+
 			UBA_ASSERTF(viewA, TC("MapViewOfFile3 failed in MapView 0x%p %llu for % s (%s)"), memory, pos, h.GetHint(), LastErrorToText().data); (void)viewA;
 			pos += slice.size;
 		}
@@ -343,32 +506,41 @@ namespace uba
 
 		if (h.slices.size() == 1)
 		{
-			UBA_ASSERT(memory >= m_memory && memory < m_memory + m_committed);
+			UBA_ASSERT(memory >= m_memory && memory < m_memory + m_reserved);
 			ReleaseHandle(handle);
 			return;
 		}
 		u64 pos = 0;
 		u32 res; (void)res;
 
-		auto& stats = KernelStats::GetCurrent().unmapViewOfFile2;
-
-		for (Slice slice : h.slices)
-		{
-			stats.bytes += slice.size;
-			TimerScope ts(stats);
-			res = UnmapViewOfFile2(GetCurrentProcess(), memory + pos, MEM_PRESERVE_PLACEHOLDER);
-			UBA_ASSERTF(res, TC("UnmapViewOfFile2 failed on pos %llu for %s (%s)"), pos, h.GetHint(), LastErrorToText().data);
-			pos += slice.size;
-		}
-		if (h.slices.size() > 1)
 		{
-			res = VirtualFree(memory, pos, MEM_RELEASE | MEM_COALESCE_PLACEHOLDERS);
-			UBA_ASSERTF(res, TC("VirtualFree failed (%s)"), LastErrorToText().data);
-		}
+			auto& kernelStats = KernelStats::GetCurrent();
 
-		res = VirtualFree(memory, 0, MEM_RELEASE);
-		UBA_ASSERTF(res, TC("VirtualFree failed (%s)"), LastErrorToText().data);
+			#if PLATFORM_WINDOWS
+			for (Slice slice : h.slices)
+			{
+				TimerScope ts(kernelStats.unmapViewOfFile2);
+				res = UnmapViewOfFile2(GetCurrentProcess(), memory + pos, MEM_PRESERVE_PLACEHOLDER);
+				UBA_ASSERTF(res, TC("UnmapViewOfFile2 failed on pos %llu for %s (%s)"), pos, h.GetHint(), LastErrorToText().data);
+				pos += slice.size;
+			}
+			if (h.slices.size() > 1)
+			{
+				TimerScope ts(kernelStats.virtualFree);
+				res = VirtualFree(memory, pos, MEM_RELEASE | MEM_COALESCE_PLACEHOLDERS);
+				UBA_ASSERTF(res, TC("VirtualFree failed (%s)"), LastErrorToText().data);
+			}
 
+			TimerScope ts(kernelStats.virtualFree);
+			res = VirtualFree(memory, 0, MEM_RELEASE);
+			UBA_ASSERTF(res, TC("VirtualFree failed (%s)"), LastErrorToText().data);
+			#else
+			u64 committed = 0;
+			for (Slice slice : h.slices)
+				committed += slice.size;
+			munmap(memory, committed);
+			#endif
+		}
 		ReleaseHandle(handle);
 	}
 
@@ -379,46 +551,70 @@ namespace uba
 
 	SharedMemoryAllocator::~SharedMemoryAllocator()
 	{
-		if (m_memory)
-			UnmapViewOfFile(m_logger, m_memory, m_reserved, TC("SharedMemoryAllocator"));
-		if (m_fileMapping.IsValid())
-			CloseFileMapping(m_logger, m_fileMapping, TC("SharedMemoryAllocator"));
+		Deinit(false);
+	}
+
+	SharedMemoryHandle SharedMemoryAllocator::RegisterExternalMapping(FileMappingHandle file)
+	{
+		SharedMemoryHandle h;
+		h.internalHandle = 0x8000'0000'0000'0000 + file.ToU64();
+		return h;
 	}
 
 	SharedMemoryAllocator::InternalHandle& SharedMemoryAllocator::GetHandle(SharedMemoryHandle h)
 	{
+		UBA_ASSERT(!(h.internalHandle & 0x8000'0000'0000'0000));
 		if (h.internalHandle <= sizeof_array(m_priorityHandles))
 			return *m_priorityHandles[h.internalHandle - 1];
 		return *(InternalHandle*)h.internalHandle;
 	}
 
-	void SharedMemoryAllocator::ReleaseHandle(SharedMemoryHandle h)
+	void SharedMemoryAllocator::ReleaseHandle(SharedMemoryHandle handle)
 	{
-		InternalHandle& handle = GetHandle(h);
-		UBA_ASSERT(handle.refCount);
-		if (--handle.refCount)
+		InternalHandle& h = GetHandle(handle);
+		UBA_REF_COUNT(ReleaseHandle);
+		if (--h.refCount)
 			return;
-		ClearHandle(handle);
+
+		#if UBA_TRACK_SHARED_MEMORY_ALLOCATIONS
+		{
+			SCOPED_FUTEX(m_handleLock, l);
+			if (m_firstHandle == &h)
+				m_firstHandle = h.next;
+			if (h.prev)
+				h.prev->next = h.next;
+			if (h.next)
+				h.next->prev = h.prev;
+		}
+		#endif
+
+		ClearHandle(h);
+
 		#if !UBA_DEBUG_HINTS
-		delete& handle;
+		delete& h;
 		#endif
-		if (h.internalHandle <= sizeof_array(m_priorityHandles))
-			m_priorityHandles[h.internalHandle - 1] = 0;
+		if (handle.internalHandle <= sizeof_array(m_priorityHandles))
+			m_priorityHandles[handle.internalHandle - 1] = 0;
 		--m_handleCount;
 	}
 
-	void SharedMemoryAllocator::ExtendHandle(InternalHandle& handle, u64 size, const tchar* hint, bool zeroMemory, u64 preferredOffset)
+	void SharedMemoryAllocator::ExtendHandle(InternalHandle& handle, u64 size, const tchar* hint, bool zeroMemory)
 	{
-		size = AlignUp(size, 64 * 1024);
+		size = AlignUp(size, PageSize);
 		u64 left = size;
 
+		#if UBA_TRACK_SHARED_MEMORY_ALLOCATIONS
+		handle.committed += size;
+		#endif
+
 		auto& slices = handle.slices;
+
 		u32 sliceCount = u32(slices.size());
 
 		SCOPED_FUTEX(m_lock, lock);
 		while (left)
 		{
-			Slice slice = AllocateNoLock(left, preferredOffset);
+			Slice slice = AllocateNoLock(left, slices.size());
 			slices.push_back(slice);
 			left -= slice.size;
 		}
@@ -433,6 +629,11 @@ namespace uba
 	{
 		if (handle.slices.empty())
 			return;
+
+		#if UBA_TRACK_SHARED_MEMORY_ALLOCATIONS
+		handle.committed = 0;
+		#endif
+
 		SCOPED_FUTEX(m_lock, lock);
 		for (auto& slice : handle.slices)
 			FreeNoLock(slice);
@@ -440,18 +641,20 @@ namespace uba
 		handle.slices.clear();
 	}
 
-	SharedMemoryAllocator::Slice SharedMemoryAllocator::AllocateNoLock(u64 desiredSize, u64 preferredOffset)
+	SharedMemoryAllocator::Slice SharedMemoryAllocator::AllocateNoLock(u64 desiredSize, u64 existingSliceCount)
 	{
 		// Find best fit.
 		// Will return first iterator not less than desiredSize
 		auto it = m_freeMemoryBySize.lower_bound({ desiredSize, 0 });
 		if (it == m_freeMemoryBySize.end())
 		{
-			if (!m_freeMemory.empty())
+			bool isEmpty = m_freeMemory.empty();
+
+			if (!isEmpty)
 			{
 				--it;
 				u64 size = it->size;
-				if (size >= desiredSize / 2) // Half the size is ok
+				if (size >= desiredSize / 4) // Quarter the size is ok
 				{
 					u64 offset = it->offset;
 					m_freeMemoryBySize.erase(it);
@@ -461,13 +664,29 @@ namespace uba
 				}
 			}
 
+			UBA_ASSERT(m_memory);
 			// No block is big enough, create new one
-			u64 toCommit = Min(AlignUp(desiredSize, m_commitStepSize), m_reserved);
-			UBA_SHIPPING_ASSERTF(m_committed + toCommit <= m_reserved, TC("Not enough space reserved for commit (Reserved: %llu, Committed: %llu Requested: %llu"), m_reserved, m_committed, toCommit);
-			void* res = VirtualAlloc(m_memory + m_committed, toCommit, MEM_COMMIT, PAGE_READWRITE);
-			UBA_SHIPPING_ASSERTF(res == m_memory + m_committed, TC("VirtualAlloc failed (%s)"), LastErrorToText().data);
-			u64 offset = m_committed;
+			u64 toCommit = Min(AlignUp(desiredSize, m_commitStepSize), m_capacity);
+			UBA_SHIPPING_ASSERTF(m_reserved + toCommit <= m_capacity, TC("Not enough space in capacity for commit (Capacity: %llu, Reserved: %llu Requested: %llu"), m_capacity, m_reserved, toCommit);
+			bool res = MapViewCommit(m_logger, m_memory + m_reserved, toCommit);
+			UBA_SHIPPING_ASSERTF(res, TC("VirtualAlloc failed (%s)"), LastErrorToText().data);
+			u64 offset = m_reserved;
+			m_reserved += toCommit;
 			m_committed += toCommit;
+
+			// Coalece with last block if it ends at offset
+			if (!isEmpty)
+			{
+				auto lastIt = m_freeMemory.rbegin();
+				if (lastIt->first + lastIt->second == offset)
+				{
+					offset = lastIt->first;
+					toCommit += lastIt->second;
+					m_freeMemoryBySize.erase({ lastIt->second, lastIt->first });
+					m_freeMemory.erase(lastIt->first);
+				}
+			}
+
 			u64 size = toCommit;
 			if (toCommit > desiredSize)
 			{
@@ -517,11 +736,15 @@ namespace uba
 
 		m_used -= slice.size;
 
+		// Need to measure if this is a good idea
+		// DiscardVirtualMemory(m_memory + slice.offset, slice.size);
+
 		if (it == m_freeMemory.end()) // The slice to add is after last
 		{
 			if (m_freeMemory.empty()) // No slices found at all
 			{
-				m_freeMemory.try_emplace(slice.offset, slice.size);
+				bool res = m_freeMemory.try_emplace(slice.offset, slice.size).second;
+				UBA_ASSERT(res);(void)res;
 				m_freeMemoryBySize.insert({ slice.size, slice.offset });
 				return;
 			}
@@ -538,7 +761,8 @@ namespace uba
 				}
 				else // This is a new slot
 				{
-					m_freeMemory.try_emplace(slice.offset, slice.size);
+					bool res = m_freeMemory.try_emplace(slice.offset, slice.size).second;
+					UBA_ASSERT(res);(void)res;
 					m_freeMemoryBySize.insert({ slice.size, slice.offset });
 					return;
 				}
@@ -565,7 +789,9 @@ namespace uba
 
 			if (!addedToPrev)
 			{
-				it = m_freeMemory.emplace(slice.offset, slice.size).first;
+				auto res = m_freeMemory.try_emplace(slice.offset, slice.size);
+				UBA_ASSERT(res.second);
+				it = res.first;
 			}
 
 			auto next = it; ++next;
@@ -590,6 +816,28 @@ namespace uba
 		StringBuffer<> str;
 		str.Appendf(TC("%s/%s (%llu)"), BytesToText(m_used).str, BytesToText(m_committed).str, m_handleCount.load());
 		trace.StatusUpdate(startRow, 7, str, LogEntryType_Info);
+
+		#if UBA_TRACK_SHARED_MEMORY_ALLOCATIONS
+		SCOPED_FUTEX(m_handleLock, lock);
+		Map<SizeAndOffset, const tchar*> handles;
+		for (InternalHandle* it = m_firstHandle; it; it = it->next)
+			if (u64 committed = it->committed.load())
+				handles.emplace(SizeAndOffset{ committed, u64(it) }, it->GetHint());
+		Vector<u32> items;
+		for (auto ri = handles.rbegin(), re = handles.rend(); ri !=re; ++ri)
+		{
+			const tchar* hint = ri->second;
+			if (const tchar* lastSlash = TStrrchr(hint, PathSeparator))
+				hint = lastSlash + 1;
+			str.Clear().Appendf(TC("%s  %s"), BytesToText(ri->first.size).str, hint);
+			u32 stringId = trace.StringAdd(str);
+			items.push_back(stringId);
+		}
+		lock.Leave();
+
+		u32 allocsStringId = trace.StringAdd(TCV("Allocs"));
+		trace.TableUpdate(allocsStringId, 25, 50, items.data(), u32(items.size()));
+		#endif
 	}
 
 	////////////////////////////////////////////////////////////////////////////////////////////////////
diff --git a/Engine/Source/Programs/UnrealBuildAccelerator/Common/Private/UbaStorage.cpp b/Engine/Source/Programs/UnrealBuildAccelerator/Common/Private/UbaStorage.cpp
--- a/Engine/Source/Programs/UnrealBuildAccelerator/Common/Private/UbaStorage.cpp
+++ b/Engine/Source/Programs/UnrealBuildAccelerator/Common/Private/UbaStorage.cpp
@@ -22,9 +22,10 @@ namespace uba
 
     CasKey EmptyFileKey = []() { CasKeyHasher hasher; return ToCasKey(hasher, false); }();
 
-	void Storage::GetMappingString(StringBufferBase& out, FileMappingHandle mappingHandle, u64 offset)
+	void Storage::GetMappingString(StringBufferBase& out, SharedMemoryHandle handle, u64 offset, bool canBeFreed)
 	{
-		out.Append('^').AppendBase62(mappingHandle.ToU64());
+		out.Append(canBeFreed ? FreeableMemoryHandleChar : MemoryHandleChar);
+		out.AppendBase62(handle.internalHandle);
 		if (offset)
 			out.Append('-').AppendBase62(offset);
 	}
@@ -47,9 +48,9 @@ namespace uba
 		table.GetValueAsString(rootDir, TC("RootDir"));
 		table.GetValueAsBool(asyncUnmapViewOfFile, TC("AsyncUnmapViewOfFile"));
 		table.GetValueAsU64(casCapacityBytes, TC("CasCapacityBytes"));
-		table.GetValueAsBool(keepTransientDataMapped, TC("KeepTransientDataMapped"));
 		table.GetValueAsBool(storeCompressed, TC("StoreCompressed"));
 		table.GetValueAsBool(allowDeferredDeletes, TC("AllowDeferredDeletes"));
+		table.GetValueAsBool(skipCleanupOnShutdown, TC("SkipCleanupOnShutdown"));
 		table.GetValueAsU32(fileEntryTableWipeThreshold, TC("FileEntryTableWipeThreshold"));
 		table.GetValueAsU32(maxParallelCopyOrLink, TC("MaxParallelCopyOrLink"));
 		table.GetValueAsU32(deleteCasBatchSize, TC("DeleteCasBatchSize"));
@@ -63,35 +64,42 @@ namespace uba
 			casCompressionLevel = GetCompressionLevel(compressionLevel);
 	}
 
-	struct FileMappingScope
+	struct MemoryMappingScope
 	{
-		u8* MapView(const tchar* func, FileMappingHandle mappingHandle, u64 mappingOffset, u64 fileSize, const tchar* fileName_)
+		u8* MapView(const tchar* func, SharedMemoryHandle memoryHandle, u64 memoryOffset, u64 fileSize, const tchar* fileName_)
 		{
+			handle = memoryHandle;
 			fileName = fileName_;
-			constexpr u64 pageSize = 64*1024;
-			u64 endOffset = mappingOffset + fileSize;
-			u64 alignedOffsetStart = AlignUp(mappingOffset - (pageSize - 1), pageSize);
-			u64 alignedOffsetEnd = AlignUp(endOffset, pageSize);
-			mapSize = alignedOffsetEnd - alignedOffsetStart;
-			mappedMem = MapViewOfFile(logger, mappingHandle, FILE_MAP_READ, alignedOffsetStart, mapSize);
-			if (!mappedMem)
-				mappedMem = MapViewOfFile(logger, mappingHandle, FILE_MAP_READ, alignedOffsetStart, 0);
+			u64 alignedOffsetStart = AlignUp(memoryOffset - (PageSize - 1), PageSize);
+			mapSize = (memoryOffset - alignedOffsetStart) + fileSize;
+
+			//u64 endOffset = memoryOffset + fileSize;
+			//u64 alignedOffsetStart = AlignUp(memoryOffset - (PageSize - 1), PageSize);
+			//u64 alignedOffsetEnd = AlignUp(endOffset, PageSize);
+			//mapSize = alignedOffsetEnd - alignedOffsetStart;
+			
+			allocator.MapView(view, handle, fileName, SharedMemoryMapType_ReadOnly, alignedOffsetStart, mapSize);
+			mappedMem = view.GetMemory();
 			if (mappedMem)
-				return mappedMem + (mappingOffset - alignedOffsetStart);
+				return mappedMem + (memoryOffset - alignedOffsetStart);
 			logger.Error(TC("[%s] Failed to map view of file %s at offset %llu and size %llu (%s)"), func, fileName, alignedOffsetStart, mapSize, LastErrorToText().data);
 			return nullptr;
 		}
 
-		FileMappingScope(Logger& l) : logger(l) {}
-		~FileMappingScope()
+		MemoryMappingScope(Logger& l, SharedMemoryAllocator& a) : logger(l), allocator(a) {}
+		~MemoryMappingScope()
 		{
-			if (mappedMem)
-				if (!UnmapViewOfFile(logger, mappedMem, mapSize, fileName))
-					logger.Error(TC("Failed to unmap memory %s at 0x%llx with size %llu (%s)"), fileName, mappedMem, mapSize, LastErrorToText().data);
+			if (!mappedMem)
+				return;
+			view.Reset();
 		}
 
 		Logger& logger;
+		SharedMemoryAllocator& allocator;
+		SharedMemoryView view;
+
 		const tchar* fileName;
+		SharedMemoryHandle handle;
 		u8* mappedMem = nullptr;
 		u64 mapSize = 0;
 	};
@@ -142,7 +150,7 @@ namespace uba
 			CasEntryReadLock casEntryLock(casEntry); // Note, this lock is taken again outside CasEntryAccessed.. so if this takes a long time it won't help to remove this lock
 			if (casEntry.dropped)
 				return;
-			//if (casEntry.mappingHandle.IsValid())
+			//if (casEntry.memoryHandle.IsValid())
 			//	return;
 		}
 
@@ -180,7 +188,7 @@ namespace uba
 
 		casEntry.size = size;
 
-		UBA_ASSERT(!m_casCapacityBytes || !casEntry.mappingHandle.IsValid());
+		UBA_ASSERT(!m_casCapacityBytes || !casEntry.memoryHandle.IsValid());
 
 		if (!m_casCapacityBytes || m_overflowReported || m_manuallyHandleOverflow)
 			return;
@@ -348,7 +356,7 @@ namespace uba
 
 	bool StorageImpl::WriteCompressedFile(WriteResult& out, const tchar* from, FileHandle readHandle, u8* readMem, u64 fileSize, const tchar* to, const void* header, u64 headerSize, u64 lastWriteTime)
 	{
-		CompressWriter w(m_logger, m_bufferSlots, m_workManager, Stats(), m_casCompressor, m_casCompressionLevel, m_asyncUnmapViewOfFile);
+		CompressWriter w(m_logger, m_sharedMemory, m_bufferSlots, m_workManager, Stats(), m_casCompressor, m_casCompressionLevel, m_asyncUnmapViewOfFile);
 		return w.CompressToFile(out.storedSize, from, readHandle, readMem, fileSize, to, header, headerSize, lastWriteTime, m_tempPath.data);
 	}
 
@@ -362,7 +370,7 @@ namespace uba
 			ExternalFileMapping externalMapping;
 			if (m_externalFileMappingsProvider(externalMapping, fileNameKey, fileName))
 			{
-				FileMappingScope scope(m_logger);
+				MemoryMappingScope scope(m_logger, m_sharedMemory);
 				u8* fileMem = scope.MapView(TC("WriteCompressed"), externalMapping.handle, externalMapping.offset, externalMapping.size, fileName);
 				if (!fileMem)
 					return false;
@@ -370,9 +378,9 @@ namespace uba
 				out.isIndependent = externalMapping.createIndependentMapping;
 				if (externalMapping.createIndependentMapping)
 				{
-					CompressWriter w(m_logger, m_bufferSlots, m_workManager, Stats(), m_casCompressor, m_casCompressionLevel, m_asyncUnmapViewOfFile);
-					out.mappingOffset = 0;
-					return w.CompressToMapping(out.mappingHandle, out.storedSize, fileMem, externalMapping.size, casFile);
+					CompressWriter w(m_logger, m_sharedMemory, m_bufferSlots, m_workManager, Stats(), m_casCompressor, m_casCompressionLevel, m_asyncUnmapViewOfFile);
+					out.memoryOffset = 0;
+					return w.CompressToMapping(out.memoryHandle, out.storedSize, fileMem, externalMapping.size, casFile);
 				}
 				else
 					return WriteCompressedFile(out, fileName, InvalidFileHandle, fileMem, externalMapping.size, casFile, nullptr, 0);
@@ -892,9 +900,10 @@ namespace uba
 	StorageImpl::StorageImpl(const StorageCreateInfo& info, const tchar* logPrefix)
 	:	m_logger(info.writer, logPrefix)
 	,	m_workManager(info.workManager)
+	,	m_sharedMemory(m_logger)
 	,	m_bufferSlots(BufferSlotSize)
 	,	m_maxParallelCopyOrLinkBottleneck(info.maxParallelCopyOrLink)
-	,	m_casDataBuffer(m_logger, &info.workManager)
+	,	m_casDataBuffer(m_logger, m_sharedMemory)
 	{
 		m_casCapacityBytes = info.casCapacityBytes;
 		m_fileEntryTableWipeThreshold = info.fileEntryTableWipeThreshold;
@@ -904,6 +913,7 @@ namespace uba
 		m_asyncUnmapViewOfFile = info.asyncUnmapViewOfFile;
 		m_allowDeleteVerified = info.allowDeleteVerified;
 		m_allowDeferredDeletes = info.allowDeferredDeletes;
+		m_skipCleanupOnShutdown = info.skipCleanupOnShutdown;
 		m_writeToDisk = info.writeToDisk;
 		if (!m_writeToDisk)
 			m_casCapacityBytes = 0;
@@ -925,16 +935,22 @@ namespace uba
 
 		m_rootDir.Append(TCV("cas")).EnsureEndsWithSlash();
 
-		m_casDataBuffer.AddTransient(TC("CasData"), info.keepTransientDataMapped);
+		m_casDataBuffer.Init(TC("CasData"), 32ull*1024*1024*1024, 1ull*1024*1024);
 
 		m_casCompressor = info.casCompressor;
 		m_casCompressionLevel = info.casCompressionLevel;
+
+		u64 reservedSpace = 1024ull * 1024 * 1024 * 1024;
+		u64 commitStepSize = 64ull * 1024 * 1024;
+		m_sharedMemory.Init(reservedSpace, commitStepSize);
 	}
 
 	StorageImpl::~StorageImpl()
 	{
 		SaveCasTable(true, true);
 		CloseMutex(m_exclusiveMutex);
+
+		m_sharedMemory.Deinit(m_skipCleanupOnShutdown);
 	}
 
 	MutexHandle StorageImpl::GetExclusiveAccess(Logger& logger, const StringView& rootDir, bool reportError)
@@ -1298,11 +1314,10 @@ namespace uba
 						continue;
 					if (DeleteFileW(casFileName.data))
 						++deleteCount;
-					if (casEntry.mappingIsIndependent) // Log files etc could end up here
+					if (m_casDataBuffer.IsIndependent(casEntry.memoryHandle)) // Log files etc could end up here
 					{
-						CloseFileMapping(m_logger, casEntry.mappingHandle, TC("SaveCasTable"));
-						casEntry.mappingHandle = {};
-						casEntry.mappingIsIndependent = false;
+						m_casDataBuffer.CloseIndependent(casEntry.memoryHandle, casEntry.memorySize, TC("SaveCasTable"));
+						casEntry.memoryHandle = {};
 					}
 					continue;
 				}
@@ -1717,7 +1732,7 @@ namespace uba
 		{
 			if (externalMapping.handle.IsValid())
 			{
-				FileMappingScope scope(m_logger);
+				MemoryMappingScope scope(m_logger, m_sharedMemory);
 				u8* fileMem = scope.MapView(TC("StoreCasFile"), externalMapping.handle, externalMapping.offset, externalMapping.size, fileName);
 				if (!fileMem)
 					return false;
@@ -1897,7 +1912,7 @@ namespace uba
 		return true;
 	}
 
-	bool StorageImpl::StoreCasFileClient(CasKey& out, StringKey fileNameKey, const tchar* fileName, FileMappingHandle mappingHandle, u64 mappingOffset, u64 fileSize, const tchar* hint, bool keepMappingInMemory, bool storeCompressed)
+	bool StorageImpl::StoreCasFileClient(CasKey& out, StringKey fileNameKey, const tchar* fileName, SharedMemoryHandle memoryHandle, u64 memoryOffset, u64 fileSize, const tchar* hint, bool keepMappingInMemory, bool storeCompressed)
 	{
 		UBA_ASSERT(false);
 		return false;
@@ -1957,12 +1972,11 @@ namespace uba
 			m_logger.Warning(TC("Deferred cas %s was deferred with source file size %llu but was read with source file size %llu (%s)"), CasKeyString(casKey).str, deferredFileSize, res.sourceStoredSize, deferredCreation.data);
 
 		bool hasMapping = false;
-		if (res.mappingHandle.IsValid())
+		if (res.memoryHandle.IsValid())
 		{
-			casEntry.mappingHandle = res.mappingHandle;
-			casEntry.mappingOffset = res.mappingOffset;
-			casEntry.mappingSize = res.storedSize;
-			casEntry.mappingIsIndependent = res.isIndependent;
+			casEntry.memoryHandle = res.memoryHandle;
+			casEntry.memoryOffset = res.memoryOffset;
+			casEntry.memorySize = res.storedSize;
 			hasMapping = true;
 		}
 
@@ -2038,7 +2052,7 @@ namespace uba
 		return true;
 	}
 
-	MappedView StorageImpl::MapView(const CasKey& casKey, const tchar* hint)
+	MappedView2 StorageImpl::MapView(const CasKey& casKey, const tchar* hint)
 	{
 		SCOPED_READ_LOCK(m_casLookupLock, lookupLock);
 		auto findIt = m_casLookup.find(casKey);
@@ -2046,7 +2060,7 @@ namespace uba
 		if (!foundEntry)
 		{
 			m_logger.Error(TC("Can't find %s inside cas database (%s)"), CasKeyString(casKey).str, hint);
-			return {{},0,0 };
+			return {};
 		}
 		CasEntry& casEntry = findIt->second;
 		lookupLock.Leave();
@@ -2056,12 +2070,12 @@ namespace uba
 		//	m_logger.Error(TC("Trying to use unverified mapping of %s (%s)"), CasKeyString(casKey).str, hint);
 		//	return {{},0,0 };
 		//}
-		FileMappingHandle handle = casEntry.mappingHandle;
-		u64 offset = casEntry.mappingOffset;
-		u64 size = casEntry.mappingSize;
+		SharedMemoryHandle memoryHandle = casEntry.memoryHandle;
+		u64 offset = casEntry.memoryOffset;
+		u64 size = casEntry.memorySize;
 		casEntryLock.Leave();
 
-		auto res = m_casDataBuffer.MapView(handle, offset, size, hint);
+		MappedView2 res = m_casDataBuffer.MapView(memoryHandle, offset, size, hint);
 		if (!res.memory)
 			m_logger.Error(TC("Failed to map view for %s (%s)"), CasKeyString(casKey).str, hint);
 		return res;
@@ -2072,6 +2086,11 @@ namespace uba
 		m_casDataBuffer.UnmapView(view, hint);
 	}
 
+	SharedMemoryAllocator& StorageImpl::GetSharedMemory()
+	{
+		return m_sharedMemory;
+	}
+
 	bool StorageImpl::DropCasFile(const CasKey& casKey, bool forceDelete, const tchar* hint, bool reduceRefCount)
 	{
 		SCOPED_READ_LOCK(m_casLookupLock, lookupLock);
@@ -2109,13 +2128,12 @@ namespace uba
 			UBA_ASSERTF(casEntry.refCount == 0, TC("Force deleting cas %s with refcount %u (%s)"), casFile.data, casEntry.refCount, hint);
 			u64 sizeDeleted = 0;
 
-			if (casEntry.mappingHandle.IsValid())
+			if (casEntry.memoryHandle.IsValid())
 			{
-				if (casEntry.mappingIsIndependent)
+				if (m_casDataBuffer.IsIndependent(casEntry.memoryHandle))
 				{
-					CloseFileMapping(m_logger, casEntry.mappingHandle, TC("DropCasFile"));
-					casEntry.mappingHandle = {};
-					casEntry.mappingIsIndependent = false;
+					m_casDataBuffer.CloseIndependent(casEntry.memoryHandle, casEntry.memorySize, TC("DropCasFile"));
+					casEntry.memoryHandle = {};
 				}
 			}
 			else
@@ -2242,10 +2260,10 @@ namespace uba
 				StringBuffer<512> casFile;
 				u64 decompressedSize;
 
-				if (casEntry.mappingHandle.IsValid())
+				if (casEntry.memoryHandle.IsValid())
 				{
 					casFile.Append(CasKeyString(actualKey).str);
-					mappedView = m_casDataBuffer.MapView(casEntry.mappingHandle, casEntry.mappingOffset, casEntry.mappingSize, casFile.data, casEntry.mappingIsIndependent);
+					mappedView = m_casDataBuffer.MapView(casEntry.memoryHandle, casEntry.memoryOffset, casEntry.memorySize, casFile.data);
 					compressedData = mappedView.memory;
 					if (!compressedData)
 						return m_logger.Error(TC("[CopyOrLink] Failed to map view of mapping %s (%s)"), casFile.data, LastErrorToText().data);
@@ -2265,7 +2283,7 @@ namespace uba
 						return m_logger.Error(TC("[CopyOrLink] Failed to read first 8 bytes from compressed file %s (%s)"), casFile.data, LastErrorToText().data);
 				}
 
-				bool writeDirectlyToFile = casEntry.mappingHandle.IsValid() && false; // Experiment to try to fix bottlenecks on cloud
+				bool writeDirectlyToFile = casEntry.memoryHandle.IsValid() && false; // Experiment to try to fix bottlenecks on cloud
 				bool useNoBuffering = writeDirectlyToFile && false; // Experiment to try to fix bottlenecks on cloud
 
 				u32 writeFlags = useNoBuffering ? FILE_FLAG_NO_BUFFERING : 0;
@@ -2367,12 +2385,12 @@ namespace uba
 				return done(destinationFile, sizeOnDisk);
 			}
 
-			if (casEntry.mappingHandle.IsValid())
+			if (casEntry.memoryHandle.IsValid())
 			{
 				SCOPED_FUTEX(fileEntry.lock, entryLock);
 				fileEntry.verified = false;
 
-				MappedView mappedView = m_casDataBuffer.MapView(casEntry.mappingHandle, casEntry.mappingOffset, casEntry.mappingSize, TC(""), casEntry.mappingIsIndependent);
+				MappedView mappedView = m_casDataBuffer.MapView(casEntry.memoryHandle, casEntry.memoryOffset, casEntry.memorySize, TC(""));
 				auto mapViewGuard = MakeGuard([&](){ m_casDataBuffer.UnmapView(mappedView, destination); });
 
 				FileAccessor destinationFile(m_logger, destination);
@@ -2521,6 +2539,11 @@ namespace uba
 		m_stats.Add(stats);
 	}
 
+	void StorageImpl::TraceUpdate(Trace& trace, u32 startRow)
+	{
+		m_casDataBuffer.TraceStatus(trace, startRow, TC("Cas"), 1);
+	}
+
 	void StorageImpl::PrintSummary(Logger& logger)
 	{
 		logger.Info(TC("  ------- Storage stats summary -------"));
@@ -2535,10 +2558,7 @@ namespace uba
 		logger.Info(TC("  FileTable            %6u"), u32(m_fileTableLookup.size()));
 
 		StorageStats& stats = Stats();
-		u64 casBufferSize;
-		u32 casBufferCount;
-		m_casDataBuffer.GetSizeAndCount(MappedView_Transient, casBufferSize, casBufferCount);
-		logger.Info(TC("  CasDataBuffers       %6u %9s"), casBufferCount, BytesToText(casBufferSize).str);
+		logger.Info(TC("  CasDataBuffers       %6u %9s"), 1, BytesToText(m_casDataBuffer.GetUsed()).str);
 		logger.Info(TC("  CasTable             %6u %9s"), u32(m_casLookup.size()), BytesToText(m_casTotalBytes).str);
 		logger.Info(TC("     Dropped           %6u %9s"), m_casDroppedCount, BytesToText(m_casDroppedBytes).str);
 		logger.Info(TC("     Evicted           %6u %9s"), m_casEvictedCount, BytesToText(m_casEvictedBytes).str);
@@ -2550,7 +2570,7 @@ namespace uba
 		logger.Info(TC(""));
 	}
 
-	CasKey StorageImpl::CalculateCasKey(u8* fileMem, u64 fileSize, bool storeCompressed)
+	CasKey StorageImpl::CalculateCasKey(const u8* fileMem, u64 fileSize, bool storeCompressed)
 	{
 		StorageStats& stats = Stats();
 		TimerScope ts(stats.calculateCasKey);
@@ -2634,15 +2654,15 @@ namespace uba
 			return m_logger.Error(TC("GetCasEntry - Couldn't find cas entry for %s"), hint);
 		CasEntry& casEntry = findIt->second;
 		CasEntryReadLock casEntryLock(casEntry);
-		if (!casEntry.mappingHandle.IsValid())
+		if (!casEntry.memoryHandle.IsValid())
 			return m_logger.Error(TC("Code path only supports mapped cas entries"));
 		if (!IsCompressed(casEntry.key))
 		{
-			outSize = casEntry.mappingSize;
+			outSize = casEntry.memorySize;
 			return true;
 		}
-		FileMappingScope scope(m_logger);
-		u8* fileMem = scope.MapView(TC("ReadFileSize"), casEntry.mappingHandle, casEntry.mappingOffset, 8, hint);
+		MemoryMappingScope scope(m_logger, m_sharedMemory);
+		u8* fileMem = scope.MapView(TC("ReadFileSize"), casEntry.memoryHandle, casEntry.memoryOffset, 8, hint);
 		if (!fileMem)
 			return false;
 		outSize = *(u64*)fileMem;
@@ -2657,17 +2677,18 @@ namespace uba
 			return m_logger.Error(TC("GetCasEntry - Couldn't find cas entry for %s"), hint);
 		CasEntry& casEntry = findIt->second;
 		CasEntryReadLock casEntryLock(casEntry);
-		FileMappingScope scope(m_logger);
-		u8* fileMem = scope.MapView(TC("ReadFileSize"), casEntry.mappingHandle, casEntry.mappingOffset, casEntry.mappingSize, hint);
+
+		MemoryMappingScope scope(m_logger, m_sharedMemory);
+		u8* fileMem = scope.MapView(TC("ReadFileSize"), casEntry.memoryHandle, casEntry.memoryOffset, casEntry.memorySize, hint);
 		if (!fileMem)
 			return false;
 		if (!IsCompressed(casEntry.key))
 		{
-			memcpy(outData, fileMem, casEntry.mappingSize);
+			memcpy(outData, fileMem, casEntry.memorySize);
 			return true;
 		}
 		u64 decompressedSize = *(u64*)fileMem;
 		fileMem += sizeof(u64);
-		return DecompressMemoryToMemory(fileMem, casEntry.mappingSize, (u8*)outData, decompressedSize, CasKeyString(casEntry.key).str, hint);
+		return DecompressMemoryToMemory(fileMem, casEntry.memorySize, (u8*)outData, decompressedSize, CasKeyString(casEntry.key).str, hint);
 	}
 }
diff --git a/Engine/Source/Programs/UnrealBuildAccelerator/Common/Private/UbaStorageClient.cpp b/Engine/Source/Programs/UnrealBuildAccelerator/Common/Private/UbaStorageClient.cpp
--- a/Engine/Source/Programs/UnrealBuildAccelerator/Common/Private/UbaStorageClient.cpp
+++ b/Engine/Source/Programs/UnrealBuildAccelerator/Common/Private/UbaStorageClient.cpp
@@ -120,7 +120,7 @@ namespace uba
 	{
 		delete m_proxyClient;
 		for (auto& pair : m_localStorageFiles)
-			CloseFileMapping(m_logger, pair.second.casEntry.mappingHandle, pair.second.fileName.c_str());
+			m_sharedMemory.CloseHandle(m_logger, pair.second.casEntry.memoryHandle, pair.second.fileName.c_str());
 	}
 
 	bool StorageClient::IsUsingProxy()
@@ -165,9 +165,9 @@ namespace uba
 		if (findIt != m_localStorageFiles.end())
 		{
 			const LocalFile& localFile = findIt->second;
-			if (localFile.casEntry.mappingHandle.IsValid())
+			if (localFile.casEntry.memoryHandle.IsValid())
 			{
-				GetMappingString(out, localFile.casEntry.mappingHandle, 0);
+				GetMappingString(out, localFile.casEntry.memoryHandle, 0, false);
 				return true;
 			}
 
@@ -182,7 +182,7 @@ namespace uba
 		return StorageImpl::GetCasFileName(out, casKey);
 	}
 
-	MappedView StorageClient::MapView(const CasKey& casKey, const tchar* hint)
+	MappedView2 StorageClient::MapView(const CasKey& casKey, const tchar* hint)
 	{
 		SCOPED_READ_LOCK(m_localStorageFilesLock, tempLock);
 		auto findIt = m_localStorageFiles.find(AsCompressed(casKey, false));
@@ -192,11 +192,11 @@ namespace uba
 			return StorageImpl::MapView(casKey, hint);
 
 		LocalFile& file = findIt->second;
-		if (!file.casEntry.mappingHandle.IsValid())
+		if (!file.casEntry.memoryHandle.IsValid())
 			return StorageImpl::MapView(casKey, hint);
 
-		MappedView view;
-		view.handle = file.casEntry.mappingHandle;
+		MappedView2 view;
+		view.handle = file.casEntry.memoryHandle;
 		view.size = file.casEntry.size;
 		view.offset = 0;
 		view.isCompressed = false;
@@ -215,7 +215,6 @@ namespace uba
 	{
 		CasKey casKey = casKeyTmp;
 
-		FileMappingType mappingType = MappedView_Transient;
 		bool shouldStore = mappingBuffer == nullptr;
 		UBA_ASSERT(AsCompressed(casKey, false) != CasKeyZero);
 
@@ -238,10 +237,10 @@ namespace uba
 			if (lf.casEntry.exists)
 			{
 				out.casKey = findIt->first;
-				if (lf.casEntry.mappingHandle.IsValid())
+				if (lf.casEntry.memoryHandle.IsValid())
 				{
 					out.size = lf.casEntry.size;
-					out.view.handle = lf.casEntry.mappingHandle;
+					out.view.handle = lf.casEntry.memoryHandle;
 					out.view.size = lf.casEntry.size;
 					out.view.isCompressed = false;
 				}
@@ -575,7 +574,7 @@ namespace uba
 						UBA_ASSERT(!writeMem || mappedView.size == sizeOnDisk);
 						if (!writeMem)
 						{
-							mappedView = mappingBuffer->AllocAndMapView(mappingType, sizeOnDisk, memoryMapAlignment, hint);
+							mappedView = mappingBuffer->AllocAndMapView(sizeOnDisk, memoryMapAlignment, hint);
 							writeMem = mappedView.memory;
 							if (!writeMem)
 								return false;
@@ -863,9 +862,9 @@ namespace uba
 
 		if (shouldStore)
 		{
-			casEntry->mappingHandle = mappedView.handle;
-			casEntry->mappingOffset = mappedView.offset;
-			casEntry->mappingSize = fileSize;
+			casEntry->memoryHandle = mappedView.handle;
+			casEntry->memoryOffset = mappedView.offset;
+			casEntry->memorySize = fileSize;
 
 			casEntry->exists = true;
 			casEntryLock.Execute();
@@ -908,7 +907,7 @@ namespace uba
 		return StorageImpl::HasCasFile(casKey, out);
 	}
 
-	bool StorageClient::StoreCasFileClient(CasKey& out, StringKey fileNameKey, const tchar* fileName, FileMappingHandle mappingHandle, u64 mappingOffset, u64 fileSize, const tchar* hint, bool keepMappingInMemory, bool storeCompressed)
+	bool StorageClient::StoreCasFileClient(CasKey& out, StringKey fileNameKey, const tchar* fileName, SharedMemoryHandle memoryHandle, u64 memoryOffset, u64 fileSize, const tchar* hint, bool keepMappingInMemory, bool storeCompressed)
 	{
 		NetworkClient& client = m_client; // Don't use proxy
 
@@ -918,7 +917,7 @@ namespace uba
 		u8* fileMem = nullptr;
 
 		FileAccessor source(m_logger, fileName);
-		if (!mappingHandle.IsValid())
+		if (!memoryHandle.IsValid())
 		{
 			if (!source.OpenMemoryRead())
 				return false;
@@ -927,13 +926,18 @@ namespace uba
 		}
 		else
 		{
-			fileMem = MapViewOfFile(m_logger, mappingHandle, FILE_MAP_READ, mappingOffset, fileSize);
+			UBA_ASSERT(!memoryOffset);
+			fileMem = m_sharedMemory.MapView(memoryHandle, hint);
 			if (!fileMem)
 				return m_logger.Error(TC("%s - MapViewOfFile failed (%s)"), fileName, LastErrorToText().data);
 			isPersistentMapping = true;
 		}
 
-		auto unmapGuard = MakeGuard([&](){ if (isPersistentMapping) UnmapViewOfFile(m_logger, fileMem, fileSize, fileName); });
+		auto unmapGuard = MakeGuard([&]()
+			{
+				if (isPersistentMapping)
+					m_sharedMemory.UnmapView(memoryHandle, fileMem);
+			});
 
 		CasKey casKey = CalculateCasKey(fileMem, fileSize, storeCompressed);
 		if (casKey == CasKeyZero)
@@ -942,20 +946,15 @@ namespace uba
 		SCOPED_WRITE_LOCK(m_localStorageFilesLock, lock);
 		auto insres = m_localStorageFiles.try_emplace(AsCompressed(casKey, false));
 		LocalFile& localFile = insres.first->second;
-		if (keepMappingInMemory && isPersistentMapping && !localFile.casEntry.mappingHandle.IsValid())
+
+		if (keepMappingInMemory && isPersistentMapping && !localFile.casEntry.memoryHandle.IsValid())
 		{
-			FileMappingHandle mappingHandle2;
-			if (DuplicateFileMapping(m_logger, GetCurrentProcessHandle(), mappingHandle, GetCurrentProcessHandle(), mappingHandle2, FILE_MAP_READ, false, 0, fileName))
-			{
-				localFile.casEntry.mappingHandle = mappingHandle2;
-				localFile.casEntry.size = fileSize;
-				localFile.casEntry.exists = true;
-			}
-			else
-				m_logger.Warning(TC("Failed to duplicate handle for file mapping %s (%s)"), fileName, LastErrorToText().data);
+			localFile.casEntry.memoryHandle = m_sharedMemory.DuplicateHandle(memoryHandle, TC(""));
+			localFile.casEntry.size = fileSize;
+			localFile.casEntry.exists = true;
 		}
 
-		if (!isPersistentMapping && !localFile.casEntry.mappingHandle.IsValid() && localFile.fileName.empty())
+		if (!isPersistentMapping && !localFile.casEntry.memoryHandle.IsValid() && localFile.fileName.empty())
 		{
 			localFile.casEntry.size = fileSize;
 			localFile.casEntry.verified = true;
diff --git a/Engine/Source/Programs/UnrealBuildAccelerator/Common/Private/UbaStorageServer.cpp b/Engine/Source/Programs/UnrealBuildAccelerator/Common/Private/UbaStorageServer.cpp
--- a/Engine/Source/Programs/UnrealBuildAccelerator/Common/Private/UbaStorageServer.cpp
+++ b/Engine/Source/Programs/UnrealBuildAccelerator/Common/Private/UbaStorageServer.cpp
@@ -153,25 +153,24 @@ namespace uba
 
 	void StorageServer::ActiveFetch::Release(StorageServer& server, const tchar* reason)
 	{
-		if (mappedView.handle.IsValid())
+		if (ownsMapping)
 		{
-			if (ownsMapping)
-			{
-				++server.m_activeUnmap;
-				server.GetServer().AddWork([&server, mb = memoryBegin, mp = mappedView, rfh = readFileHandle, reason](const WorkContext& context)
-					{
-						if (!UnmapViewOfFile(server.m_logger, mb, mp.size, reason))
-							server.m_logger.Error(TC("%s - Failed to unmap memory at 0x%llx with size %llu (%s)"), reason, mb, mp.size, LastErrorToText().data);
-						if (!CloseFileMapping(server.m_logger, mp.handle, reason))
-							server.m_logger.Error(TC("%s - Failed to close file mapping %llu (%s)"), reason, mp.handle.ToU64(), LastErrorToText().data);
-						if (!CloseFile(nullptr, rfh))
-							server.m_logger.Error(TC("%s - Failed to close file (%s)"), reason, LastErrorToText().data);
-
-						--server.m_activeUnmap;
-					}, 1, TC("ActiveFetchRelease"));
-			}
-			else
-				server.m_casDataBuffer.UnmapView(mappedView, reason);
+			++server.m_activeUnmap;
+			server.GetServer().AddWork([&server, mb = memoryBegin, mps = mappedView.size, fmp = fileMappingHandle, rfh = readFileHandle, reason](const WorkContext& context)
+				{
+					if (!UnmapViewOfFile(server.m_logger, mb, mps, reason))
+						server.m_logger.Error(TC("%s - Failed to unmap memory at 0x%llx with size %llu (%s)"), reason, mb, mps, LastErrorToText().data);
+					if (!CloseFileMapping(server.m_logger, fmp, reason))
+						server.m_logger.Error(TC("%s - Failed to close file mapping %llu (%s)"), reason, fmp.ToU64(), LastErrorToText().data);
+					if (!CloseFile(nullptr, rfh))
+						server.m_logger.Error(TC("%s - Failed to close file (%s)"), reason, LastErrorToText().data);
+
+					--server.m_activeUnmap;
+				}, 1, TC("ActiveFetchRelease"));
+		}
+		else if (mappedView.handle.IsValid())
+		{
+			server.m_casDataBuffer.UnmapView(mappedView, reason);
 		}
 		else
 		{
@@ -235,8 +234,8 @@ namespace uba
 				}
 
 				m_casDataBuffer.UnmapView(store.mappedView, TC("OnDisconnected"));
-				if (store.mappedView.isIndependent)
-					CloseFileMapping(m_logger, store.mappedView.handle, TC("StoreDone"));
+				if (m_casDataBuffer.IsIndependent(store.mappedView.handle))
+					m_casDataBuffer.CloseIndependent(store.mappedView.handle, store.mappedView.size, TC("StoreDone"));
 
 				CasEntryWriteLock casEntryLock(*store.casEntry);
 				if (!store.casEntry->beingWritten)
@@ -303,18 +302,17 @@ namespace uba
 			return;
 		if (!casEntry.dropped || !m_dropIndependentMappingsOnCopy)
 			return;
-		if (!casEntry.mappingIsIndependent || !casEntry.mappingHandle.IsValid())
+		if (!casEntry.memoryHandle.IsValid() || !m_casDataBuffer.IsIndependent(casEntry.memoryHandle))
 			return;
-		FileMappingHandle mapping = casEntry.mappingHandle;
-		casEntry.mappingHandle = {};
-		casEntry.mappingIsIndependent = false;
+		SharedMemoryHandle memoryHandle = casEntry.memoryHandle;
+		casEntry.memoryHandle = {};
 		casEntry.verified = true;
 		casEntry.exists = false;
 		casEntry.dropped = false;
 
-		m_server.AddWork([this, mapping](const WorkContext& context)
+		m_server.AddWork([this, memoryHandle, size = casEntry.memorySize](const WorkContext& context)
 			{
-				CloseFileMapping(m_logger, mapping, TC("ReleaseCasEntryNoLock"));
+				m_casDataBuffer.CloseIndependent(memoryHandle, size, TC("ReleaseCasEntryNoLock"));
 			}, 1, TC("ReleaseCasEntryNoLock"));
 	}
 
@@ -638,7 +636,7 @@ namespace uba
 							m_logger.Info(TC("Server did not find cas %s for %s in file table lookup. Recalculating cas key"), CasKeyString(casKey).str, hint.data);
 							if (!CalculateCasKey(checkedCasKey, hint.data))
 							{
-								m_logger.Error(TC("FetchBegin failed for cas file %s (%s) requested by client &u. Can't calculate cas key for file"), CasKeyString(casKey).str, hint.data, connectionInfo.GetId());
+								m_logger.Error(TC("FetchBegin failed for cas file %s (%s) requested by client %u. Can't calculate cas key for file"), CasKeyString(casKey).str, hint.data, connectionInfo.GetId());
 								writer.WriteU16(0);
 								return false;
 							}
@@ -707,6 +705,7 @@ namespace uba
 				const u8* memoryPos = nullptr;
 				bool ownsMapping = false;
 
+				FileMappingHandle fileMappingHandle;
 				MappedView mappedView;
 				auto mvg = MakeGuard([&]()
 					{
@@ -716,9 +715,9 @@ namespace uba
 							if (mappedView.memory)
 								if (!UnmapViewOfFile(m_logger, mappedView.memory, fileSize, TC("FetchBegin")))
 									m_logger.Error(TC("Failed to unmap memory %s at 0x%llx with size %llu (%s)"), hint.data, mappedView.memory, fileSize, LastErrorToText().data);
-							if (mappedView.handle.IsValid())
-								if (!CloseFileMapping(m_logger, mappedView.handle, hint.data))
-									m_logger.Error(TC("Failed to close file mapping for %s with handle %llu (%s)"), hint.data, mappedView.handle.ToU64(), LastErrorToText().data);
+							if (fileMappingHandle.IsValid())
+								if (!CloseFileMapping(m_logger, fileMappingHandle, hint.data))
+									m_logger.Error(TC("Failed to close file mapping for %s with handle %llu (%s)"), hint.data, fileMappingHandle.ToU64(), LastErrorToText().data);
 						}
 						else
 						{
@@ -729,12 +728,12 @@ namespace uba
 
 				todoRemoveMe(6);
 
-				bool useFileMapping = casEntry.mappingHandle.IsValid();
+				bool useFileMapping = casEntry.memoryHandle.IsValid();
 				if (useFileMapping)
 				{
-					mappedView = m_casDataBuffer.MapView(casEntry.mappingHandle, casEntry.mappingOffset, casEntry.mappingSize, CasKeyString(casKey).str, casEntry.mappingIsIndependent);
+					mappedView = m_casDataBuffer.MapView(casEntry.memoryHandle, casEntry.memoryOffset, casEntry.memorySize, CasKeyString(casKey).str);
 					memoryBegin = mappedView.memory;
-					fileSize = casEntry.mappingSize;
+					fileSize = casEntry.memorySize;
 					if (!memoryBegin)
 						return m_logger.Error(TC("Failed to map memory map for %s. Will use file handle instead (%s)"), CasKeyString(casKey).str, LastErrorToText().data);
 					memoryPos = memoryBegin;
@@ -754,11 +753,11 @@ namespace uba
 					if (fileSize > BufferSlotSize)
 					{
 						ownsMapping = true;
-						mappedView.handle = CreateFileMappingW(m_logger, readFileHandle, PAGE_READONLY, fileSize, casFile.data);
-						if (!mappedView.handle.IsValid())
+						fileMappingHandle = CreateFileMappingW(m_logger, readFileHandle, PAGE_READONLY, fileSize, casFile.data);
+						if (!fileMappingHandle.IsValid())
 							return m_logger.Error(TC("Failed to create file mapping of %s (%s)"), casFile.data, LastErrorToText().data);
 						u64 offset = memoryPos - memoryBegin;
-						mappedView.memory = MapViewOfFile(m_logger, mappedView.handle, FILE_MAP_READ, 0, fileSize);
+						mappedView.memory = MapViewOfFile(m_logger, fileMappingHandle, FILE_MAP_READ, 0, fileSize);
 						if (!mappedView.memory)
 							return m_logger.Error(TC("Failed to map memory of %s with size %llu (%s)"), casFile.data, fileSize, LastErrorToText().data);
 						memoryBegin = mappedView.memory;
@@ -799,26 +798,26 @@ namespace uba
 				{
 					if (!ReadFile(m_logger, casFile.data, readFileHandle, writeBuffer, toWrite))
 					{
-						UBA_ASSERT(false); // Implement
-						return m_logger.Error(TC("Failed to read file %s (%s) (1)"), casFile.data, LastErrorToText().data);;
+						UBA_ASSERTF(false, TC("ReadFile failed for file %s trying to read %u bytes (%s)"), casFile.data, toWrite, LastErrorToText().data); // Implement
+						return m_logger.Error(TC("Failed to read file %s (%s) (1)"), casFile.data, LastErrorToText().data);
 					}
 				}
 				else
 				{
-					u8* buffer = m_bufferSlots.Pop();;
+					u8* buffer = m_bufferSlots.Pop();
 					memoryBegin = buffer;
 					memoryPos = memoryBegin;
 					u32 toRead = u32(Min(left, BufferSlotSize));
  					if (!ReadFile(m_logger, casFile.data, readFileHandle, buffer, toRead))
 					{
 						UBA_ASSERT(false); // Implement
-						return m_logger.Error(TC("Failed to read file %s (%s) (2)"), casFile.data, LastErrorToText().data);;
+						return m_logger.Error(TC("Failed to read file %s (%s) (2)"), casFile.data, LastErrorToText().data);
 					}
 					memcpy(writeBuffer, memoryPos, toWrite);
 					memoryPos += toWrite;
 
 					if (!CloseFile(casFile.data, readFileHandle))
-						m_logger.Error(TC("Failed to close file %s (%s)"), casFile.data, LastErrorToText().data);;
+						m_logger.Error(TC("Failed to close file %s (%s)"), casFile.data, LastErrorToText().data);
 					readFileHandle = InvalidFileHandle;
 				}
 
@@ -864,6 +863,7 @@ namespace uba
 				mappedView.size = fileSize;
 
 				fetch.readFileHandle = readFileHandle;
+				fetch.fileMappingHandle = fileMappingHandle;
 				fetch.mappedView = mappedView;
 				fetch.ownsMapping = ownsMapping;
 				fetch.memoryBegin = memoryBegin;
@@ -1077,7 +1077,7 @@ namespace uba
 				}
 				else
 				{
-					mappedView = m_casDataBuffer.AllocAndMapView(MappedView_Transient, fileSize, 1, CasKeyString(casKey).str, false, m_createIndependentMappings);
+					mappedView = m_casDataBuffer.AllocAndMapView(fileSize, 1, CasKeyString(casKey).str, m_createIndependentMappings);
 					if (!mappedView.memory)
 					{
 						casEntry.verified = false;
@@ -1159,10 +1159,9 @@ namespace uba
 					CasEntry& casEntry = *activeStore.casEntry;
 					{
 						CasEntryWriteLock casEntryLock(casEntry);
-						casEntry.mappingHandle = activeStore.mappedView.handle;
-						casEntry.mappingOffset = activeStore.mappedView.offset;
-						casEntry.mappingSize = totalWritten;
-						casEntry.mappingIsIndependent = activeStore.mappedView.isIndependent;
+						casEntry.memoryHandle = activeStore.mappedView.handle;
+						casEntry.memoryOffset = activeStore.mappedView.offset;
+						casEntry.memorySize = totalWritten;
 						casEntry.exists = true;
 						casEntry.beingWritten = false;
 					}
diff --git a/Engine/Source/Programs/UnrealBuildAccelerator/Common/Private/UbaStorageUtils.cpp b/Engine/Source/Programs/UnrealBuildAccelerator/Common/Private/UbaStorageUtils.cpp
--- a/Engine/Source/Programs/UnrealBuildAccelerator/Common/Private/UbaStorageUtils.cpp
+++ b/Engine/Source/Programs/UnrealBuildAccelerator/Common/Private/UbaStorageUtils.cpp
@@ -51,7 +51,7 @@ namespace uba
 	}
 
 
-	CasKey CalculateCasKey(u8* fileMem, u64 fileSize, bool storeCompressed, WorkManager* workManager, const tchar* hint)
+	CasKey CalculateCasKey(const u8* fileMem, u64 fileSize, bool storeCompressed, WorkManager* workManager, const tchar* hint)
 	{
 		CasKeyHasher hasher;
 
@@ -67,7 +67,7 @@ namespace uba
 				Atomic<u64> refCount = 2;
 				Atomic<u64> counter = 0;
 				Atomic<u64> doneCounter = 0;
-				u8* fileMem = nullptr;
+				const u8* fileMem = nullptr;
 				u64 workCount = 0;
 				u64 fileSize = 0;
 				bool error = false;
@@ -98,7 +98,7 @@ namespace uba
 
 					u64 startOffset = BufferSlotSize*index;
 					u64 toRead = Min(BufferSlotSize, rec->fileSize - startOffset);
-					u8* slot = rec->fileMem + startOffset;
+					const u8* slot = rec->fileMem + startOffset;
 					CasKeyHasher hasher;
 					hasher.Update(slot, toRead);
 					rec->keys[index] = ToCasKey(hasher, false);
@@ -986,8 +986,9 @@ namespace uba
 
 	////////////////////////////////////////////////////////////////////////////////////////////////////
 
-	CompressWriter::CompressWriter(Logger& l, BufferCache& bs, WorkManager& wm, StorageStats& s, u8 cc, u8 ccl, bool avof)
+	CompressWriter::CompressWriter(Logger& l, SharedMemoryAllocator& a, BufferCache& bs, WorkManager& wm, StorageStats& s, u8 cc, u8 ccl, bool avof)
 	:	m_logger(l)
+	,	m_allocator(a)
 	,	m_bufferSlots(bs)
 	,	m_workManager(wm)
 	,	m_stats(s)
@@ -1018,18 +1019,17 @@ namespace uba
 		return true;
 	}
 
-	bool CompressWriter::CompressToMapping(FileMappingHandle& outMappingHandle, u64& outMappingSize, u8* readMem, u64 fileSize, const tchar* hint)
+	bool CompressWriter::CompressToMapping(SharedMemoryHandle& outMappingHandle, u64& outMappingSize, u8* readMem, u64 fileSize, const tchar* hint)
 	{
 		u64 mappingSize = fileSize + 1024;
-		FileMappingHandle destMapping = CreateMemoryMappingW(m_logger, PAGE_READWRITE, mappingSize, nullptr, hint);
-		if (!destMapping.IsValid())
-			return false;
-		auto destGuard = MakeGuard([&]() { CloseFileMapping(m_logger, destMapping, hint); });
 
-		u8* compressedData = MapViewOfFile(m_logger, destMapping, FILE_MAP_WRITE, 0, mappingSize);
+		SharedMemoryHandle destMapping = m_allocator.CreateHandle(m_logger, hint);
+		m_allocator.ExtendMemory(destMapping, mappingSize);
+		auto destGuard = MakeGuard([&]() { m_allocator.CloseHandle(m_logger, destMapping, hint); });
+		u8* compressedData = m_allocator.MapView(destMapping, hint, SharedMemoryMapType_ReadWrite);
 		if (!compressedData)
 			return false;
-		auto unmapGuard = MakeGuard([&]() { UnmapViewOfFile(m_logger, compressedData, mappingSize, hint); });
+		auto unmapGuard = MakeGuard([&]() { m_allocator.UnmapView(destMapping, compressedData); });
 
 		LinearWriterMem destination(compressedData, mappingSize);
 		if (!CompressFromMemOrFile(destination, hint, InvalidFileHandle, readMem, fileSize))
diff --git a/Engine/Source/Programs/UnrealBuildAccelerator/Common/Public/UbaFileMapping.h b/Engine/Source/Programs/UnrealBuildAccelerator/Common/Public/UbaFileMapping.h
--- a/Engine/Source/Programs/UnrealBuildAccelerator/Common/Public/UbaFileMapping.h
+++ b/Engine/Source/Programs/UnrealBuildAccelerator/Common/Public/UbaFileMapping.h
@@ -2,15 +2,16 @@
 
 #pragma once
 
+#include "UbaDefinitions.h"
 #include "UbaEnvironment.h"
 #include "UbaEvent.h"
 #include "UbaHash.h"
+#include "UbaSharedMemoryView.h"
 #include "UbaSynchronization.h"
 
 namespace uba
 {
 	class Logger;
-	class WorkManager;
 
 	////////////////////////////////////////////////////////////////////////////////////////////////////
 
@@ -42,117 +43,31 @@ namespace uba
 
 	////////////////////////////////////////////////////////////////////////////////////////////////////
 
-	struct MappedView
+	struct FileMappingAllocation
 	{
 		FileMappingHandle handle;
 		u64 offset = 0;
-		u64 size = 0;
 		u8* memory = nullptr;
-		bool isCompressed = true;
-		bool isIndependent = false;
-	};
-
-	enum FileMappingType : u8
-	{
-		MappedView_Transient = 0,
-		MappedView_Persistent = 1,
-	};
-
-	class FileMappingBuffer
-	{
-	public:
-		FileMappingBuffer(Logger& logger, WorkManager* workManager = nullptr);
-		~FileMappingBuffer();
-
-		bool AddTransient(const tchar* name, bool keepMapped);
-		bool AddPersistent(const tchar* name, FileHandle fileHandle, u64 size, u64 capacity);
-		void CloseDatabase();
-
-		MappedView AllocAndMapView(FileMappingType type, u64 size, u64 alignment, const tchar* hint, bool allowShrink = false, bool isIndependent = false);
-		MappedView MapView(FileMappingHandle handle, u64 offset, u64 size, const tchar* hint, bool isIndependent = false);
-		MappedView MapViewWrite(FileMappingHandle handle, u64 offset, u64 size, const tchar* hint, bool isIndependent = false);
-		MappedView MapView(const StringView& str, u64 size, const tchar* hint);
-		void UnmapView(MappedView view, const tchar* hint, u64 newSize = InvalidValue, bool allowDeferredUnmap = true);
-		
-		void GetSizeAndCount(FileMappingType type, u64& outSize, u32& outCount);
-		FileMappingHandle GetPersistentHandle(u32 index) { return m_storage[MappedView_Persistent].files[index].handle; }
-		FileHandle GetPersistentFile(u32 index) { return m_storage[MappedView_Persistent].files[index].file; }
-		u64 GetPersistentSize(u32 index) { return m_storage[MappedView_Persistent].files[index].size; }
-
-		u64 GetFileMappingCapacity();
-
-	private:
-		struct File
-		{
-			File* next = nullptr;
-			const tchar* name = nullptr;
-			FileHandle file = InvalidFileHandle;
-			FileMappingHandle handle;
-			u64 size = 0;
-			u64 capacity = 0;
-			u8* mappedMemory = nullptr;
-			//Atomic<u32> activeMapCount;
-			bool commitOnAlloc = false;
-			bool keepMapped = false;
-		};
-
-		MappedView MapView(FileMappingHandle handle, u64 offset, u64 size, const tchar* hint, bool isIndependent, u32 desiredAccess);
-		MappedView AllocAndMapViewNoLock(File& file, u64 size, u64 alignment, const tchar* hint);
-		File& GetFile(FileMappingHandle handle, u8& outStorageIndex);
-
-		Logger& m_logger;
-		WorkManager* m_workManager;
-		u64 m_pageSize = 0;
-
-		struct MappingStorage
-		{
-			File files[128];
-			u32 fileCount = 0;
-			u32 fullFileCount = 0;
-
-			Futex availableFilesLock;
-			Event availableFilesEvent;
-			File* availableFile = nullptr;
-
-			MappingStorage() : availableFilesEvent(EventResetType_Auto) {}
-		};
-
-		File* PopFile(MappingStorage& storage, u64 size, u64 alignment);
-		void PushFile(MappingStorage& storage, File* file);
-		void CloseMappingStorage(MappingStorage& storage);
-
-		MappingStorage m_storage[2];
-
-		FileMappingBuffer(const FileMappingBuffer&) = delete;
-		void operator=(const FileMappingBuffer&) = delete;
 	};
 
 	////////////////////////////////////////////////////////////////////////////////////////////////////
 
-	class FileMappingAllocator
+	class FixedSizeFileMappingAllocator
 	{
 	public:
-		FileMappingAllocator(Logger& logger, const tchar* name);
-		~FileMappingAllocator();
+		FixedSizeFileMappingAllocator(Logger& logger, const tchar* name);
+		~FixedSizeFileMappingAllocator();
 
 		bool Init(u64 blockSize, u64 capacity);
 
-		struct Allocation
-		{
-			FileMappingHandle handle;
-			u64 offset = 0;
-			u8* memory = nullptr;
-		};
-
-		Allocation Alloc(const tchar* hint);
-		void Free(Allocation allocation);
+		FileMappingAllocation Alloc(const tchar* hint);
+		void Free(FileMappingAllocation allocation);
 
 		u64 GetSize();
 
 	private:
 		Logger& m_logger;
 		const tchar* m_name;
-		u64 m_pageSize = 0;
 		u64 m_capacity = 0;
 		u64 m_blockSize = 0;
 
@@ -162,7 +77,9 @@ namespace uba
 
 		Set<u64> m_availableBlocks;
 
-		FileMappingAllocator(const FileMappingAllocator&) = delete;
-		void operator=(const FileMappingAllocator&) = delete;
+		FixedSizeFileMappingAllocator(const FixedSizeFileMappingAllocator&) = delete;
+		void operator=(const FixedSizeFileMappingAllocator&) = delete;
 	};
+
+	////////////////////////////////////////////////////////////////////////////////////////////////////
 }
diff --git a/Engine/Source/Programs/UnrealBuildAccelerator/Common/Public/UbaFileMappingBuffer.h b/Engine/Source/Programs/UnrealBuildAccelerator/Common/Public/UbaFileMappingBuffer.h
new file mode 100644
--- /dev/null
+++ b/Engine/Source/Programs/UnrealBuildAccelerator/Common/Public/UbaFileMappingBuffer.h
@@ -0,0 +1,70 @@
+// Copyright Epic Games, Inc. All Rights Reserved.
+
+#pragma once
+
+#include "UbaFileMapping.h"
+
+namespace uba
+{
+	class SharedMemoryAllocator;
+	class Trace;
+
+	////////////////////////////////////////////////////////////////////////////////////////////////////
+
+	struct MappedView
+	{
+		SharedMemoryHandle handle;
+		u64 offset = 0;
+		u64 size = 0;
+		u8* memory = nullptr;
+	};
+
+	////////////////////////////////////////////////////////////////////////////////////////////////////
+
+	class FileMappingBuffer
+	{
+	public:
+		FileMappingBuffer(Logger& logger, SharedMemoryAllocator& allocator);
+		~FileMappingBuffer();
+
+		bool Init(const tchar* name, u64 capacity, u64 commitStepSize);
+
+		MappedView AllocAndMapView(u64 size, u64 alignment, const tchar* hint, bool isIndependent = false);
+		MappedView MapView(SharedMemoryHandle handle, u64 offset, u64 size, const tchar* hint);
+		MappedView MapView(StringView handleName, u64 size, const tchar* hint);
+		void UnmapView(MappedView view, const tchar* hint, u64 newSize = InvalidValue);
+
+		bool IsIndependent(SharedMemoryHandle handle);
+		void CloseIndependent(SharedMemoryHandle handle, u64 size, const tchar* hint);
+
+		SharedMemoryHandle GetMemoryHandle();
+		u64 GetUsed();
+
+		void TraceStatus(Trace& trace, u32 startRow, const tchar* hint, u32 traceLevel);
+
+	private:
+		MappedView InternalMapView(SharedMemoryHandle handle, u64 offset, u64 size, const tchar* hint);
+
+		Logger& m_logger;
+		SharedMemoryAllocator& m_allocator;
+
+		TString m_name;
+		u64 m_commitStepSize = 0;
+
+		SharedMemoryHandle m_memoryHandle;
+		Futex m_usedLock;
+		u8* m_memory = nullptr;
+		u64 m_capacity = 0;
+		u64 m_committed = 0;
+		u64 m_used = 0;
+
+		Atomic<u64> m_independentMappingCount = 0;
+		Atomic<u64> m_independentSize = 0;
+
+		FileMappingBuffer(const FileMappingBuffer&) = delete;
+		void operator=(const FileMappingBuffer&) = delete;
+	};
+
+	////////////////////////////////////////////////////////////////////////////////////////////////////
+
+}
diff --git a/Engine/Source/Programs/UnrealBuildAccelerator/Common/Public/UbaProcess.h b/Engine/Source/Programs/UnrealBuildAccelerator/Common/Public/UbaProcess.h
--- a/Engine/Source/Programs/UnrealBuildAccelerator/Common/Public/UbaProcess.h
+++ b/Engine/Source/Programs/UnrealBuildAccelerator/Common/Public/UbaProcess.h
@@ -115,7 +115,7 @@ namespace uba
 		ProcessImpl* m_parentProcess;
 		Futex m_initLock;
 		u32 m_id;
-		FileMappingAllocator::Allocation m_comMemory;
+		FileMappingAllocation m_comMemory;
 
 	#if PLATFORM_WINDOWS
 		SharedEvent m_cancelEvent;
@@ -146,6 +146,7 @@ namespace uba
 			Futex writtenFilesLock;
 			UnorderedMap<StringKey, WrittenFile> writtenFiles;
 			OverflowedWrittenFilesMessage* firstOverflowMessage = nullptr;
+			UnorderedSet<SharedMemoryHandle> sharedMemoryHandles; // Track created memory handles. entries can be transferred to WrittenFiles
 		};
 
 		u32 m_nativeProcessId = 0;
diff --git a/Engine/Source/Programs/UnrealBuildAccelerator/Common/Public/UbaSession.h b/Engine/Source/Programs/UnrealBuildAccelerator/Common/Public/UbaSession.h
--- a/Engine/Source/Programs/UnrealBuildAccelerator/Common/Public/UbaSession.h
+++ b/Engine/Source/Programs/UnrealBuildAccelerator/Common/Public/UbaSession.h
@@ -7,22 +7,24 @@
 #include "UbaFile.h"
 #include "UbaDefinitions.h"
 #include "UbaDirectoryTable.h"
-#include "UbaFileMapping.h"
+#include "UbaFileMappingBuffer.h"
 #include "UbaProcessHandle.h"
 #include "UbaProcessStartInfo.h"
 #include "UbaRootPaths.h"
 #include "UbaSessionCreateInfo.h"
+#include "UbaSharedMemoryView.h"
 #include "UbaStats.h"
 #include "UbaTrace.h"
 #include "UbaThread.h"
 
-#define UBA_DEBUG_TRACK_MAPPING 0 // UBA_DEBUG_LOGGER
+#define UBA_DEBUG_TRACK_MAPPING UBA_DEBUG_LOGGER
 
 namespace uba
 {
 	class Process;
 	class ProcessHandle;
 	class ProcessImpl;
+	class SharedMemoryAllocator;
 	class Storage;
 	class WorkManager;
 	struct ProcessStartInfo;
@@ -113,6 +115,8 @@ namespace uba
 		const u8* GetDirectoryTableMemory() { return m_directoryTableMem; }
 		u32 GetDirectoryTableSize();
 
+		void FlushDeadProcesses();
+
 	protected:
 		Session(const SessionCreateInfo& info, const tchar* logPrefix, bool runningRemote, WorkManager& workManager);
 		bool Create(const SessionCreateInfo& info);
@@ -121,7 +125,6 @@ namespace uba
 		ProcessHandle InternalRunProcess(const ProcessStartInfo& startInfo, bool async, ProcessImpl* parent, bool enableDetour);
 		void ProcessAdded(Process& process, u32 sessionId);
 		void ProcessExited(ProcessImpl& process, u64 executionTime);
-		void FlushDeadProcesses();
 		void PrintProcessStats(ProcessStats& stats, const tchar* logName);
 		u32 GetFileMappingSize();
 		u32 GetMemoryMapAlignment(const StringView& fileName) const;
@@ -133,7 +136,7 @@ namespace uba
 		bool GetBinaryModules(Vector<BinaryModule>& out, const tchar* application);
 		void Free(Vector<BinaryModule>& v);
 		bool IsRarelyRead(ProcessImpl& process, const StringView& fileName) const;
-		bool IsRarelyReadAfterWritten(ProcessImpl& process, const StringView& fileName) const;
+		bool IsMaybeReadAfterWritten(ProcessImpl& process, const StringView& fileName) const;
 		bool IsKnownSystemFile(const tchar* applicationName);
 		bool ShouldWriteToDisk(const StringView& fileName);
 		u32 WriteDirectoryEntries(const StringKey& dirKey, const StringView& dirPath, u32* outTableOffset = nullptr);
@@ -143,6 +146,7 @@ namespace uba
 		struct FileMappingEntry;
 		bool GetOrCreateMemoryMapFromFile(MemoryMap& out, StringKey fileNameKey, const tchar* fileName, bool isCompressedCas, u64 alignment, const tchar* hint, ProcessImpl* requestingProcess, bool canBeFreed);
 		bool GetOrCreateMemoryMapFromStorage(MemoryMap& out, StringKey fileNameKey, const tchar* fileName, const CasKey& casKey, u64 alignment);
+		void ThreadGcFileMappings();
 
 		bool RegisterCreateFileForWrite(StringKey fileNameKey, const StringView& fileName, bool registerRealFile, u64 fileSize = 0, u64 lastWriteTime = 0, bool invalidateStorage = true);
 		u32 RegisterDeleteFile(StringKey fileNameKey, const StringView& fileName);
@@ -151,9 +155,10 @@ namespace uba
 		virtual void* GetProcessEnvironmentVariables();
 		virtual void PrintSessionStats(Logger& logger);
 
-		virtual bool RegisterVirtualFileInternal(const StringKey& fileNameKey, const StringView& filePath, const tchar* sourceFile, u64 sourceOffset, u64 sourceSize);
-		virtual bool CreateVirtualFileInternal(const StringKey& fileNameKey, const StringView& filePath, const void* memory, u64 memorySize, bool transient);
-		virtual bool RegisterVirtualFileInternal(const StringKey& fileNameKey, const StringView& filePath, FileMappingHandle mappingHandle, u64 mappingSize, u64 mappingOffset, bool transient);
+		bool RegisterVirtualFileInternal(const StringKey& fileNameKey, const StringView& filePath, const tchar* sourceFile, u64 sourceOffset, u64 sourceSize);
+		bool CreateVirtualFileInternal(const StringKey& fileNameKey, const StringView& filePath, const void* memory, u64 memorySize, bool transient);
+		bool RegisterVirtualFileInternal(const StringKey& fileNameKey, const StringView& filePath, SharedMemoryHandle memoryHandle, u64 mappingSize, u64 mappingOffset, bool transient);
+
 		virtual bool GetOutputFileSizeInternal(u64& outSize, const StringKey& fileNameKey, StringView filePath);
 		virtual bool GetOutputFileDataInternal(void* outData, const StringKey& fileNameKey, StringView filePath, bool deleteInternalMapping);
 		virtual bool WriteOutputFileInternal(const StringKey& fileNameKey, StringView filePath, bool deleteInternalMapping);
@@ -217,6 +222,7 @@ namespace uba
 		void RunDependencyCrawler(ProcessImpl& process);
 
 		Storage& m_storage;
+		SharedMemoryAllocator& m_sharedMemory;
 		MutableLogger m_logger;
 		WorkManager& m_workManager;
 
@@ -265,6 +271,9 @@ namespace uba
 		FileMappingHandle m_fileMappingTableHandle;
 		FileMappingBuffer m_fileMappingBuffer;
 
+		Atomic<u64> m_independentMappingCreated;
+		Atomic<u64> m_independentMappingActive;
+
 		ReaderWriterLock m_fileMappingTableMemLock;
 		u8* m_fileMappingTableMem;
 		u32 m_fileMappingTableSize = 0;
@@ -272,8 +281,10 @@ namespace uba
 		struct FileMappingEntry
 		{
 			Futex lock;
-			FileMappingHandle mapping;
-			u64 mappingOffset = 0;
+
+			SharedMemoryHandle memoryHandle;
+			u64 memoryOffset = 0;
+
 			u64 contentSize = 0; // Size of uncompressed file or mapping
 			u64 lastWriteTime = 0;
 			u64 storedSize = InvalidValue; // Size stored on disk. Can be compressed or uncompressed
@@ -294,12 +305,17 @@ namespace uba
 		};
 		UnorderedMap<StringKey, FileMappingEntry> m_fileMappingTableLookup;
 
+		Futex m_fileMappingsThatCanBeFreedLock;
+		UnorderedMap<StringKey, u64> m_fileMappingsThatCanBeFreed;
+		Thread m_fileMappingGcThread;
+		Event m_fileMappingGcDone;
+
 		static constexpr u64 NameToHashMemSize = 48*1024*1024;
 		MemoryBlock m_nameToHashTableMem;
 
 		Atomic<u64> m_fileIndexCounter = 8000000000;
 
-		FileMappingAllocator m_processCommunicationAllocator;
+		FixedSizeFileMappingAllocator m_processCommunicationAllocator;
 		std::string m_detoursLibrary[2]; // x64 and arm64
 
 		Futex m_processStatsLock;
@@ -319,7 +335,7 @@ namespace uba
 		UnorderedMap<u32, ActiveFile> m_activeFiles;
 
 		Futex m_virtualSourceFilesLock;
-		struct VirtualSourceFile { FileMappingHandle mappingHandle; u64 size; };
+		struct VirtualSourceFile { FileMappingHandle mappingHandle; SharedMemoryHandle memoryHandle; u64 size; };
 		UnorderedMap<StringKey, VirtualSourceFile> m_virtualSourceFiles;
 
 		Futex m_rootsLookupLock;
@@ -397,6 +413,9 @@ namespace uba
 
 	struct InitResponse
 	{
+		SharedMemoryAllocatorHandle sharedMemoryAllocatorHandle;
+		SharedMemoryHandle permanentFilesHandle;
+		u64 permanentFilesOffset = 0;
 		u64 directoryTableHandle = 0;
 		u32 directoryTableSize = 0;
 		u32 directoryTableCount = 0;
@@ -433,8 +452,8 @@ namespace uba
 		u32 attributes = 0;
 		bool deleteOnClose = false;
 		bool success = true;
-		FileMappingHandle mappingHandle;
-		u64 mappingWritten = 0;
+		SharedMemoryHandle memoryHandle;
+		u64 memoryWritten = 0;
 	};
 
 	struct CloseFileResponse
@@ -579,9 +598,11 @@ namespace uba
 		StringKey key;
 		TString name;
 		TString backedName;
-		FileMappingHandle mappingHandle;
-		u64 mappingWritten = 0;
-		FileMappingHandle originalMappingHandle;
+
+		SharedMemoryHandle memoryHandle;
+		SharedMemoryHandle originalMemoryHandle;
+
+		u64 memoryWritten = 0;
 		u64 lastWriteTime = 0;
 		u32 attributes = 0;
 	};
diff --git a/Engine/Source/Programs/UnrealBuildAccelerator/Common/Public/UbaSessionCreateInfo.h b/Engine/Source/Programs/UnrealBuildAccelerator/Common/Public/UbaSessionCreateInfo.h
--- a/Engine/Source/Programs/UnrealBuildAccelerator/Common/Public/UbaSessionCreateInfo.h
+++ b/Engine/Source/Programs/UnrealBuildAccelerator/Common/Public/UbaSessionCreateInfo.h
@@ -42,7 +42,6 @@ namespace uba
 		bool extractObjFilesSymbols = false;	// Will extract import/export symbols to a file
 		bool treatTempDirAsEmpty = true;		// If this is true, directory table will always see temp as empty. Use sub directories of temp instead
 		bool useFakeVolumeSerial = true;		// Fake volume serials reduce the directory table size since serials end up being compressed down to one byte
-		bool keepTransientDataMapped = true;	// Will keep the transient data mapped in instead of doing map/unmap everytime we access it
 		bool allowDiscardVirtualMemory = false;	// Set to true to use DiscardVirtualMemory around map/unmap views of file mappings. This seems to speed up perf quite a lot
 		bool allowLinkDependencyCrawler = true;	// Enable crawler that reads ahead to find needed obj files for linker
 		int linkCrawlerMaxParallel = -2;		// If allowLinkDependencyCrawler is set, this controls number of decompress in parallel. If negative it is max cores - value
diff --git a/Engine/Source/Programs/UnrealBuildAccelerator/Common/Public/UbaSharedMemoryAllocator.h b/Engine/Source/Programs/UnrealBuildAccelerator/Common/Public/UbaSharedMemoryAllocator.h
--- a/Engine/Source/Programs/UnrealBuildAccelerator/Common/Public/UbaSharedMemoryAllocator.h
+++ b/Engine/Source/Programs/UnrealBuildAccelerator/Common/Public/UbaSharedMemoryAllocator.h
@@ -2,10 +2,14 @@
 
 #pragma once
 
+#include "UbaDefinitions.h"
 #include "UbaMemory.h"
 #include "UbaPlatform.h"
 #include "UbaSharedMemoryView.h"
 
+// Track shared memory allocations to be able to trace biggest allocations
+#define UBA_TRACK_SHARED_MEMORY_ALLOCATIONS 0
+
 namespace uba
 {
 	struct BinaryWriter;
@@ -18,17 +22,25 @@ namespace uba
 	{
 	public:
 		bool Init(u64 capacity, u64 commitStepSize = 16ull*1024*1024);
+		void Deinit(bool keepFileMapping); // keepFileMapping can be used to just leak handle for shutdown
 
 		operator SharedMemoryAllocatorHandle() const;
-		u64 GetReserved();
+		u64 GetCapacity();
 
 		SharedMemoryHandle CreateHandle(Logger& logger, const tchar* hint, bool priority = false); // Priority means it will have a compressable handle id
+		SharedMemoryHandle DuplicateHandle(SharedMemoryHandle handle, const tchar* hint);
 		void ClearHandle(SharedMemoryHandle handle);
 		void CloseHandle(Logger& logger, SharedMemoryHandle handle, const tchar* hint);
 
 		void ExtendMemory(BinaryWriter& out, SharedMemoryHandle handle, u64 size, const tchar* hint = TC(""), bool zeroMemory = true);
 		void ExtendMemory(SharedMemoryHandle handle, u64 size, const tchar* hint = TC(""), bool zeroMemory = true);
-		void ExtendMemory(SharedMemoryView& view, SharedMemoryHandle handle, u64 size, const tchar* hint = TC(""), bool zeroMemory = true, u64 preferredOffset = 0);
+		void ExtendMemory(SharedMemoryView& view, SharedMemoryHandle handle, u64 size, const tchar* hint = TC(""), bool zeroMemory = true);
+		void TrimMemory(SharedMemoryHandle handle, u64 usedSize);
+
+		u8* ReserveMemory(SharedMemoryHandle handle, u64 size);
+		void CommitMemory(SharedMemoryHandle handle, u64 offset, u64 size);
+		void DecommitMemory(SharedMemoryHandle handle, u64 offset, u64 size);
+		void UnreserveMemory(SharedMemoryHandle handle);
 
 		u64 GetCommitted(SharedMemoryHandle handle);
 		u64 GetRefCount(SharedMemoryHandle handle);
@@ -37,7 +49,7 @@ namespace uba
 		bool MapView(SharedMemoryView& view, SharedMemoryHandle handle, const tchar* hint, SharedMemoryMapType mapType = SharedMemoryMapType_ReadOnly, u64 offset = 0, u64 size = ~0llu);
 
 
-		u8* MapView(SharedMemoryHandle handle, SharedMemoryMapType mapType = SharedMemoryMapType_ReadOnly);
+		u8* MapView(SharedMemoryHandle handle, const tchar* hint, SharedMemoryMapType mapType = SharedMemoryMapType_ReadOnly);
 		void UnmapView(SharedMemoryHandle handle, u8* memory);
 
 
@@ -46,24 +58,28 @@ namespace uba
 		SharedMemoryAllocator(Logger& logger);
 		~SharedMemoryAllocator();
 
+		// This is really ugly. But very contained for virtual file mappings
+		SharedMemoryHandle RegisterExternalMapping(FileMappingHandle file);
+
 	private:
 		struct InternalHandle;
 		struct Slice;
 
 		InternalHandle& GetHandle(SharedMemoryHandle h);
 		void ReleaseHandle(SharedMemoryHandle h);
-		void ExtendHandle(InternalHandle& handle, u64 size, const tchar* hint, bool zeroMemory, u64 preferredOffset);
+		void ExtendHandle(InternalHandle& handle, u64 size, const tchar* hint, bool zeroMemory);
 		void ClearHandle(InternalHandle& handle);
-		Slice AllocateNoLock(u64 desiredSize, u64 preferredOffset);
+		Slice AllocateNoLock(u64 desiredSize, u64 existingSliceCount);
 		void FreeNoLock(Slice slice);
 
 		Logger& m_logger;
 
 		u64 m_commitStepSize;
 		FileMappingHandle m_fileMapping;
-		u64 m_reserved = 0;
+		u64 m_capacity = 0; // Capacity is total reserved address space from system
+		u64 m_reserved = 0; // Reserved is the amount of address space used
+		u64 m_committed = 0; // Committed is the amount of pages in address space that has been committed
 		u8* m_memory = nullptr;
-		u64 m_committed = 0;
 
 		u64 m_used = 0;
 
@@ -81,6 +97,11 @@ namespace uba
 
 		SharedMemoryAllocator(const SharedMemoryAllocator&) = delete;
 		SharedMemoryAllocator& operator=(const SharedMemoryAllocator&) = delete;
+
+		#if UBA_TRACK_SHARED_MEMORY_ALLOCATIONS
+		InternalHandle* m_firstHandle = nullptr;
+		Futex m_handleLock;
+		#endif
 	};
 
 	////////////////////////////////////////////////////////////////////////////////////////////////////
diff --git a/Engine/Source/Programs/UnrealBuildAccelerator/Common/Public/UbaStorage.h b/Engine/Source/Programs/UnrealBuildAccelerator/Common/Public/UbaStorage.h
--- a/Engine/Source/Programs/UnrealBuildAccelerator/Common/Public/UbaStorage.h
+++ b/Engine/Source/Programs/UnrealBuildAccelerator/Common/Public/UbaStorage.h
@@ -5,10 +5,11 @@
 #include "UbaBottleneck.h"
 #include "UbaBufferCache.h"
 #include "UbaFile.h"
-#include "UbaFileMapping.h"
+#include "UbaFileMappingBuffer.h"
 #include "UbaLogger.h"
 #include "UbaMemory.h"
 #include "UbaPathUtils.h"
+#include "UbaSharedMemoryAllocator.h"
 #include "UbaStats.h"
 
 #define UBA_TRACK_IS_EXECUTABLE !PLATFORM_WINDOWS
@@ -22,6 +23,14 @@ namespace uba
 	struct DirectoryEntry;
 	extern CasKey EmptyFileKey;
 	
+	struct MappedView2 : MappedView
+	{
+		MappedView2() = default;
+		MappedView2(const MappedView2&) = default;
+		MappedView2(const MappedView& v) { *(MappedView*)this = v; }
+		bool isCompressed = true;
+	};
+
 	class Storage
 	{
 	public:
@@ -42,7 +51,7 @@ namespace uba
 		virtual bool GetFileSize(u64& outSize, const CasKey& casKey, const tchar* hint) = 0;
 		virtual bool GetFileData(void* outData, const CasKey& casKey, const tchar* hint) = 0;
 
-		struct RetrieveResult { CasKey casKey; u64 size = 0; MappedView view; };
+		struct RetrieveResult { CasKey casKey; u64 size = 0; MappedView2 view; };
 		virtual bool RetrieveCasFile(RetrieveResult& out, const CasKey& casKey, const tchar* hint, FileMappingBuffer* mappingBuffer = nullptr, u64 memoryMapAlignment = 1, bool allowProxy = true, u32 clientId = 0) = 0;
 
 		struct CachedFileInfo { CasKey casKey; };
@@ -50,7 +59,7 @@ namespace uba
 		virtual bool InvalidateCachedFileInfo(StringKey fileNameKey) = 0;
 
 		virtual bool StoreCasFile(CasKey& out, const tchar* fileName, const CasKey& casKeyOverride, bool deferCreation) = 0;
-		virtual bool StoreCasFileClient(CasKey& out, StringKey fileNameKey, const tchar* fileName, FileMappingHandle mappingHandle, u64 mappingOffset, u64 fileSize, const tchar* hint, bool keepMappingInMemory = false, bool storeCompressed = true) = 0;
+		virtual bool StoreCasFileClient(CasKey& out, StringKey fileNameKey, const tchar* fileName, SharedMemoryHandle memoryHandle, u64 memoryOffset, u64 fileSize, const tchar* hint, bool keepMappingInMemory = false, bool storeCompressed = true) = 0;
 		virtual bool DropCasFile(const CasKey& casKey, bool forceDelete, const tchar* hint, bool reduceRefCount = false) = 0;
 		virtual bool ReportBadCasFile(const CasKey& casKey) = 0;
 		virtual bool CalculateCasKey(CasKey& out, const tchar* fileName) = 0;
@@ -60,23 +69,26 @@ namespace uba
 		virtual bool FakeCopy(const CasKey& casKey, const tchar* destination, u64 size, u64 lastWritten, u32 fileAttributes, bool deleteExisting = true) = 0;
 		virtual bool GetCasFileName(StringBufferBase& out, const CasKey& casKey) = 0;
 
-		virtual MappedView MapView(const CasKey& casKey, const tchar* hint) = 0;
+		virtual MappedView2 MapView(const CasKey& casKey, const tchar* hint) = 0;
 		virtual void UnmapView(const MappedView& view, const tchar* hint) = 0;
 
 		virtual void ReportFileWrite(StringKey fileNameKey, const tchar* fileName) = 0;
 		
 		virtual StorageStats& Stats() = 0;
 		virtual void AddStats(StorageStats& stats) = 0;
+		virtual void TraceUpdate(Trace& trace, u32 startRow) = 0;
 
-		static void GetMappingString(StringBufferBase& out, FileMappingHandle mappingHandle, u64 offset);
+		static void GetMappingString(StringBufferBase& out, SharedMemoryHandle handle, u64 offset, bool canBeFreed);
 
 		virtual void SetTrace(Trace* trace, bool detailed) {}
 		virtual void Ping() {}
 
+		virtual SharedMemoryAllocator& GetSharedMemory() = 0;
+
 		struct WriteResult
 		{
-			FileMappingHandle mappingHandle;
-			u64 mappingOffset = InvalidValue;
+			SharedMemoryHandle memoryHandle;
+			u64 memoryOffset = InvalidValue;
 			u64 storedSize = InvalidValue; // Stored size of the destination
 			u64 sourceStoredSize = InvalidValue; // Stored size of the source used to write compressed file
 			bool isIndependent = false;
@@ -85,7 +97,7 @@ namespace uba
 
 		struct ExternalFileMapping
 		{
-			FileMappingHandle handle;
+			SharedMemoryHandle handle;
 			u64 offset = 0;
 			u64 size = 0;
 			u64 lastWriteTime = 0;
@@ -113,9 +125,9 @@ namespace uba
 		bool manuallyHandleOverflow = false;
 		bool asyncUnmapViewOfFile = true;
 		bool writeToDisk = true;
-		bool keepTransientDataMapped = true;	// Will keep the transient data mapped in instead of doing map/unmap everytime we access it
 		bool allowDeleteVerified = false;
 		bool allowDeferredDeletes = true;
+		bool skipCleanupOnShutdown = false;
 		MutexHandle exclusiveMutex = InvalidMutexHandle; // Set this to hand over pre-exclusive access to storage. Storage will release mutex on shutdown
 		u8 casCompressor = 0;
 		u8 casCompressionLevel = 0;
@@ -155,20 +167,22 @@ namespace uba
 		virtual bool VerifyAndGetCachedFileInfo(CachedFileInfo& out, StringKey fileNameKey, u64 verifiedLastWriteTime, u64 verifiedSize) override;
 		virtual bool InvalidateCachedFileInfo(StringKey fileNameKey) override;
 		virtual bool StoreCasFile(CasKey& out, const tchar* fileName, const CasKey& casKeyOverride, bool deferCreation) override;
-		virtual bool StoreCasFileClient(CasKey& out, StringKey fileNameKey, const tchar* fileName, FileMappingHandle mappingHandle, u64 mappingOffset, u64 fileSize, const tchar* hint, bool keepMappingInMemory = false, bool storeCompressed = true) override;
+		virtual bool StoreCasFileClient(CasKey& out, StringKey fileNameKey, const tchar* fileName, SharedMemoryHandle memoryHandle, u64 memoryOffset, u64 fileSize, const tchar* hint, bool keepMappingInMemory = false, bool storeCompressed = true) override;
 		virtual bool DropCasFile(const CasKey& casKey, bool forceDelete, const tchar* hint, bool reduceRefCount = false) override;
 		virtual bool ReportBadCasFile(const CasKey& casKey) override;
 		virtual bool CalculateCasKey(CasKey& out, const tchar* fileName) override;
 		virtual bool CopyOrLink(const CasKey& casKey, const tchar* destination, u32 fileAttributes, bool writeCompressed = false, u32 clientId = 0, const FormattingFunc& formattingFunc = {}, bool isTemp = false, bool allowHardLink = true) override;
 		virtual bool FakeCopy(const CasKey& casKey, const tchar* destination, u64 size, u64 lastWritten, u32 fileAttributes, bool deleteExisting = true) override;
 		virtual bool GetCasFileName(StringBufferBase& out, const CasKey& casKey) override;
-		virtual MappedView MapView(const CasKey& casKey, const tchar* hint) override;
+		virtual MappedView2 MapView(const CasKey& casKey, const tchar* hint) override;
 		virtual void UnmapView(const MappedView& view, const tchar* hint) override;
+		virtual SharedMemoryAllocator& GetSharedMemory() override;
 
 		virtual void ReportFileWrite(StringKey fileNameKey, const tchar* fileName) override;
 		
 		virtual StorageStats& Stats() final;
 		virtual void AddStats(StorageStats& stats) override;
+		virtual void TraceUpdate(Trace& trace, u32 startRow) override;
 
 		struct CasEntry;
 		struct CasEntryWriteLock;
@@ -176,7 +190,7 @@ namespace uba
 		virtual bool HasCasFile(const CasKey& casKey, CasEntry** out = nullptr);
 		bool EnsureCasFile(const CasKey& casKey, const tchar* fileName);
 		CasKey CalculateCasKey(const tchar* fileName, FileHandle fileHandle, u64 fileSize, bool storeCompressed);
-		CasKey CalculateCasKey(u8* fileMem, u64 fileSize, bool storeCompressed);
+		CasKey CalculateCasKey(const u8* fileMem, u64 fileSize, bool storeCompressed);
 		bool StoreCasKey(CasKey& out, const tchar* fileName, const CasKey& casKeyOverride);
 		bool StoreCasKey(CasKey& out, const StringKey& fileNameKey, const tchar* fileName, const CasKey& casKeyOverride, bool showErrors = true);
 		bool IsFileVerified(const StringKey& fileNameKey);
@@ -215,6 +229,7 @@ namespace uba
 
 		MutableLogger m_logger;
 		WorkManager& m_workManager;
+		SharedMemoryAllocator m_sharedMemory;
 		BufferCache m_bufferSlots;
 
 		StringBuffer<> m_rootDir;
@@ -271,10 +286,9 @@ namespace uba
 			bool beingWritten = false; // This is set while file is being written (when coming from network)..
 			bool disallowed = false; // This is set if cas is created from disallowed file
 
-			bool mappingIsIndependent = false; // if true, mappingHandle is stand-alone and must be manually deleted from the outside using DropCasFile
-			FileMappingHandle mappingHandle;
-			u64 mappingOffset = 0;
-			u64 mappingSize = 0;
+			SharedMemoryHandle memoryHandle;
+			u64 memoryOffset = 0;
+			u64 memorySize = 0;
 		};
 
 		UnorderedMap<CasKey, CasEntry> m_casLookup;
@@ -298,6 +312,7 @@ namespace uba
 		bool m_asyncUnmapViewOfFile = true;
 		bool m_allowDeleteVerified = false;
 		bool m_allowDeferredDeletes = true;
+		bool m_skipCleanupOnShutdown = false;
 		bool m_writeToDisk;
 
 		MutexHandle m_exclusiveMutex = InvalidMutexHandle;
diff --git a/Engine/Source/Programs/UnrealBuildAccelerator/Common/Public/UbaStorageClient.h b/Engine/Source/Programs/UnrealBuildAccelerator/Common/Public/UbaStorageClient.h
--- a/Engine/Source/Programs/UnrealBuildAccelerator/Common/Public/UbaStorageClient.h
+++ b/Engine/Source/Programs/UnrealBuildAccelerator/Common/Public/UbaStorageClient.h
@@ -52,12 +52,12 @@ namespace uba
 
 		virtual bool GetCasFileName(StringBufferBase& out, const CasKey& casKey) override;
 
-		virtual MappedView MapView(const CasKey& casKey, const tchar* hint) override;
+		virtual MappedView2 MapView(const CasKey& casKey, const tchar* hint) override;
 
 		virtual bool GetZone(StringBufferBase& out) override;
 		virtual bool RetrieveCasFile(RetrieveResult& out, const CasKey& casKey, const tchar* hint, FileMappingBuffer* mappingBuffer = nullptr, u64 memoryMapAlignment = 1, bool allowProxy = true, u32 clientId = 0) override;
 		virtual bool StoreCasFile(CasKey& out, const tchar* fileName, const CasKey& casKeyOverride, bool deferCreation) override;
-		virtual bool StoreCasFileClient(CasKey& out, StringKey fileNameKey, const tchar* fileName, FileMappingHandle mappingFile, u64 mappingOffset, u64 fileSize, const tchar* hint, bool keepMappingInMemory = false, bool storeCompressed = true) override;
+		virtual bool StoreCasFileClient(CasKey& out, StringKey fileNameKey, const tchar* fileName, SharedMemoryHandle mappingFile, u64 mappingOffset, u64 fileSize, const tchar* hint, bool keepMappingInMemory = false, bool storeCompressed = true) override;
 		virtual bool HasCasFile(const CasKey& casKey, CasEntry** out = nullptr) override;
 		virtual void Ping() override;
 		virtual void PrintSummary(Logger& logger) override;
diff --git a/Engine/Source/Programs/UnrealBuildAccelerator/Common/Public/UbaStorageServer.h b/Engine/Source/Programs/UnrealBuildAccelerator/Common/Public/UbaStorageServer.h
--- a/Engine/Source/Programs/UnrealBuildAccelerator/Common/Public/UbaStorageServer.h
+++ b/Engine/Source/Programs/UnrealBuildAccelerator/Common/Public/UbaStorageServer.h
@@ -103,6 +103,7 @@ namespace uba
 			Atomic<u64> sendCasTime = 0;
 
 			FileHandle readFileHandle = InvalidFileHandle;
+			FileMappingHandle fileMappingHandle;
 			MappedView mappedView;
 			const u8* memoryBegin = nullptr;
 			const u8* memoryPos = nullptr;
diff --git a/Engine/Source/Programs/UnrealBuildAccelerator/Common/Public/UbaStorageUtils.h b/Engine/Source/Programs/UnrealBuildAccelerator/Common/Public/UbaStorageUtils.h
--- a/Engine/Source/Programs/UnrealBuildAccelerator/Common/Public/UbaStorageUtils.h
+++ b/Engine/Source/Programs/UnrealBuildAccelerator/Common/Public/UbaStorageUtils.h
@@ -4,6 +4,7 @@
 
 #include "UbaBinaryReaderWriter.h"
 #include "UbaFile.h"
+#include "UbaFileMapping.h"
 #include "UbaHash.h"
 #include "UbaNetworkMessage.h"
 #include <oodle2.h>
@@ -11,6 +12,7 @@
 namespace uba
 {
 	class FileAccessor;
+	class SharedMemoryAllocator;
 	struct BufferCache;
 	struct StorageStats;
 
@@ -22,7 +24,7 @@ namespace uba
 	u8 GetCompressor(const tchar* str);
 	u8 GetCompressionLevel(const tchar* str);
 
-	CasKey CalculateCasKey(u8* fileMem, u64 fileSize, bool storeCompressed, WorkManager* workManager, const tchar* hint);
+	CasKey CalculateCasKey(const u8* fileMem, u64 fileSize, bool storeCompressed, WorkManager* workManager, const tchar* hint);
 	bool SendBatchMessages(Logger& logger, NetworkClient& client, u16 fetchId, u8* slot, u64 capacity, u64 left, u32 messageMaxSize, u32& readIndex, u32& responseSize, const Function<bool()>& runInWaitFunc = {}, const tchar* hint = TC(""), u32* outError = nullptr, const Function<bool(u64)>& progressFunc = {});
 	bool SendFile(Logger& logger, NetworkClient& client, WorkManager* workManager, const CasKey& casKey, const u8* sourceMem, u64 sourceSize, const tchar* hint);
 
@@ -72,14 +74,15 @@ namespace uba
 
 	struct CompressWriter
 	{
-		CompressWriter(Logger& l, BufferCache& bs, WorkManager& wm, StorageStats& s, u8 cc = DefaultCompressor, u8 ccl = DefaultCompressionLevel, bool avof = false);
+		CompressWriter(Logger& l, SharedMemoryAllocator& a, BufferCache& bs, WorkManager& wm, StorageStats& s, u8 cc = DefaultCompressor, u8 ccl = DefaultCompressionLevel, bool avof = false);
 
 		bool CompressToFile(u64& outCompressedSize, const tchar* from, FileHandle readHandle, u8* readMem, u64 fileSize, const tchar* toFile, const void* header, u64 headerSize, u64 lastWriteTime, const tchar* tempPath);
-		bool CompressToMapping(FileMappingHandle& outMappingHandle, u64& outMappingSize, u8* readMem, u64 fileSize, const tchar* hint);
+		bool CompressToMapping(SharedMemoryHandle& outMappingHandle, u64& outMappingSize, u8* readMem, u64 fileSize, const tchar* hint);
 		bool CompressFromMem(LinearWriter& destination, u32 workCount, const u8* uncompressedData, u64 fileSize, u64 maxUncompressedBlock, u64& totalWritten);
 		bool CompressFromMemOrFile(LinearWriter& destination, const tchar* from, FileHandle readHandle, u8* readMem, u64 fileSize);
 
 		Logger& m_logger;
+		SharedMemoryAllocator& m_allocator;
 		BufferCache& m_bufferSlots;
 		WorkManager& m_workManager;
 		StorageStats& m_stats;
diff --git a/Engine/Source/Programs/UnrealBuildAccelerator/Core/Private/UbaApplicationRules.cpp b/Engine/Source/Programs/UnrealBuildAccelerator/Core/Private/UbaApplicationRules.cpp
--- a/Engine/Source/Programs/UnrealBuildAccelerator/Core/Private/UbaApplicationRules.cpp
+++ b/Engine/Source/Programs/UnrealBuildAccelerator/Core/Private/UbaApplicationRules.cpp
@@ -60,6 +60,15 @@ namespace uba
 				return 1024 * 1024 * 1024;
 			return Super::FileTypeMaxSize(file, isSystemOrTempFile);
 		}
+		virtual u32 FileTypeCommitSize(StringView file) const
+		{
+			if (file.EndsWith(TCV(".pch")))
+				return 8 * 1024 * 1024;
+			if (file.EndsWith(TCV(".sarif")) || file.EndsWith(TCV(".json")) || file.EndsWith(TCV(".h")))
+				return 64 * 1024;
+			return Super::FileTypeCommitSize(file);
+		}
+
 		virtual bool IsThrowAway(StringView fileName, bool isRunningRemote) const override
 		{
 			return fileName.Contains(TCV("vctip_")) || Super::IsThrowAway(fileName, isRunningRemote);
@@ -103,12 +112,12 @@ namespace uba
 				;
 		}
 
-		virtual bool IsRarelyReadAfterWritten(const StringView& fileName) const override
+		virtual bool IsMaybeReadAfterWritten(const StringView& fileName) const override
 		{
-			return fileName.EndsWith(TCV(".dep.json"))
+			return !(fileName.EndsWith(TCV(".dep.json"))
 				|| fileName.EndsWith(TCV(".sarif"))
 				|| fileName.EndsWith(TCV(".exe"))
-				|| fileName.EndsWith(TCV(".dll"));
+				|| fileName.EndsWith(TCV(".dll")));
 		}
 
 		virtual bool ShouldExtractSymbols(const StringView& fileName) const override
@@ -194,11 +203,12 @@ namespace uba
 			return Super::AllowStorageProxy(file);
 		}
 
-		virtual bool IsRarelyReadAfterWritten(const StringView& fileName) const override
+		virtual bool IsMaybeReadAfterWritten(const StringView& fileName) const override
 		{
-			return fileName.EndsWith(TCV(".pdb"))
+			return !(fileName.EndsWith(TCV(".pdb"))
 				|| fileName.EndsWith(TCV(".exe"))
-				|| fileName.EndsWith(TCV(".dll"));
+				//|| fileName.EndsWith(TCV(".exp"))  // This is used by clang
+				|| fileName.EndsWith(TCV(".dll")));
 		}
 
 		virtual bool CanDependOnCompressedFiles() override
@@ -342,9 +352,9 @@ namespace uba
 				|| (file.EndsWith(TCV(".rsp")) && !file.EndsWith(TCV("Shared.rsp")));
 		}
 
-		virtual bool IsRarelyReadAfterWritten(const StringView& fileName) const override
+		virtual bool IsMaybeReadAfterWritten(const StringView& fileName) const override
 		{
-			return fileName.EndsWith(TCV(".d")) || fileName.EndsWith(TCV(".d.tmp"));
+			return !(fileName.EndsWith(TCV(".d")) || fileName.EndsWith(TCV(".d.tmp")));
 		}
 
 		virtual bool ShouldExtractSymbols(const StringView& fileName) const override
diff --git a/Engine/Source/Programs/UnrealBuildAccelerator/Core/Private/UbaSharedMemoryView.cpp b/Engine/Source/Programs/UnrealBuildAccelerator/Core/Private/UbaSharedMemoryView.cpp
--- a/Engine/Source/Programs/UnrealBuildAccelerator/Core/Private/UbaSharedMemoryView.cpp
+++ b/Engine/Source/Programs/UnrealBuildAccelerator/Core/Private/UbaSharedMemoryView.cpp
@@ -9,15 +9,6 @@
 #if PLATFORM_WINDOWS
 #include <windows.h>
 #pragma comment(lib, "onecore.lib")
-#else
-#define GetCurrentProcess() 0
-#define UnmapViewOfFile2(...) false
-#define VirtualFree(...) false
-#define VirtualAlloc2(...) nullptr
-#define MapViewOfFile3(...) nullptr
-#define PAGE_READONLY 0x02
-#define PAGE_READWRITE 0x04
-#define PAGE_WRITECOPY 0x08
 #endif
 
 namespace uba
@@ -40,11 +31,9 @@ namespace uba
 			m_memory = sharedMemory + sliceOffset;
 			return true;
 		}
-
 		if (!InternalInit(handle, capacity, nullptr, hint))
 			return false;
-		AddRequestedMemory(handle, reader, sliceCount, SharedMemoryMapType_ReadOnly, 0, hint, 0);
-		return true;
+		return AddRequestedMemory(handle, reader, sliceCount, SharedMemoryMapType_ReadOnly, 0, hint, 0);
 	}
 
 	bool SharedMemoryView::AddRequestedMemory(SharedMemoryAllocatorHandle handle, BinaryReader& reader, SharedMemoryMapType mapType, u64 offset, const tchar* hint, u8 init)
@@ -66,26 +55,28 @@ namespace uba
 			return;
 		}
 
+		#if PLATFORM_WINDOWS
 		bool res; (void)res;
 
 		if (m_mappedSize)
 		{
-			auto& stats = KernelStats::GetCurrent().unmapViewOfFile2;
+			auto& kernelStats = KernelStats::GetCurrent();
 			{
-				TimerScope ts(stats);
+				TimerScope ts(kernelStats.unmapViewOfFile2);
 				res = UnmapViewOfFile2(GetCurrentProcess(), m_memory, MEM_PRESERVE_PLACEHOLDER);
 				UBA_ASSERT(res);
 			}
 
 			for (u64 mapping : m_additionalMappings)
 			{
-				TimerScope ts(stats);
+				TimerScope ts(kernelStats.unmapViewOfFile2);
 				res = UnmapViewOfFile2(GetCurrentProcess(), m_memory + mapping, MEM_PRESERVE_PLACEHOLDER);
 				UBA_ASSERT(res);
 			}
 
 			if (m_mappedSize < m_capacity || !m_additionalMappings.empty())
 			{
+				TimerScope ts(kernelStats.virtualFree);
 				res = VirtualFree(m_memory, m_capacity, MEM_RELEASE | MEM_COALESCE_PLACEHOLDERS);
 				UBA_ASSERTF(res, TC("VirtualFree failed (%u)"), GetLastError());
 			}
@@ -96,9 +87,22 @@ namespace uba
 		}
 		res = VirtualFree(m_memory, 0, MEM_RELEASE);
 		UBA_ASSERTF(res, TC("VirtualFree failed (%u)"), GetLastError());
+		#else
+		munmap(m_memory, m_capacity); 
+		#endif
+
+		m_memory = nullptr;
+		m_capacity = 0;
+	}
 
+	u8* SharedMemoryView::DetachMemory()
+	{
+		u8* mem = m_memory;
 		m_memory = nullptr;
 		m_capacity = 0;
+		m_mappedSize = 0;
+		m_additionalMappings.clear();
+		return mem;
 	}
 
 	SharedMemoryView::~SharedMemoryView()
@@ -108,8 +112,16 @@ namespace uba
 
 	bool SharedMemoryView::InternalInit(SharedMemoryAllocatorHandle handle, u64 capacity, void* baseAddress, const tchar* hint)
 	{
-		m_capacity = AlignUp(capacity, 64 * 1024);
+		m_capacity = AlignUp(capacity, PageSize);
+
+		TimerScope ts(KernelStats::GetCurrent().virtualAlloc2);
+
+		#if PLATFORM_WINDOWS
 		m_memory = (u8*)VirtualAlloc2(GetCurrentProcess(), baseAddress, m_capacity, MEM_RESERVE | MEM_RESERVE_PLACEHOLDER, PAGE_NOACCESS, nullptr, 0);
+		#else
+		m_memory = (u8*)mmap(NULL, m_capacity, PROT_NONE, MAP_PRIVATE | MAP_ANONYMOUS, -1, 0);
+		#endif
+
 		UBA_ASSERTF(m_memory, TC("VirtualAlloc2 failed (%s)"), LastErrorToText().data);
 		return m_memory != 0;
 	}
@@ -151,11 +163,11 @@ namespace uba
 
 	bool SharedMemoryView::ExtendMemory(SharedMemoryAllocatorHandle handle, u64 offset, u64 size, SharedMemoryMapType mapType, const tchar* hint, u8 init)
 	{
-		UBA_ASSERTF(handle.internalHandle, TC("Allocator handle not initialized (%s)"), hint);
+		UBA_ASSERTF(handle.IsValid(), TC("Allocator handle not initialized (%s)"), hint);
 		UBA_ASSERTF(m_memory, TC("Init failed?"));
 		UBA_ASSERTF(m_capacity != ~0ull, TC("Can't extend memory on a view directly mapped on to allocator mapping (%s)"), hint);
-		UBA_ASSERTF(AlignUp(offset, 64 * 1024) == offset, TC("ExtendMemory got offset %llu parameter which is not aligned (%s)"), offset, hint);
-		UBA_ASSERTF(AlignUp(size, 64 * 1024) == size, TC("ExtendMemory got size %llu which is not aligned (%s)"), size, hint);
+		UBA_ASSERTF(AlignUp(offset, PageSize) == offset, TC("ExtendMemory got offset %llu parameter which is not aligned (%s)"), offset, hint);
+		UBA_ASSERTF(AlignUp(size, PageSize) == size, TC("ExtendMemory got size %llu which is not aligned (%s)"), size, hint);
 
 		u8* targetAddress = m_memory + m_mappedSize;
 
@@ -164,25 +176,36 @@ namespace uba
 		if (left < size)
 			return false;
 
+		auto& kernelStats = KernelStats::GetCurrent();
+
+		#if PLATFORM_WINDOWS
+		u32 protect = PAGE_READWRITE;
+		if (mapType == SharedMemoryMapType_ReadOnly)
+			protect = PAGE_READONLY;
+		else if (mapType == SharedMemoryMapType_CopyOnWrite)
+			protect = PAGE_WRITECOPY;
+
 		if (u64 toRelease = m_capacity - m_mappedSize - size)
 		{
+			TimerScope ts(kernelStats.virtualFree);
 			bool res = VirtualFree(targetAddress + size, toRelease, MEM_RELEASE | MEM_PRESERVE_PLACEHOLDER);
 			UBA_ASSERTF(res, TC("VirtualFree for placeholder failed releasing %llu at %llu (Capacity: %llu MappedSize: %llu) - %s (%s)"), toRelease, size, m_capacity, m_mappedSize, LastErrorToText().data, hint);
 			if (!res)
 				return false;
 		}
 
-		u32 protect = PAGE_READWRITE;
-		if (mapType == SharedMemoryMapType_ReadOnly)
-			protect = PAGE_READONLY;
-		else if (mapType == SharedMemoryMapType_CopyOnWrite)
-			protect = PAGE_WRITECOPY;
-
-		auto& stats = KernelStats::GetCurrent().mapViewOfFile3;
-		stats.bytes += size;
-		TimerScope ts(stats);
+		kernelStats.mapViewOfFile3.bytes += size;
+		TimerScope ts(kernelStats.mapViewOfFile3);
+		void* viewA = MapViewOfFile3(handle.mh, GetCurrentProcess(), targetAddress, offset, size, MEM_REPLACE_PLACEHOLDER, protect, nullptr, 0);
+		#else
+		UBA_ASSERT(mapType != SharedMemoryMapType_CopyOnWrite);
+		u32 protect = PROT_READ;
+		if (mapType == SharedMemoryMapType_ReadWrite)
+			protect |= PROT_WRITE;
+		TimerScope ts(kernelStats.mapViewOfFile3);
+		void* viewA = mmap(targetAddress, size, protect, MAP_SHARED | MAP_FIXED, handle.shmFd, offset);
+		#endif
 
-		void* viewA = MapViewOfFile3((HANDLE)handle.internalHandle, GetCurrentProcess(), targetAddress, offset, size, MEM_REPLACE_PLACEHOLDER, protect, nullptr, 0);
 		UBA_ASSERTF(viewA, TC("MapViewOfFile3 failed - %s (%s)"), LastErrorToText().data, hint);
 		if (!viewA)
 			return false;
@@ -199,4 +222,18 @@ namespace uba
 	}
 
 	////////////////////////////////////////////////////////////////////////////////////////////////////
+
+	bool IsMemoryHandle(const tchar* name)
+	{
+		UBA_ASSERT(name && *name);
+		return *name == MemoryHandleChar || *name == FreeableMemoryHandleChar;
+	}
+
+	bool IsFreeableMemoryHandle(const tchar* name)
+	{
+		UBA_ASSERT(name && *name);
+		return *name == FreeableMemoryHandleChar;
+	}
+
+	////////////////////////////////////////////////////////////////////////////////////////////////////
 }
diff --git a/Engine/Source/Programs/UnrealBuildAccelerator/Core/Public/UbaApplicationRules.h b/Engine/Source/Programs/UnrealBuildAccelerator/Core/Public/UbaApplicationRules.h
--- a/Engine/Source/Programs/UnrealBuildAccelerator/Core/Public/UbaApplicationRules.h
+++ b/Engine/Source/Programs/UnrealBuildAccelerator/Core/Public/UbaApplicationRules.h
@@ -63,6 +63,12 @@ namespace uba
 			return 8ull * 1024 * 1024 * 1024;
 		}
 
+		// When writing this type of file, how big commits should shared memory system take. We want as few as possible without wasting memory
+		virtual u32 FileTypeCommitSize(StringView file) const
+		{
+			return 1 * 1024 * 1024;
+		}
+
 		// Outputfile means that it is kept in memory and then sent back to session process which can decide to write it to disk or send it over network
 		virtual bool IsOutputFile(StringView fileName, StringView systemTemp) const
 		{
@@ -96,9 +102,9 @@ namespace uba
 
 		// Return true if the file is never/rarely read after it was written.
 		// This is an optimization where the written file is not kept in file mappings after written
-		virtual bool IsRarelyReadAfterWritten(const StringView& fileName) const
+		virtual bool IsMaybeReadAfterWritten(const StringView& fileName) const
 		{
-			return true;
+			return false;
 		}
 
 		virtual bool AllowStorageProxy(const StringView& file) const
diff --git a/Engine/Source/Programs/UnrealBuildAccelerator/Core/Public/UbaDefinitions.h b/Engine/Source/Programs/UnrealBuildAccelerator/Core/Public/UbaDefinitions.h
--- a/Engine/Source/Programs/UnrealBuildAccelerator/Core/Public/UbaDefinitions.h
+++ b/Engine/Source/Programs/UnrealBuildAccelerator/Core/Public/UbaDefinitions.h
@@ -5,17 +5,15 @@
 // Use native mach semaphores for IPC.
 #define UBA_USE_NATIVE_MAC_SEMAPHORES 0//PLATFORM_MAC
 
-// Experimental. Testing to use file mappings that are backed by files to reduce page file pressure
-#define UBA_ENABLE_ON_DISK_FILE_MAPPINGS 0
-
 // Disable pthread cancellation to reduce work when waking up from event.
 // No uba code should rely on cancelling threads and shutdown is gracefully so this can be disabled
 #define UBA_DISABLE_PTHREAD_CANCELLATIONS 1
 
-
 // Enable this to test semaphore implementation for local events
 #define UBA_TEST_WINDOWS_SEMAPHORES 0
 
+// Enable to add logging of file mappings
+#define UBA_DEBUG_FILE_MAPPING 0
 
 namespace uba
 {
diff --git a/Engine/Source/Programs/UnrealBuildAccelerator/Core/Public/UbaProtocol.h b/Engine/Source/Programs/UnrealBuildAccelerator/Core/Public/UbaProtocol.h
--- a/Engine/Source/Programs/UnrealBuildAccelerator/Core/Public/UbaProtocol.h
+++ b/Engine/Source/Programs/UnrealBuildAccelerator/Core/Public/UbaProtocol.h
@@ -38,6 +38,9 @@ namespace uba
 		UBA_PROCESS_MESSAGE(ResolveCallstack) \
 		UBA_PROCESS_MESSAGE(CheckRemapping) \
 		UBA_PROCESS_MESSAGE(RunSpecialProgram) \
+		UBA_PROCESS_MESSAGE(CreateSharedMemory) \
+		UBA_PROCESS_MESSAGE(CommitSharedMemory) \
+		UBA_PROCESS_MESSAGE(GetSharedMemory) \
 
 	enum MessageType : u8
 	{
diff --git a/Engine/Source/Programs/UnrealBuildAccelerator/Core/Public/UbaSharedMemoryView.h b/Engine/Source/Programs/UnrealBuildAccelerator/Core/Public/UbaSharedMemoryView.h
--- a/Engine/Source/Programs/UnrealBuildAccelerator/Core/Public/UbaSharedMemoryView.h
+++ b/Engine/Source/Programs/UnrealBuildAccelerator/Core/Public/UbaSharedMemoryView.h
@@ -3,6 +3,7 @@
 #pragma once
 
 #include "UbaMemory.h"
+#include "UbaPlatform.h"
 
 namespace uba
 {
@@ -10,7 +11,8 @@ namespace uba
 
 	////////////////////////////////////////////////////////////////////////////////////////////////////
 
-	struct SharedMemoryAllocatorHandle { u64 internalHandle = 0; };
+	using SharedMemoryAllocatorHandle = FileMappingHandle;
+	
 	struct SharedMemoryHandle
 	{
 		inline bool IsValid() const { return internalHandle != 0; }
@@ -40,6 +42,7 @@ namespace uba
 		u64 GetMappedSize() { return m_mappedSize; }
 		u64 GetCapacity() { return m_capacity; }
 
+		u8* DetachMemory();
 		~SharedMemoryView();
 
 	private:
@@ -60,4 +63,14 @@ namespace uba
 	};
 
 	////////////////////////////////////////////////////////////////////////////////////////////////////
+
+	bool IsMemoryHandle(const tchar* name);
+	bool IsFreeableMemoryHandle(const tchar* name);
+	inline constexpr tchar MemoryHandleChar = '^';
+	inline constexpr tchar FreeableMemoryHandleChar = '>';
+	inline constexpr tchar WrittenMemoryHandleChar = ':';
+
+	////////////////////////////////////////////////////////////////////////////////////////////////////
 }
+
+template<> struct std::hash<uba::SharedMemoryHandle> { size_t operator()(const uba::SharedMemoryHandle& h) const { return h.internalHandle; } };
diff --git a/Engine/Source/Programs/UnrealBuildAccelerator/Detours/Private/PosixOS/UbaDetoursFunctionsPosix.cpp b/Engine/Source/Programs/UnrealBuildAccelerator/Detours/Private/PosixOS/UbaDetoursFunctionsPosix.cpp
--- a/Engine/Source/Programs/UnrealBuildAccelerator/Detours/Private/PosixOS/UbaDetoursFunctionsPosix.cpp
+++ b/Engine/Source/Programs/UnrealBuildAccelerator/Detours/Private/PosixOS/UbaDetoursFunctionsPosix.cpp
@@ -488,7 +488,7 @@ int Shared_open(const char* funcName, const char* file, int flags, int mode, con
 		return -1;
 	}
 
-	if (realFileName[0] == '^')
+	if (IsMemoryHandle(realFileName))
 	{
 		DEBUG_LOG_DETOURED(funcName, "FAILED %s (%s)", fileName.data, realFileName);
 		UBA_ASSERTF(false, "^ filenames not implemented");
@@ -600,8 +600,8 @@ void Shared_close(int fd, const TrueClose& trueClose)
 	if (--fo->refCount)
 		return;
 
-	FileMappingHandle mappingHandle;
-	u64 mappingWritten = 0;
+	SharedMemoryHandle memoryHandle;
+	u64 memoryWritten = 0;
 	FileInfo& fi = *fo->fileInfo;
 	const tchar* path = fi.name;
 
@@ -612,7 +612,7 @@ void Shared_close(int fd, const TrueClose& trueClose)
 		closeGuard.Execute();
 		//for (auto& kv : g_fileHandles)
 		//	UBA_ASSERT(&fi != kv.second.fileObject->fileInfo);
-		Rpc_UpdateCloseHandle(path, fo->closeId, fo->deleteOnClose, fo->newName.c_str(), mappingHandle, mappingWritten, true);
+		Rpc_UpdateCloseHandle(path, fo->closeId, fo->deleteOnClose, fo->newName.c_str(), memoryHandle, memoryWritten, true);
 	}
 	delete fo;
 }
@@ -2863,6 +2863,8 @@ namespace uba
 			reader.ReadString(g_virtualApplication);
 			reader.ReadString(g_virtualWorkingDir);
 
+			reader.ReadU64(); // shared memory allocator handle
+			reader.ReadU64(); // permanent files memory handle
 			directoryTableHandle = reader.ReadU64();
 			directoryTableSize = reader.ReadU32();
 			directoryTableCount = reader.ReadU32();
@@ -2951,11 +2953,11 @@ namespace uba
 				if (!fo->closeId)
 					continue;
 				TRUE_WRAPPER(close)(kv.first);
-				FileMappingHandle mappingHandle;
-				u64 mappingWritten = 0;
+				SharedMemoryHandle memoryHandle;
+				u64 memoryWritten = 0;
 				FileInfo& fi = *fo->fileInfo;
 				const tchar* path = fi.name;
-				Rpc_UpdateCloseHandle(path, fo->closeId, fo->deleteOnClose, fo->newName.c_str(), mappingHandle, mappingWritten, true);
+				Rpc_UpdateCloseHandle(path, fo->closeId, fo->deleteOnClose, fo->newName.c_str(), memoryHandle, memoryWritten, true);
 			}
 		}
 		
@@ -3028,9 +3030,11 @@ namespace uba
 		}
 		else
 		{
+			#if UBA_DEBUG_LOG_ENABLED
 			g_debugFile = (FileHandle)open("/dev/tty", O_WRONLY);
 			sb.Append('\n');
 			WriteDebug(sb.data, sb.count);
+			#endif
 		}
 
 		CloseCom();
diff --git a/Engine/Source/Programs/UnrealBuildAccelerator/Detours/Private/UbaDetoursFileMappingTable.cpp b/Engine/Source/Programs/UnrealBuildAccelerator/Detours/Private/UbaDetoursFileMappingTable.cpp
--- a/Engine/Source/Programs/UnrealBuildAccelerator/Detours/Private/UbaDetoursFileMappingTable.cpp
+++ b/Engine/Source/Programs/UnrealBuildAccelerator/Detours/Private/UbaDetoursFileMappingTable.cpp
@@ -40,7 +40,7 @@ namespace uba
 			FileInfo& info = insres.first->second;
 			if (!insres.second)
 			{
-				if (info.name && info.name[0] == '^' && !mappedFileName.Equals(info.name, false)) // Mapped file has been re-mapped.
+				if (info.name && IsMemoryHandle(info.name) && !mappedFileName.Equals(info.name, false)) // Mapped file has been re-mapped.
 				{
 					UBA_ASSERTF(!info.memoryFile, TC("Mapped file %s has changed mapping (%s to %s) while being in use"), info.originalName, info.name, mappedFileName.data);
 					info.name = m_memoryBlock.Strdup(mappedFileName).data;
@@ -129,7 +129,7 @@ namespace uba
 		return tableOffset;
 	}
 
-	void Rpc_UpdateCloseHandle(const tchar* handleName, u32 closeId, bool deleteOnClose, const tchar* newName, const FileMappingHandle& mappingHandle, u64 mappingWritten, bool success)
+	void Rpc_UpdateCloseHandle(const tchar* handleName, u32 closeId, bool deleteOnClose, const tchar* newName, SharedMemoryHandle memoryHandle, u64 memoryWritten, bool success)
 	{
 		u32 directoryTableSize;
 		{
@@ -139,8 +139,8 @@ namespace uba
 			//writer.WriteU32(attributes); // TODO THIS
 			writer.WriteBool(deleteOnClose);
 			writer.WriteBool(success);
-			writer.WriteU64(mappingHandle.ToU64());
-			writer.WriteU64(mappingWritten);
+			writer.WriteU64(memoryHandle.ToU64());
+			writer.WriteU64(memoryWritten);
 			if (*newName)
 			{
 				StringBuffer<> fixedName;
@@ -185,14 +185,14 @@ namespace uba
 			StringBufferBase& backedName = originalName.Clear();
 			reader.ReadString(backedName);
 
-			FileMappingHandle mappingHandle = FileMappingHandle::FromU64(reader.Read7BitEncoded());
+			SharedMemoryHandle memoryHandle = SharedMemoryHandle::FromU64(reader.Read7BitEncoded());
 			u64 fileSize = reader.Read7BitEncoded();
 			info.fileNameKey = key;
 			info.size = fileSize;
 			info.created = true;
 
-			if (mappingHandle.IsValid())
-				backedName.Clear().Append(':').AppendBase62(mappingHandle.ToU64());
+			if (memoryHandle.IsValid())
+				backedName.Clear().Append(WrittenMemoryHandleChar).AppendBase62(memoryHandle.ToU64());
 
 			if (!info.name || !backedName.Equals(info.name))
 				info.name = g_mappedFileTable.m_memoryBlock.Strdup(backedName).data;
@@ -207,9 +207,13 @@ namespace uba
 				if (!mf->isLocalOnly)
 				{
 					DEBUG_LOG(TC("Unmapping old memoryfile for: %s"), info.originalName);
-					UnmapViewOfFile(mf->baseAddress);
-					CloseHandle(mf->mappingHandle.mh);
-					CloseHandle(mf->mappingHandle.fh);
+
+#if UBA_DEBUG
+					UBA_ASSERTF(false, TC("REPORT THIS TO UBA DEVS"));
+					//UnmapViewOfFile(mf->baseAddress);
+					//CloseHandle(mf->memoryHandle.mh);
+					//CloseHandle(mf->memoryHandle.fh);
+#endif
 				}
 				#endif
 				// delete mf; // Let them leak
@@ -264,6 +268,66 @@ namespace uba
 		Rpc_GetWrittenFilesNoLock(true);
 	}
 
+	SharedMemoryHandle Rpc_CreateSharedMemory(SharedMemoryView& out, SharedMemoryAllocatorHandle allocatorHandle, u64 capacity, u64 initialSize, const tchar* hint)
+	{
+		RPC_MESSAGE(CreateSharedMemory, createSharedMemory)
+		writer.WriteU64(initialSize);
+		writer.WriteString(hint);
+		writer.Flush();
+		BinaryReader reader;
+		SharedMemoryHandle handle{ reader.Read7BitEncoded() };
+
+		out.Init(allocatorHandle, capacity);
+		if (initialSize)
+			out.AddRequestedMemory(allocatorHandle, reader);
+		return handle;
+	}
+
+	void Rpc_CommitSharedMemory(SharedMemoryView& view, SharedMemoryAllocatorHandle allocatorHandle, SharedMemoryHandle handle, u64 size)
+	{
+		if (!handle.IsValid())
+			return;
+		RPC_MESSAGE(CommitSharedMemory, commitSharedMemory)
+		writer.WriteU64(handle.ToU64());
+		writer.WriteU64(size);
+		writer.Flush();
+		BinaryReader reader;
+		view.AddRequestedMemory(allocatorHandle, reader);
+	}
+
+	void Rpc_GetSharedMemory(SharedMemoryView& view, SharedMemoryAllocatorHandle allocatorHandle, u8* allocatorMem, SharedMemoryHandle handle, u64 offset, u64 size, const tchar* hint)
+	{
+		UBA_ASSERTF(AlignUp(offset, PageSize) == offset, TC("Offset %llu unaligned. Can't get unaligned shared memory (%s)"), offset, hint);
+		//UBA_ASSERTF(AlignUp(size, PageSize) == size, TC("Size %llu unaligned. Can't get unaligned shared memory (%s)"), size, hint);
+
+		RPC_MESSAGE(GetSharedMemory, getSharedMemory)
+		writer.WriteU64(handle.ToU64());
+		writer.WriteU64(offset);
+		writer.WriteU64(size);
+		writer.Flush();
+		BinaryReader reader;
+		u64 committed = size;
+		if (!committed)
+			committed = reader.Read7BitEncoded();
+		view.Init(allocatorHandle, allocatorMem, reader, committed, hint);
+		//UBA_ASSERTF(u64(view.GetMemory()) != 0xFFFFFFFFFFFFFFFFllu, TC("%llu %llu %llu"), offset, size, committed);
+		//view.AddRequestedMemory(allocatorHandle, reader, SharedMemoryMapType_ReadOnly);
+	}
+
+	void Rpc_UpdateSharedMemory(SharedMemoryView& view, SharedMemoryAllocatorHandle allocatorHandle, SharedMemoryHandle handle, u64 capacity)
+	{
+		RPC_MESSAGE(GetSharedMemory, getSharedMemory)
+		writer.WriteU64(handle.ToU64());
+		writer.WriteU64(0);
+		writer.WriteU64(0);
+		writer.Flush();
+		BinaryReader reader;
+		u64 committed = reader.Read7BitEncoded(); (void)committed;
+		UBA_ASSERT(committed <= capacity);
+		view.Init(allocatorHandle, capacity);
+		view.AddRequestedMemory(allocatorHandle, reader);
+	}
+
 	u32 Rpc_GetEntryOffset(const StringKey& entryNameKey, StringView entryName, bool checkIfDir)
 	{
 		u32 dirTableOffset = ~u32(0);
diff --git a/Engine/Source/Programs/UnrealBuildAccelerator/Detours/Private/UbaDetoursFileMappingTable.h b/Engine/Source/Programs/UnrealBuildAccelerator/Detours/Private/UbaDetoursFileMappingTable.h
--- a/Engine/Source/Programs/UnrealBuildAccelerator/Detours/Private/UbaDetoursFileMappingTable.h
+++ b/Engine/Source/Programs/UnrealBuildAccelerator/Detours/Private/UbaDetoursFileMappingTable.h
@@ -3,14 +3,16 @@
 #pragma once
 
 #include "UbaBinaryReaderWriter.h"
+#include "UbaDefinitions.h"
 #include "UbaDetoursShared.h"
+#include "UbaSharedMemoryView.h"
+
 
 namespace uba
 {
 	struct DetouredHandle;
 	struct MemoryFile;
 
-
 	struct FileInfo
 	{
 		// Size of file. Can be different than what Directory table says (For example for decompressed obj files)
@@ -49,7 +51,7 @@ namespace uba
 		bool created = false;
 
 		// File mapping handle coming from session process (Set in NtCreateFile)
-		void* trueFileMapHandle = 0;
+		SharedMemoryHandle trueFileMapHandle;
 
 		// Offset into trueFileMapHandle
 		u64 trueFileMapOffset = 0;
@@ -64,7 +66,6 @@ namespace uba
 		u32 refCount = 0;
 		u32 fileMapDesiredAccess = PAGE_READONLY;
 		u32 fileMapViewDesiredAccess = FILE_MAP_READ; // This is the parameter to MapViewOfFile
-		bool mappingChecked = false;
 		#endif
 	};
 
@@ -94,13 +95,18 @@ namespace uba
 	void Rpc_CreateFileW(const StringView& fileName, const StringKey& fileNameKey, u8 accessFlags, tchar* outNewName, u64 newNameCapacity, u64& outSize, u32& outCloseId, bool lock);
 	void Rpc_CheckRemapping(const StringView& fileName, const StringKey& fileNameKey);
 	u32  Rpc_UpdateDirectory(const StringKey& dirKey, const tchar* dirName, u64 dirNameLen, bool lockDirTable = true);
-	void Rpc_UpdateCloseHandle(const tchar* handleName, u32 closeId, bool deleteOnClose, const tchar* newName, const FileMappingHandle& mappingHandle, u64 mappingWritten, bool success);
+	void Rpc_UpdateCloseHandle(const tchar* handleName, u32 closeId, bool deleteOnClose, const tchar* newName, SharedMemoryHandle memoryHandle, u64 memoryWritten, bool success);
 	void Rpc_UpdateTables();
 	u32  Rpc_GetEntryOffset(const StringKey& entryNameKey, StringView entryName, bool checkIfDir = false);
 	void Rpc_GetFullFileName(const tchar*& path, u64& pathLen, StringBufferBase& tempBuf, bool useVirtualName, const tchar* const* loaderPaths = nullptr);
 	void Rpc_GetFullFileName2(const tchar* path, StringBufferBase& outReal, StringBufferBase& outVirtual, const tchar* const* loaderPaths = nullptr);
 	void Rpc_GetWrittenFiles();
 
+	SharedMemoryHandle Rpc_CreateSharedMemory(SharedMemoryView& out, SharedMemoryAllocatorHandle allocatorHandle, u64 capacity, u64 initialSize, const tchar* hint);
+	void Rpc_CommitSharedMemory(SharedMemoryView& view, SharedMemoryAllocatorHandle allocatorHandle, SharedMemoryHandle handle, u64 size);
+	void Rpc_GetSharedMemory(SharedMemoryView& view, SharedMemoryAllocatorHandle allocatorHandle, u8* allocatorMem, SharedMemoryHandle handle, u64 offset, u64 size, const tchar* hint);
+	void Rpc_UpdateSharedMemory(SharedMemoryView& view, SharedMemoryAllocatorHandle allocatorHandle, SharedMemoryHandle handle, u64 capacity);
+
 	struct DirHash
 	{
 		DirHash(const StringView& str);
diff --git a/Engine/Source/Programs/UnrealBuildAccelerator/Detours/Private/Windows/UbaDetoursFunctionsKernelBase.inl b/Engine/Source/Programs/UnrealBuildAccelerator/Detours/Private/Windows/UbaDetoursFunctionsKernelBase.inl
--- a/Engine/Source/Programs/UnrealBuildAccelerator/Detours/Private/Windows/UbaDetoursFunctionsKernelBase.inl
+++ b/Engine/Source/Programs/UnrealBuildAccelerator/Detours/Private/Windows/UbaDetoursFunctionsKernelBase.inl
@@ -1,5 +1,7 @@
 // Copyright Epic Games, Inc. All Rights Reserved.
 
+#define UBA_ASSERT_ON_DETOURED_STD_PIPE 0 // Errors on clang on llvm-symbolizer.exe.. needs investigation
+
 DWORD Local_GetLongPathNameW(LPCWSTR lpszShortPath, LPWSTR lpszLongPath, DWORD cchBuffer)
 {
 	SCOPED_WRITE_LOCK(g_longPathNameCacheLock, lock);
@@ -713,7 +715,7 @@ BOOL Detoured_WriteFile(HANDLE hFile, LPCVOID lpBuffer, DWORD nNumberOfBytesToWr
 
 				__try
 				{
-					mf->Write(dh, lpBuffer, nNumberOfBytesToWrite);
+					mf->Write(dh, lpBuffer, nNumberOfBytesToWrite, fi.name);
 				}
 				__except(WriteFileAccessViolation(GetExceptionInformation()))
 				{
@@ -2294,7 +2296,7 @@ BOOL Detoured_GetFileInformationByHandle(HANDLE hFile, LPBY_HANDLE_FILE_INFORMAT
 					return false;
 				}
 
-				UBA_ASSERT(attr.fileIndex);
+				UBA_ASSERTF(attr.fileIndex, TC("No file index set for %s"), HandleToName(hFile));
 
 				LARGE_INTEGER li = ToLargeInteger(attr.fileIndex);
 				/*
@@ -2315,7 +2317,7 @@ BOOL Detoured_GetFileInformationByHandle(HANDLE hFile, LPBY_HANDLE_FILE_INFORMAT
 				*/
 				SetLastError(ERROR_SUCCESS);
 
-				UBA_ASSERT(attr.volumeSerial);
+				UBA_ASSERTF(attr.volumeSerial, TC("No volume serial set for %s"), HandleToName(hFile));
 
 				lpFileInformation->dwFileAttributes = attr.data.dwFileAttributes;
 				lpFileInformation->ftCreationTime = attr.data.ftCreationTime;
@@ -2532,49 +2534,22 @@ LPVOID Detoured_MapViewOfFileEx(HANDLE hFileMappingObject, DWORD dwDesiredAccess
 			FileInfo& fi = *dh.fileObject->fileInfo;
 			if (fi.fileMapMem && lpBaseAddress && lpBaseAddress != fi.fileMapMem) // This scenario happens with pch files in msvc cl.exe
 			{
-				if (dh.trueHandle == INVALID_HANDLE_VALUE) // This means we are using the "built-in" mapping handle
-					trueMappingObject = fi.trueFileMapHandle;
-
-				u8* res;
-
-				// This actually uses more memory because "offset" can be really high
-				//if (res = (u8*)True_MapViewOfFileEx(trueMappingObject, dwDesiredAccess, ToHigh(fi.trueFileMapOffset), ToLow(fi.trueFileMapOffset), dwNumberOfBytesToMap + offset, (u8*)lpBaseAddress - offset))
-				//{
-				//	// If we don't have any real MapViewOfFile on existing fileMapMem we can just drop the old one and use the new one.
-				//	SCOPED_WRITE_LOCK(g_mappedFileTable.m_memLookupLock, lock);
-				//	auto it = g_mappedFileTable.m_memLookup.find(fi.fileMapMem);
-				//	if (it == g_mappedFileTable.m_memLookup.end())
-				//	{
-				//		if (!True_UnmapViewOfFile(fi.fileMapMem))
-				//		{
-				//			UBA_ASSERT(false);
-				//		}
-				//		fi.fileMapMem = (u8*)res;
-				//		fi.fileMapMemSize = dwNumberOfBytesToMap;
-				//		g_mappedFileTable.m_memLookup[res + offset] = 1;
-				//	}
-				//	res += offset;
-				//}
-				//else
-				{
-					offset += fi.trueFileMapOffset;
-
-					// We have retry here because this is typically where oom happen. Some of these mappings are 1gb pch files
-					u32 counter = 0;
-					do
-					{
-						TimerScope ts(g_kernelStats.mapViewOfFile);
-						res = (u8*)True_MapViewOfFileEx(trueMappingObject, dwDesiredAccess, ToHigh(offset), ToLow(offset), dwNumberOfBytesToMap, lpBaseAddress);
-						if (res)
-							break;
-						DWORD error = GetLastError();
-						if (error != ERROR_NOT_ENOUGH_MEMORY && error != ERROR_COMMITMENT_LIMIT)
-							break;
-						Rpc_AllocFailed(L"MapViewOfFile", error);
-						++counter;
-					} while (counter <= 10);
-				}
-
+				// This is really hacky.. cl.exe keeps the old mapping and creates a new one with an offset which is set copy-on-write..
+				// This is where cl.exe "continues" its compile from the pch
+				UBA_ASSERT(dwDesiredAccess == FILE_MAP_COPY);
+
+				SharedMemoryView& view = *new SharedMemoryView();
+				RPC_MESSAGE(GetSharedMemory, getSharedMemory)
+				writer.WriteU64(fi.trueFileMapHandle.ToU64());
+				writer.WriteU64(offset + fi.trueFileMapOffset);
+				writer.WriteU64(dwNumberOfBytesToMap);
+				writer.Flush();
+				BinaryReader reader;
+				view.Init(g_sharedMemoryAllocator, dwNumberOfBytesToMap, lpBaseAddress);
+				view.AddRequestedMemory(g_sharedMemoryAllocator, reader, SharedMemoryMapType_CopyOnWrite, 0);
+				SetLastError(ERROR_SUCCESS);
+				//UBA_ASSERT(!fi.trueFileMapOffset);
+				void* res = view.GetMemory();
 				DEBUG_LOG_TRUE(L"MapViewOfFileEx", L"(INTERNAL) New FileObject for different base address %llu (%s) -> 0x%llx", uintptr_t(hFileMappingObject), HandleToName(hFileMappingObject), uintptr_t(res));
 				return res;
 			}
@@ -3234,7 +3209,10 @@ BOOL Detoured_CreateProcessW(LPCWSTR lpApplicationName, LPWSTR lpCommandLine, LP
 			if (dh.trueHandle != INVALID_HANDLE_VALUE)
 				return;
 
+			#if UBA_ASSERT_ON_DETOURED_STD_PIPE
 			auto assertGuard = MakeGuard([&]() { UBA_ASSERTF(false, L"%s is detoured and there is no proper handle (%s)", name, lpApplicationName); });
+			#endif
+
 			if (dh.type != HandleType_File)
 				return;
 			if (!dh.fileObject)
@@ -3248,8 +3226,11 @@ BOOL Detoured_CreateProcessW(LPCWSTR lpApplicationName, LPWSTR lpCommandLine, LP
 				return;
 			if (mf->isLocalOnly)
 				return;
+
 			// TODO: This is a memory file
+			#if UBA_ASSERT_ON_DETOURED_STD_PIPE
 			assertGuard.Cancel();
+			#endif
 		};
 
 
@@ -3266,7 +3247,10 @@ BOOL Detoured_CreateProcessW(LPCWSTR lpApplicationName, LPWSTR lpCommandLine, LP
 			lpStartupInfo->hStdInput = g_isDetachedProcess ? 0 : True_GetStdHandle(STD_ERROR_HANDLE);
 		else
 		{
+			#if UBA_ASSERT_ON_DETOURED_STD_PIPE
 			UBA_ASSERTF(false, L"hStdInput is detoured (%s)", lpApplicationName);
+			#endif
+			
 			lpStartupInfo->hStdInput = dh.trueHandle;
 		}
 	}
diff --git a/Engine/Source/Programs/UnrealBuildAccelerator/Detours/Private/Windows/UbaDetoursFunctionsNtDll.inl b/Engine/Source/Programs/UnrealBuildAccelerator/Detours/Private/Windows/UbaDetoursFunctionsNtDll.inl
--- a/Engine/Source/Programs/UnrealBuildAccelerator/Detours/Private/Windows/UbaDetoursFunctionsNtDll.inl
+++ b/Engine/Source/Programs/UnrealBuildAccelerator/Detours/Private/Windows/UbaDetoursFunctionsNtDll.inl
@@ -957,7 +957,7 @@ NTSTATUS NTAPI Shared_NtCreateFile(bool IsCreateFunc, PHANDLE hFileHandle, ACCES
 
 				bool isWriteAttributes = (DesiredAccess & FILE_WRITE_ATTRIBUTES) != 0;
 
-				// If file is an output file we still allllow this path and accept wrong (compressed) size.
+				// If file is an output file we still allow this path and accept wrong (compressed) size.
 				// This is a bit hacky but we don't want to transfer and decompress file just to get the size
 				if (allowEarlyOut && !useContent && !isWriteAttributes && (!CouldBeCompressedFile(fileName) || g_rules->IsOutputFile(fileName, g_systemTemp)))
 				{
@@ -1125,11 +1125,8 @@ NTSTATUS NTAPI Shared_NtCreateFile(bool IsCreateFunc, PHANDLE hFileHandle, ACCES
 		}
 		else
 		{
-			if (!info.mappingChecked && info.name[0] == '^' && !g_runningRemote && CouldBeCompressedFile(fileName))
-			{
+			if (IsFreeableMemoryHandle(info.name))
 				Rpc_CheckRemapping(fileName, fileNameKey);
-				info.mappingChecked = true;
-			}
 			lpFileName = info.name;
 		}
 	}
@@ -1195,73 +1192,109 @@ NTSTATUS NTAPI Shared_NtCreateFile(bool IsCreateFunc, PHANDLE hFileHandle, ACCES
 		return STATUS_SUCCESS;
 	}
 
-	if (lpFileName[0] == '^') // It is a HANDLE from session process
+	auto GetHandleAndOffset = [&](u64& outHandle, u64& outOffset)
+		{
+			const wchar_t* handleStr = lpFileName + 1;
+			const wchar_t* handleStrEnd = wcschr(handleStr, '-');
+			outOffset = 0;
+			if (!handleStrEnd)
+			{
+				handleStrEnd = handleStr + wcslen(handleStr);
+			}
+			else
+			{
+				const wchar_t* mappingOffsetStr = handleStrEnd + 1;
+				outOffset = StringToValueBase62(mappingOffsetStr, wcslen(mappingOffsetStr));
+			}
+			outHandle = StringToValueBase62(handleStr, handleStrEnd - handleStr);
+		};
+
+	if (IsMemoryHandle(lpFileName)) // It is a HANDLE from session process
 	{
-		if (isWrite)
+		u64 mappingHandle;
+		u64 mappingOffset;
+		GetHandleAndOffset(mappingHandle, mappingOffset);
+
+		if (mappingHandle & 0x8000'0000'0000'0000)
 		{
-			Rpc_WriteLogf(TC("Mapped file %s cant be open for write (%s)"), lpFileName, info.originalName);
-			UBA_ASSERTF(false, TC("Mapped file %s cant be open for write (%s)"), lpFileName, info.originalName);
-			return STATUS_ACCESS_DENIED;
+			mappingHandle &= ~0x8000'0000'0000'0000;
+
+			HANDLE fileMapHandle;
+			if (!True_DuplicateHandle(g_hostProcess, (HANDLE)mappingHandle, GetCurrentProcess(), &fileMapHandle, 0, FALSE, DUPLICATE_SAME_ACCESS))
+			{
+				Rpc_WriteLogf(L"Can't duplicate handle 0x%llx (%ls) for file %ls (Error %u)", uintptr_t(mappingHandle), lpFileName, info.originalName, GetLastError());
+				UBA_ASSERTF(fileMapHandle, L"Can't duplicate handle 0x%llx (%ls) for file %ls (Error %u)", uintptr_t(mappingHandle), lpFileName, info.originalName, GetLastError());
+				return STATUS_ACCESS_DENIED;
+			}
+
+			u64 alignedOffsetStart = AlignUp(mappingOffset - (PageSize - 1), PageSize);
+			u64 mapSize = (mappingOffset - alignedOffsetStart) + info.size;
+
+			LARGE_INTEGER li = ToLargeInteger(alignedOffsetStart);
+			u8* mem = (u8*)True_MapViewOfFile(fileMapHandle, FILE_MAP_READ, li.HighPart, li.LowPart, mapSize); // We will leak the mapping until process end
+			UBA_ASSERTF(mem, TC("MapViewOfFile failed for %s (%s)"), info.originalName, LastErrorToText().data);
+
+			info.fileMapMem = mem + mappingOffset - alignedOffsetStart;
+			info.fileMapMemSize = info.size;
+			info.isFileMap = true;
+
+			CloseHandle(fileMapHandle);
+
+			lookupLock.Leave();
+			*hFileHandle = CreateFileHandle();
+			TrackFileInput();
+
+			DEBUG_LOG_DETOURED(funcName, L"(SHAREDMEMFILE)%ls %llu (%ls) (%ls) -> Success", isWriteStr, uintptr_t(*hFileHandle), lpFileName, (fileName.data != lpFileName ? fileName.data : L""));
+			return STATUS_SUCCESS;
 		}
 
-		const wchar_t* handleStr = lpFileName + 1;
-		const wchar_t* handleStrEnd = wcschr(handleStr, '-');
-		u64 mappingOffset = 0;
-		if (!handleStrEnd)
+		SharedMemoryHandle memoryHandle = SharedMemoryHandle::FromU64(mappingHandle);
+
+		if (g_permanentFilesMemory && memoryHandle == g_permanentFilesHandle)
 		{
-			handleStrEnd = handleStr + wcslen(handleStr);
+			info.fileMapMem = g_permanentFilesMemory + mappingOffset;
+			info.fileMapMemSize = info.size;
 		}
-		else
+		else if (info.trueFileMapHandle.IsValid() && (info.trueFileMapHandle != memoryHandle || info.trueFileMapOffset != mappingOffset))
 		{
-			const wchar_t* mappingOffsetStr = handleStrEnd + 1;
-			mappingOffset = StringToValueBase62(mappingOffsetStr, wcslen(mappingOffsetStr));
+			// File has changed (should probably close previous)
+			info.fileMapMem = nullptr;
 		}
-		HANDLE mappingHandle = FileMappingHandle::FromU64(StringToValueBase62(handleStr, handleStrEnd - handleStr)).mh;
-		info.trueFileMapOffset = mappingOffset;
 
+		info.trueFileMapHandle = memoryHandle;
+		info.trueFileMapOffset = mappingOffset;
 		info.isFileMap = true;
-		if (!True_DuplicateHandle(g_hostProcess, mappingHandle, GetCurrentProcess(), &info.trueFileMapHandle, 0, FALSE, DUPLICATE_SAME_ACCESS))
-		{
-			Rpc_WriteLogf(L"Can't duplicate handle 0x%llx (%ls) for file %ls (Error %u)", uintptr_t(mappingHandle), lpFileName, info.originalName, GetLastError());
-			UBA_ASSERTF(info.trueFileMapHandle, L"Can't duplicate handle 0x%llx (%ls) for file %ls (Error %u)", uintptr_t(mappingHandle), lpFileName, info.originalName, GetLastError());
-			return STATUS_ACCESS_DENIED;
-		}
 
 		lookupLock.Leave();
-
-		UBA_ASSERT(info.size != InvalidValue);
-		UBA_ASSERTF(!isDeleteOnClose, TC("Creating file mapping %s that has delete on close"), info.originalName);
 		*hFileHandle = CreateFileHandle();
-
 		TrackFileInput();
 
-		DEBUG_LOG_DETOURED(funcName, L"(MAPPED)%ls %llu (%ls) (%ls) -> Success", isWriteStr, uintptr_t(*hFileHandle), lpFileName, (fileName.data != lpFileName ? fileName.data : L""));
+		DEBUG_LOG_DETOURED(funcName, L"(SHAREDMEMFILE)%ls %llu (%ls) (%ls) -> Success", isWriteStr, uintptr_t(*hFileHandle), lpFileName, (fileName.data != lpFileName ? fileName.data : L""));
 		return STATUS_SUCCESS;
 	}
 
-	if (lpFileName[0] == ':') // It is a HANDLE from session process. A written file that is writable
+	if (lpFileName[0] == WrittenMemoryHandleChar) // It is a HANDLE from session process. A written file that is writable
 	{
 		if (!info.memoryFile)
 		{
 			const wchar_t* handleStr = lpFileName + 1;
 			UBA_ASSERT(!wcschr(handleStr, '-'));
-			auto mappingHandle = FileMappingHandle::FromU64(StringToValueBase62(handleStr, wcslen(handleStr)));
+			auto memoryHandle = SharedMemoryHandle::FromU64(StringToValueBase62(handleStr, wcslen(handleStr)));
 			UBA_ASSERT(info.size != InvalidValue);
 			u64 mappingHandleSize = info.size;
 
 			info.memoryFile = new MemoryFile(nullptr, false);
 			MemoryFile& mf = *info.memoryFile;
-			HANDLE newHandle = 0;
-			True_DuplicateHandle(g_hostProcess, mappingHandle.mh, GetCurrentProcess(), &newHandle, 0, FALSE, DUPLICATE_SAME_ACCESS);
-			UBA_ASSERTF(newHandle, L"DuplicateHandle failed when opening temp file %ls (%u)", fileName.data, GetLastError());
 			mf.writtenSize = mappingHandleSize;
 			mf.committedSize = AlignUp(mappingHandleSize, g_pageSize);
 			mf.mappedSize = mf.committedSize;
-			TimerScope ts2(g_kernelStats.mapViewOfFile);
-			mf.baseAddress = (u8*)True_MapViewOfFile(newHandle, FILE_MAP_READ | FILE_MAP_ALL_ACCESS, 0, 0, mf.mappedSize);
-			UBA_ASSERTF(mf.baseAddress, L"MapViewOfFile failed when opening temp file %ls (%u)", fileName.data, GetLastError());
 			mf.reserveSize = FileTypeMaxSize(fileName, isSystemOrTempFile);
-			mf.mappingHandle.mh = newHandle;
+
+			TimerScope ts2(g_kernelStats.mapViewOfFile);
+
+			mf.memoryHandle = memoryHandle;
+			Rpc_UpdateSharedMemory(mf.memoryView, g_sharedMemoryAllocator, memoryHandle, mf.reserveSize);
+			mf.baseAddress = mf.memoryView.GetMemory();
 		}
 
 		if (CreateDisposition != FILE_OPEN && CreateDisposition != FILE_OPEN_IF)
@@ -1284,7 +1317,7 @@ NTSTATUS NTAPI Shared_NtCreateFile(bool IsCreateFunc, PHANDLE hFileHandle, ACCES
 		{
 			bool isOutput = (isWrite || isDeleteOnClose) && g_rules->IsOutputFile(fileName, g_systemTemp);
 			bool isLocal = !isOutput;
-			MemoryFile* mf = new MemoryFile(isLocal, FileTypeMaxSize(fileName, isSystemOrTempFile), isThrowAway, fileName.data);
+			MemoryFile* mf = new MemoryFile(isLocal, FileTypeMaxSize(fileName, isSystemOrTempFile), isThrowAway, FileTypeCommitSize(fileName), fileName.data);
 
 			if (!isThrowAway && CreateDisposition == FILE_OPEN)
 			{
@@ -1553,8 +1586,8 @@ NTSTATUS NTAPI Detoured_NtClose(HANDLE handle)
 		return res;
 	}
 
-	FileMappingHandle mappingHandle;
-	u64 mappingWritten = 0;
+	SharedMemoryHandle memoryHandle;
+	u64 memoryWritten = 0;
 	FileInfo& fi = *fo->fileInfo;
 	const wchar_t* path = fi.name;
 	wchar_t temp[512];
@@ -1563,7 +1596,7 @@ NTSTATUS NTAPI Detoured_NtClose(HANDLE handle)
 		if (IsWrite(fo->desiredAccess, 0))
 		{
 			// TODO: There are race conditions in this code. There could be other file handles accessing the same piece of memory (allthough unlikely)
-			u64 alignedWritten = AlignUp(mf->writtenSize, 64 * 1024);
+			u64 alignedWritten = AlignUp(mf->writtenSize, PageSize);
 			if (alignedWritten < mf->committedSize)
 			{
 				u64 decommitSize = u64(mf->committedSize - alignedWritten);
@@ -1575,20 +1608,12 @@ NTSTATUS NTAPI Detoured_NtClose(HANDLE handle)
 						ToInvestigate(L"Failed to decommit memory (%u)", GetLastError());
 #pragma warning(pop)
 				}
-				else
-				{
-					// Speculative change. According to stackoverflow (I know) this can hint the system that this memory is not needed anymore.
-					// Building UnrealEditor and friends put huge pressure on committed space and anything that can reduce that is valuable
-					if (!::VirtualUnlock(mf->baseAddress + alignedWritten, decommitSize))
-						if (GetLastError() != ERROR_NOT_LOCKED)
-							ToInvestigate(L"Failed to unlock memory (%u)", GetLastError());
-				}
 				mf->committedSize = alignedWritten;
 			}
 		}
 
-		mappingHandle = mf->mappingHandle;
-		mappingWritten = mf->writtenSize;
+		memoryHandle = mf->memoryHandle;
+		memoryWritten = mf->writtenSize;
 
 		u32 orginalNameLen = TStrlen(fi.originalName);
 		if ((fo->deleteOnClose || IsWrite(fo->desiredAccess, 0)) && g_rules->IsOutputFile({fi.originalName, orginalNameLen}, g_systemTemp) && !g_rules->IsThrowAway(StringView(fi.originalName, orginalNameLen), g_runningRemote))
@@ -1645,7 +1670,7 @@ NTSTATUS NTAPI Detoured_NtClose(HANDLE handle)
 
 	if (fo->closeId)
 	{
-		Rpc_UpdateCloseHandle(path, fo->closeId, fo->deleteOnClose, fo->newName.c_str(), mappingHandle, mappingWritten, true);
+		Rpc_UpdateCloseHandle(path, fo->closeId, fo->deleteOnClose, fo->newName.c_str(), memoryHandle, memoryWritten, true);
 	}
 	else
 	{
diff --git a/Engine/Source/Programs/UnrealBuildAccelerator/Detours/Private/Windows/UbaDetoursFunctionsWin.cpp b/Engine/Source/Programs/UnrealBuildAccelerator/Detours/Private/Windows/UbaDetoursFunctionsWin.cpp
--- a/Engine/Source/Programs/UnrealBuildAccelerator/Detours/Private/Windows/UbaDetoursFunctionsWin.cpp
+++ b/Engine/Source/Programs/UnrealBuildAccelerator/Detours/Private/Windows/UbaDetoursFunctionsWin.cpp
@@ -59,6 +59,11 @@ DETOURED_FUNCTIONS_MEMORY_NON_WINE
 namespace uba
 {
 	extern HANDLE g_hostProcess;
+	SharedMemoryAllocatorHandle g_sharedMemoryAllocator;
+	u8* g_sharedMemory;
+
+	SharedMemoryHandle g_permanentFilesHandle;
+	u8* g_permanentFilesMemory;
 
 bool g_useMiMalloc;
 constexpr u64  g_pageSize = 64*1024;
@@ -300,8 +305,9 @@ MemoryFile::MemoryFile(u8* data, bool localOnly)
 {
 }
 
-MemoryFile::MemoryFile(bool localOnly, u64 reserveSize_, bool isThrowAway_, const tchar* fileName)
-:	isLocalOnly(localOnly)
+MemoryFile::MemoryFile(bool localOnly, u64 reserveSize_, bool isThrowAway_, u32 commitStepSize_, const tchar* fileName)
+:	commitStepSize(commitStepSize_)
+,	isLocalOnly(localOnly)
 ,	isThrowAway(isThrowAway_)
 {
 	if (!isThrowAway)
@@ -319,65 +325,11 @@ void MemoryFile::Reserve(u64 reserveSize_, const tchar* fileName)
 			FatalError(1354, L"VirtualAlloc failed trying to reserve %llu for %s. (Error code: %u)", reserveSize, fileName, GetLastError());
 		mappedSize = reserveSize;
 	}
-#if UBA_ENABLE_ON_DISK_FILE_MAPPINGS
-	else if (fileName && !g_runningRemote)
-	{
-		StringBuffer<> tempFileName;
-		tempFileName.Appendf(L"\\??\\").Append(fileName);
-		if (false)
-			tempFileName.Append(L".uba.tmp");
-
-		UNICODE_STRING us;
-		us.Length = USHORT(tempFileName.count*sizeof(tchar));
-		us.MaximumLength = us.Length;
-		us.Buffer = tempFileName.data;
-
-		SuppressDetourScope _;
-		{
-			TimerScope ts(g_kernelStats.createFile);
-			OBJECT_ATTRIBUTES oa;
-			InitializeObjectAttributes(&oa, &us, OBJ_CASE_INSENSITIVE, NULL, NULL);
-			IO_STATUS_BLOCK iosb;
-			ACCESS_MASK desiredAccess = GENERIC_READ | GENERIC_WRITE | DELETE | SYNCHRONIZE | FILE_WRITE_ATTRIBUTES;
-			ULONG shareAccess = FILE_SHARE_READ | FILE_SHARE_WRITE | FILE_SHARE_DELETE;
-			ULONG createOptions = FILE_SYNCHRONOUS_IO_NONALERT;// | FILE_DELETE_ON_CLOSE;
-			NTSTATUS status = ZwCreateFile(&mappingHandle.fh, desiredAccess, &oa, &iosb, 0, FILE_ATTRIBUTE_NORMAL, shareAccess, FILE_OVERWRITE_IF, createOptions, 0, 0);
-			if (status != STATUS_SUCCESS)
-				FatalError(1347, L"ZwCreateFile failed to create %s. (Error code: 0x%x)", us.Buffer, status);
-		}
-		FILE_DISPOSITION_INFO info;
-		info.DeleteFile = true;
-		if (!::SetFileInformationByHandle(mappingHandle.fh, FileDispositionInfo, &info, sizeof(info)))
-			FatalError(1347, L"SetFileInformationByHandle failed to set delete-on-close on %s. (Error code: 0x%x)", us.Buffer, GetLastError());
-		{
-			TimerScope ts(g_kernelStats.createFileMapping);
-			LARGE_INTEGER liMaxSize; liMaxSize.QuadPart = 2;
-			NTSTATUS status = ZwCreateSection(&mappingHandle.mh, SECTION_ALL_ACCESS, 0, &liMaxSize, PAGE_READWRITE, SEC_COMMIT, mappingHandle.fh);
-			if (status != STATUS_SUCCESS)
-				FatalError(1348, L"NtCreateSection failed to reserve %llu. (Error code: 0x%x)", reserveSize, status);
-		}
-		TimerScope ts(g_kernelStats.mapViewOfFile);
-		NTSTATUS status = ZwMapViewOfSection(mappingHandle.mh, NtCurrentProcess(), (void**)&baseAddress, 0, 0, NULL, &reserveSize, 2, MEM_RESERVE, PAGE_READWRITE);
-		if (status != STATUS_SUCCESS)
-			FatalError(1348, L"ZwMapViewOfSection failed trying to reserve %llu. (Error code: 0x%x)", reserveSize, status);
-		mappedSize = reserveSize;
-	}
-#endif
 	else
 	{
-		HANDLE handle;
-		{
-			mappedSize = 32 * 1024 * 1024;
-			TimerScope ts(g_kernelStats.createFileMapping);
-			handle = True_CreateFileMappingW(INVALID_HANDLE_VALUE, NULL, PAGE_READWRITE | SEC_RESERVE, ToHigh(reserveSize), ToLow(reserveSize), NULL);
-			if (!handle)
-				FatalError(1348, L"CreateFileMappingW failed trying to reserve %llu for %s. (Error code: %u)", reserveSize, fileName, GetLastError());
-		}
-		TimerScope ts(g_kernelStats.mapViewOfFile);
-		baseAddress = (u8*)True_MapViewOfFile(handle, FILE_MAP_READ | FILE_MAP_WRITE, 0, 0, mappedSize);
-		if (!baseAddress)
-			FatalError(1353, L"MapViewOfFile failed trying to map %llu for %s. ReservedSize: %llu (Error code: %u)", mappedSize, fileName, reserveSize, GetLastError());
-		mappingHandle = { 0, handle };
+		memoryHandle = Rpc_CreateSharedMemory(memoryView, g_sharedMemoryAllocator, reserveSize, 0, fileName);
+		mappedSize = reserveSize;
+		baseAddress = memoryView.GetMemory();
 	}
 }
 
@@ -389,15 +341,13 @@ void MemoryFile::Unreserve()
 	}
 	else
 	{
-		True_UnmapViewOfFile(baseAddress);
-		CloseHandle(mappingHandle.mh);
-		mappingHandle.mh = nullptr;
+		memoryView.Reset();
 	}
 	baseAddress = nullptr;
 	committedSize = 0;
 }
 
-void MemoryFile::Write(DetouredHandle& handle, LPCVOID lpBuffer, u64 nNumberOfBytesToWrite)
+void MemoryFile::Write(DetouredHandle& handle, LPCVOID lpBuffer, u64 nNumberOfBytesToWrite, const tchar* hint)
 {
 	u64 newPos = handle.pos + nNumberOfBytesToWrite;
 
@@ -424,61 +374,38 @@ void MemoryFile::EnsureCommitted(const DetouredHandle& handle, u64 size)
 
 	if (committedSize >= size)
 		return;
+	if (size > reserveSize)
+		FatalError(1347, L"Reserved size of %ls is smaller than what is requested to be. ReserveSize: %llu Written: %llu Requested: %llu", HandleToName(handle), reserveSize, writtenSize, size);
+	/*
 	if (size > mappedSize)
 	{
-		bool shouldRemap = true;
-		if (size > reserveSize)
+		if (shouldRemap)
 		{
-			if (writtenSize == 0 && !isReported)
-			{
-				u64 newReserve = AlignUp(size, g_pageSize);
-				if (reserveSize)
-					Rpc_WriteLogf(L"TODO: RE-RESERVING MemoryFile. Initial reserve: %llu, New reserve: %llu. Please fix application rules", reserveSize, newReserve);
-				Unreserve();
-				Reserve(newReserve);
-				shouldRemap = false;
-			}
-			else
-				FatalError(1347, L"Reserved size of %ls is smaller than what is requested to be. ReserveSize: %llu Written: %llu Requested: %llu", HandleToName(handle), reserveSize, writtenSize, size);
+			memoryHandle = Rpc_CreateSharedMemory(memoryView, g_sharedMemoryAllocator, reserveSize, 0, TC(""));
+			mappedSize = reserveSize;
+			baseAddress = memoryView.GetMemory();
 		}
-
-		if (shouldRemap)
-			Remap(handle, size);
 	}
-
-#if UBA_ENABLE_ON_DISK_FILE_MAPPINGS
-	if (mappingHandle.fh)
-	{
-		if (!committedSize)
-			committedSize = size;
-		else
-			committedSize = Min(reserveSize, Max(size, committedSize * 2));
-		auto lg = ToLargeInteger(committedSize);
-		NtExtendSection(mappingHandle.mh, &lg);
-	}
-	else
-#endif
+	*/
 	{
 		u64 toCommit = Min(reserveSize, AlignUp(size - committedSize, g_pageSize));
-		if (!VirtualAlloc(baseAddress + committedSize, toCommit, MEM_COMMIT, PAGE_READWRITE))
-			FatalError(1347, L"Failed to ensure virtual memory for %ls trying to commit %llu at %llx. MappedSize: %llu, CommittedSize: %llu RequestedSize: %llu. (%u)", HandleToName(handle), toCommit, baseAddress + committedSize, mappedSize, committedSize, size, GetLastError());
-		if (g_allowDiscardVirtualMemory)
-			DiscardVirtualMemory(baseAddress + committedSize, toCommit);
+
+		if (!isLocalOnly)
+		{
+			toCommit = AlignUp(toCommit, u64(commitStepSize));
+			Rpc_CommitSharedMemory(memoryView, g_sharedMemoryAllocator, memoryHandle, toCommit);
+		}
+		else
+		{
+			if (!VirtualAlloc(baseAddress + committedSize, toCommit, MEM_COMMIT, PAGE_READWRITE))
+				FatalError(1347, L"Failed to ensure virtual memory for %ls trying to commit %llu at %llx. MappedSize: %llu, CommittedSize: %llu RequestedSize: %llu. (%u)", HandleToName(handle), toCommit, baseAddress + committedSize, mappedSize, committedSize, size, GetLastError());
+			if (g_allowDiscardVirtualMemory)
+				DiscardVirtualMemory(baseAddress + committedSize, toCommit);
+		}
 		committedSize += toCommit;
 	}
 }
 
-void MemoryFile::Remap(const DetouredHandle& handle, u64 size)
-{
-	UBA_ASSERT(!mappingHandle.fh);
-	True_UnmapViewOfFile(baseAddress);
-	mappedSize = Min(reserveSize, AlignUp(Max(size, mappedSize * 4), g_pageSize));
-	TimerScope ts(g_kernelStats.mapViewOfFile);
-	baseAddress = (u8*)True_MapViewOfFile(mappingHandle.mh, FILE_MAP_READ | FILE_MAP_WRITE, 0, 0, mappedSize);
-	if (!baseAddress)
-		FatalError(1347, L"MapViewOfFile failed trying to map %llu for %ls. ReservedSize: %llu (Error code: %u)", mappedSize, HandleToName(handle), reserveSize, GetLastError());
-}
-
 void ToInvestigate(const wchar_t* format, ...)
 {
 #if UBA_DEBUG_LOG_ENABLED
@@ -559,7 +486,8 @@ const wchar_t* HandleToName(HANDLE handle)
 }
 
 
-u64 FileTypeMaxSize(const StringBufferBase& file, bool isSystemOrTempFile) { return g_rules->FileTypeMaxSize(file, isSystemOrTempFile); }
+u64 FileTypeMaxSize(StringView file, bool isSystemOrTempFile) { return g_rules->FileTypeMaxSize(file, isSystemOrTempFile); }
+u32 FileTypeCommitSize(StringView file) { return g_rules->FileTypeCommitSize(file); }
 
 bool EnsureMapped(DetouredHandle& handle, DWORD dwFileOffsetHigh = 0, DWORD dwFileOffsetLow = 0, SIZE_T numberOfBytesToMap = 0, void* baseAddress = nullptr)
 {
@@ -575,35 +503,14 @@ bool EnsureMapped(DetouredHandle& handle, DWORD dwFileOffsetHigh = 0, DWORD dwFi
 		numberOfBytesToMap = info.size;
 	}
 
-	u64 alignedOffsetStart = 0;
-	if (info.trueFileMapOffset)
-	{
-		offset += info.trueFileMapOffset;
-		u64 endOffset = offset + numberOfBytesToMap;
-		alignedOffsetStart = AlignUp(offset - (g_pageSize - 1), g_pageSize);
-		u64 alignedOffsetEnd = AlignUp(endOffset, g_pageSize);
-		u64 mapSize = alignedOffsetEnd - alignedOffsetStart;
-		TimerScope ts(g_kernelStats.mapViewOfFile);
-		info.fileMapMem = (u8*)True_MapViewOfFileEx(info.trueFileMapHandle, info.fileMapViewDesiredAccess, ToHigh(alignedOffsetStart), ToLow(alignedOffsetStart), mapSize, baseAddress);
-		
-		// TODO: This is ugly.. but in some cases we have created virtual files pointing into segments of real file
-		// and real file size is not aligned to page size so we need to set map size to 0 to allow it to map up to size of file
-		if (!info.fileMapMem)
-			info.fileMapMem = (u8*)True_MapViewOfFileEx(info.trueFileMapHandle, info.fileMapViewDesiredAccess, ToHigh(alignedOffsetStart), ToLow(alignedOffsetStart), 0, baseAddress);
-	}
-	else
-	{
-		UBA_ASSERTF(!info.freeFileMapOnClose, TC("File %s has been freed because of earlier close and is now reopened (%s)"), info.name, info.originalName);
-		TimerScope ts(g_kernelStats.mapViewOfFile);
-		info.fileMapMem = (u8*)True_MapViewOfFileEx(info.trueFileMapHandle, info.fileMapViewDesiredAccess, 0, 0, numberOfBytesToMap, baseAddress);
-	}
+	offset += info.trueFileMapOffset;
+	u64 alignedOffsetStart = AlignUp(offset - (PageSize - 1), PageSize);
+	u64 mapSize = (offset - alignedOffsetStart) + numberOfBytesToMap;
 
-	if (info.fileMapMem == nullptr || (baseAddress && info.fileMapMem != baseAddress))
-	{
-		ToInvestigate(L"MapViewOfFileEx failed trying to map %llu bytes on address 0x%llx with offset %llu, using %llu with access %u (%u)", numberOfBytesToMap, u64(baseAddress), alignedOffsetStart, u64(info.trueFileMapHandle), info.fileMapViewDesiredAccess, GetLastError());
-		return false;
-	}
-	info.fileMapMem += (offset - alignedOffsetStart);
+	SharedMemoryView view;
+	Rpc_GetSharedMemory(view, g_sharedMemoryAllocator, g_sharedMemory, info.trueFileMapHandle, alignedOffsetStart, mapSize, info.originalName);
+
+	info.fileMapMem = view.DetachMemory() + (offset - alignedOffsetStart);
 	info.fileMapMemSize = info.size;
 
 	DEBUG_LOG_TRUE(L"INTERNAL MapViewOfFileEx", L"(%ls) (size: %llu) (%ls) -> 0x%llx", info.name, numberOfBytesToMap, info.originalName, uintptr_t(info.fileMapMem));
@@ -633,19 +540,17 @@ void CloseCaches()
 			DEBUG_LOG_TRUE(L"INTERNAL UnmapViewOfFile", L"0x%llx (%ls) (%ls)", uintptr_t(info.fileMapMem), info.name, info.originalName);
 			True_UnmapViewOfFile(info.fileMapMem);
 		}
-		if (info.trueFileMapHandle != 0)
+		if (info.trueFileMapHandle.IsValid())
 		{
-			DEBUG_LOG_TRUE(L"INTERNAL CloseHandle", L"%llu (%ls) (%ls)", uintptr_t(info.trueFileMapHandle), info.name, info.originalName);
-			CloseHandle(info.trueFileMapHandle);
+			DEBUG_LOG_TRUE(L"INTERNAL CloseHandle", L"%llu (%ls) (%ls)", info.trueFileMapHandle.ToU64(), info.name, info.originalName);
+
+			// TODO
+			//CloseHandle(info.trueFileMapHandle);
 		}
 
 		// Let them leak
 		if (info.memoryFile && !info.memoryFile->isLocalOnly)
-		{
-			UnmapViewOfFile(info.memoryFile->baseAddress);
-			CloseHandle(info.memoryFile->mappingHandle.mh);
-			CloseHandle(info.memoryFile->mappingHandle.fh);
-		}
+			info.memoryFile->memoryView.Reset();
 		//if (info.memoryFile && info.memoryFile != &g_emptyMemoryFile)
 		//	delete info.memoryFile;
 	}
@@ -1071,6 +976,8 @@ void Init(const DetoursPayload& payload, u64 startTime)
 	StringBuffer<512> applicationBuffer;
 	StringBuffer<512> workingDirBuffer;
 
+	HANDLE sharedMemoryAllocatorHandle;
+	SharedMemoryHandle permanentFilesHandle;
 	HANDLE directoryTableHandle;
 	u32 directoryTableSize;
 	u32 directoryTableCount;
@@ -1089,6 +996,8 @@ void Init(const DetoursPayload& payload, u64 startTime)
 		reader.ReadString(applicationBuffer);
 		reader.ReadString(workingDirBuffer);
 
+		sharedMemoryAllocatorHandle = (HANDLE)reader.ReadU64();
+		permanentFilesHandle = SharedMemoryHandle::FromU64(reader.ReadU64());
 		directoryTableHandle = FileMappingHandle::FromU64(reader.ReadU64()).mh;
 		directoryTableSize = reader.ReadU32();
 		directoryTableCount = reader.ReadU32();
@@ -1130,6 +1039,9 @@ void Init(const DetoursPayload& payload, u64 startTime)
 	}
 	#endif
 
+	if (!True_DuplicateHandle(g_hostProcess, sharedMemoryAllocatorHandle, GetCurrentProcess(), &g_sharedMemoryAllocator.mh, 0, FALSE, DUPLICATE_SAME_ACCESS))
+		UBA_ASSERTF(false, L"Failed to duplicate filetable handle (%u)", GetLastError());
+
 	if (!True_DuplicateHandle(g_hostProcess, mappedFileTableHandle, GetCurrentProcess(), &mappedFileTableHandle, 0, FALSE, DUPLICATE_SAME_ACCESS))
 		UBA_ASSERTF(false, L"Failed to duplicate filetable handle (%u)", GetLastError());
 
@@ -1158,6 +1070,14 @@ void Init(const DetoursPayload& payload, u64 startTime)
 		g_directoryTable.Init(directoryTableMem, directoryTableCount, directoryTableSize);
 	}
 
+	if (permanentFilesHandle.IsValid())
+	{
+		SharedMemoryView permanentFilesView;
+		Rpc_GetSharedMemory(permanentFilesView, g_sharedMemoryAllocator, nullptr, permanentFilesHandle, 0, 0, TC("PermanentFiles"));
+		g_permanentFilesHandle = permanentFilesHandle;
+		g_permanentFilesMemory = permanentFilesView.DetachMemory();
+	}
+
 	if (g_isChild)
 		Rpc_GetWrittenFiles();
 
diff --git a/Engine/Source/Programs/UnrealBuildAccelerator/Detours/Private/Windows/UbaDetoursUtilsWin.h b/Engine/Source/Programs/UnrealBuildAccelerator/Detours/Private/Windows/UbaDetoursUtilsWin.h
--- a/Engine/Source/Programs/UnrealBuildAccelerator/Detours/Private/Windows/UbaDetoursUtilsWin.h
+++ b/Engine/Source/Programs/UnrealBuildAccelerator/Detours/Private/Windows/UbaDetoursUtilsWin.h
@@ -2,8 +2,10 @@
 
 #pragma once
 
+#include "UbaDefinitions.h"
 #include "UbaDetoursFunctionsWin.h"
 #include "UbaDetoursFileMappingTable.h"
+#include "UbaSharedMemoryView.h"
 
 namespace uba
 {
@@ -61,19 +63,22 @@ namespace uba
 	struct MemoryFile
 	{
 		MemoryFile(u8* data = nullptr, bool localOnly = true);
-		MemoryFile(bool localOnly, u64 reserveSize_, bool isThrowAway_ = false, const tchar* fileName = nullptr);
+		MemoryFile(bool localOnly, u64 reserveSize, bool isThrowAway, u32 commitStepSize, const tchar* fileName = nullptr);
 
 		void Reserve(u64 reserveSize_, const tchar* fileName = nullptr);
 		void Unreserve();
-		void Write(struct DetouredHandle& handle, LPCVOID lpBuffer, u64 nNumberOfBytesToWrite);
+		void Write(struct DetouredHandle& handle, LPCVOID lpBuffer, u64 nNumberOfBytesToWrite, const tchar* hint);
 		void EnsureCommitted(const struct DetouredHandle& handle, u64 size);
 		void Remap(const struct DetouredHandle& handle, u64 size);
 
 		u64 fileIndex = ~u64(0);
 		u64 fileTime = ~u64(0);
 		u32 volumeSerial = 0;
+		u32 commitStepSize = 1024*1024;
+
+		SharedMemoryHandle memoryHandle;
+		SharedMemoryView memoryView;
 
-		FileMappingHandle mappingHandle;
 		u8* baseAddress = nullptr;
 		u64 reserveSize = 0;
 		u64 mappedSize = 0;
diff --git a/Engine/Source/Programs/UnrealBuildAccelerator/Test/Private/UbaTestAll.h b/Engine/Source/Programs/UnrealBuildAccelerator/Test/Private/UbaTestAll.h
--- a/Engine/Source/Programs/UnrealBuildAccelerator/Test/Private/UbaTestAll.h
+++ b/Engine/Source/Programs/UnrealBuildAccelerator/Test/Private/UbaTestAll.h
@@ -2,6 +2,7 @@
 
 #pragma once
 
+#include "UbaEvent.h"
 #include "UbaFile.h"
 #include "UbaStringBuffer.h"
 #include "UbaTest.h"
@@ -28,6 +29,7 @@ namespace uba
 		UBA_TEST(TestMemoryBlock) \
 		UBA_TEST(TestParseArguments) \
 		UBA_TEST(TestBinaryWriter) \
+		UBA_TEST(TestSharedMemory) \
 		UBA_TEST(TestFileMappingBuffer) \
 		UBA_TEST(TestCrypto) \
 		UBA_TEST(TestSockets) \
@@ -80,11 +82,12 @@ namespace uba
 		UBA_TEST(TestStdOutViaCmd) \
 		UBA_TEST(TestVolumeCache) \
 		UBA_TEST(TestDependencyCrawler) \
-		UBA_TEST(TestVirtualFile) \
-		UBA_TEST(TestRemoteVirtualFile) \
+		UBA_TEST(TestCreateVirtualFile) \
+		UBA_TEST(TestRemoteCreateVirtualFile) \
+		UBA_TEST(TestRegisterVirtualFile) \
+		UBA_TEST(TestRemoteRegisterVirtualFile) \
 		UBA_TEST(TestRemoteProcessSpecialCase1) \
 		UBA_TEST(TestSessionSpecialCopy) \
-		UBA_TEST(TestSharedMemory) \
 
 
 #define UBA_LINUX_TESTS \
@@ -127,7 +130,6 @@ namespace uba
 	UBA_TESTS
 	#undef UBA_TEST
 
-
 	#define UBA_TEST(x) \
 		if (!filter || Contains(TC(#x), filter)) \
 		{ \
diff --git a/Engine/Source/Programs/UnrealBuildAccelerator/Test/Private/UbaTestBasics.cpp b/Engine/Source/Programs/UnrealBuildAccelerator/Test/Private/UbaTestBasics.cpp
--- a/Engine/Source/Programs/UnrealBuildAccelerator/Test/Private/UbaTestBasics.cpp
+++ b/Engine/Source/Programs/UnrealBuildAccelerator/Test/Private/UbaTestBasics.cpp
@@ -4,7 +4,7 @@
 #include "UbaBinaryReaderWriter.h"
 #include "UbaDefinitions.h"
 #include "UbaFileAccessor.h"
-#include "UbaFileMapping.h"
+#include "UbaFileMappingBuffer.h"
 #include "UbaPathUtils.h"
 #include "UbaPlatform.h"
 #include "UbaProcessUtils.h"
@@ -12,6 +12,7 @@
 #include "UbaFile.h"
 #include "UbaLogger.h"
 #include "UbaRootPaths.h"
+#include "UbaSharedMemoryAllocator.h"
 #include "UbaTest.h"
 #include "UbaThread.h"
 #include "UbaTimer.h"
@@ -141,9 +142,11 @@ namespace uba
 		#if !PLATFORM_WINDOWS
 		if (IsRunningDarling()) // bug in darling, this does not work
 			return true;
-		FileMappingBuffer mappingBuffer(logger);
-		mappingBuffer.AddTransient(TC("Temp"), false);
-		MappedView view = mappingBuffer.AllocAndMapView(MappedView_Transient, 1024, 1, TC("Foo"));
+		SharedMemoryAllocator allocator(logger);
+		allocator.Init(1024 * 1024);
+		FileMappingBuffer mappingBuffer(logger, allocator);
+		mappingBuffer.Init(TC(""), 64*1024, 64 * 1024);
+		MappedView view = mappingBuffer.AllocAndMapView(1024, 1, TC("Foo"));
 		SharedEvent* ev = new (view.memory) SharedEvent();
 		if (!ev->Create(EventResetType_Manual))
 			return logger.Error(TC("Failed to create event"));
diff --git a/Engine/Source/Programs/UnrealBuildAccelerator/Test/Private/UbaTestFileMappings.cpp b/Engine/Source/Programs/UnrealBuildAccelerator/Test/Private/UbaTestFileMappings.cpp
--- a/Engine/Source/Programs/UnrealBuildAccelerator/Test/Private/UbaTestFileMappings.cpp
+++ b/Engine/Source/Programs/UnrealBuildAccelerator/Test/Private/UbaTestFileMappings.cpp
@@ -1,25 +1,61 @@
 // Copyright Epic Games, Inc. All Rights Reserved.
 
 #include "UbaLogger.h"
-#include "UbaFileMapping.h"
+#include "UbaFileMappingBuffer.h"
+#include "UbaSharedMemoryAllocator.h"
+#include "UbaTest.h"
 
 namespace uba
 {
 	bool TestFileMappingBuffer(LoggerWithWriter& logger, const StringBufferBase& rootDir)
 	{
+		TrackSystemUsage tsu(logger);
+		SharedMemoryAllocator allocator(logger);
+		allocator.Init(64ull*1024*1024);
+		FileMappingBuffer mappingBuffer(logger, allocator);
 
-		FileMappingBuffer mappingBuffer(logger);
+		mappingBuffer.Init(TC("Temp"), 4ull*1024*1024, 1024*1024);
 
-		mappingBuffer.AddTransient(TC("Temp"), false);
+		for (u32 i = 0; i != 65; ++i)
+		{
+			MappedView view = mappingBuffer.AllocAndMapView(1024, 1, TC("Foo"));
 
-		MappedView view = mappingBuffer.AllocAndMapView(MappedView_Transient, 1024, 1, TC("Foo"));
+			u64 v = i + 1337;
+			*(u64*)view.memory = v;
+			u64 value = *(u64*)view.memory;
+			if (value != v)
+				return false;
+			mappingBuffer.UnmapView(view, TC("Foo"));
 
-		*(u64*)view.memory = 1337;
-		u64 value = *(u64*)view.memory;
-		if (value != 1337)
-			return false;
-		mappingBuffer.UnmapView(view, TC("Foo"));
+			view = mappingBuffer.MapView(view.handle, view.offset, 1024, TC("Foo"));
+			value = *(u64*)view.memory;
+			if (value != v)
+				return false;
+			mappingBuffer.UnmapView(view, TC("Foo"));
+		}
 
+		{
+			// Allow shrink
+			MappedView view1 = mappingBuffer.AllocAndMapView(1024, 1, TC("Foo"));
+			MappedView view2 = mappingBuffer.AllocAndMapView(1024, 1, TC("Foo"));
+			mappingBuffer.UnmapView(view1, TC("Foo"), 1024);
+			mappingBuffer.UnmapView(view2, TC("Foo"), 512);
+		}
+
+		{
+			// Independent
+			MappedView view = mappingBuffer.AllocAndMapView(1024, 1, TC("Foo"), true);
+			*(u64*)view.memory = 1337;
+			u64 value = *(u64*)view.memory;
+			if (value != 1337)
+				return false;
+			mappingBuffer.UnmapView(view, TC("Foo"), 1024);
+			view = mappingBuffer.MapView(view.handle, view.offset, 1024, TC("Foo"));
+			value = *(u64*)view.memory;
+			if (value != 1337)
+				return false;
+			mappingBuffer.UnmapView(view, TC("Foo"), 1024);
+		}
 		return true;
 	}
 }
diff --git a/Engine/Source/Programs/UnrealBuildAccelerator/Test/Private/UbaTestSession.cpp b/Engine/Source/Programs/UnrealBuildAccelerator/Test/Private/UbaTestSession.cpp
--- a/Engine/Source/Programs/UnrealBuildAccelerator/Test/Private/UbaTestSession.cpp
+++ b/Engine/Source/Programs/UnrealBuildAccelerator/Test/Private/UbaTestSession.cpp
@@ -40,6 +40,7 @@ namespace uba
 
 		SessionServerCreateInfo sessionServerInfo(storage, server, logWriter);
 		sessionServerInfo.rootDir = rootDir.data;
+		sessionServerInfo.readIntermediateFilesCompressed = true;
 		sessionServerInfo.Apply(config);
 
 		#if UBA_DEBUG
@@ -294,6 +295,10 @@ namespace uba
 
 		if (!CreateTestFile(logger, ToView(workingDir), TCV("File4.out"), TCV("0")))
 			return false;
+		if (!CreateTestFile(logger, ToView(workingDir), TCV("File4.obj"), TCV("9")))
+			return false;
+		if (!CreateTestFile(logger, ToView(workingDir), TCV("File4.lib"), TCV("9")))
+			return false;
 
 		if (!RunTestApp(logger, workingDir, runProcess, TC("")))
 			return false;
@@ -766,7 +771,7 @@ namespace uba
 #endif
 	}
 
-	bool RunVirtualFileTest(LoggerWithWriter& logger, SessionServer& session, const tchar* workingDir, const RunProcessFunction& runProcess, void* extraData)
+	bool RunCreateVirtualFileTest(LoggerWithWriter& logger, SessionServer& session, const tchar* workingDir, const RunProcessFunction& runProcess, void* extraData)
 	{
 		bool transient = *(bool*)extraData;
 
@@ -819,34 +824,64 @@ namespace uba
 		return true;
 	}
 
-	bool TestVirtualFile(LoggerWithWriter& logger, const StringBufferBase& testRootDir)
+	bool TestCreateVirtualFile(LoggerWithWriter& logger, const StringBufferBase& testRootDir)
 	{
 		Config config;
 		config.AddTable(TC("Session")).AddValue(TC("ShouldWriteToDisk"), false);
 		bool transient = true;
-		if (!RunLocal(logger, testRootDir, config, RunVirtualFileTest, &transient))
+		if (!RunLocal(logger, testRootDir, config, RunCreateVirtualFileTest, &transient))
 			return false;
 		transient = false;
-		if (!RunLocal(logger, testRootDir, config, RunVirtualFileTest, &transient))
+		if (!RunLocal(logger, testRootDir, config, RunCreateVirtualFileTest, &transient))
 			return false;
 		return true;
 	}
 
-	bool TestRemoteVirtualFile(LoggerWithWriter& logger, const StringBufferBase& testRootDir)
+	bool TestRemoteCreateVirtualFile(LoggerWithWriter& logger, const StringBufferBase& testRootDir)
 	{
 		Config config;
 		config.AddTable(TC("Session")).AddValue(TC("ShouldWriteToDisk"), false);
 		config.AddTable(TC("Storage")).AddValue(TC("CreateIndependentMappings"), true);
 		bool transient = true;
-		if (!RunRemote(logger, testRootDir, config, RunVirtualFileTest, &transient))
+		if (!RunRemote(logger, testRootDir, config, RunCreateVirtualFileTest, &transient))
 			return false;
 		transient = false;
-		if (!RunRemote(logger, testRootDir, config, RunVirtualFileTest, &transient))
+		if (!RunRemote(logger, testRootDir, config, RunCreateVirtualFileTest, &transient))
+			return false;
+		return true;
+	}
+
+	bool RunRegisterVirtualFileTest(LoggerWithWriter& logger, SessionServer& session, const tchar* workingDir, const RunProcessFunction& runProcess, void* extraData)
+	{
+		StringBuffer<> bigFile;
+		if (!CreateTestFile(bigFile, logger, ToView(workingDir), TCV("BigFile.txt"), TCV("FooBar")))
+			return false;
+		if (!session.RegisterVirtualFile(StringBuffer(workingDir).Append(TCV("First.txt")).data, bigFile.data, 0, sizeof(tchar)*3))
+			return false;
+		if (!session.RegisterVirtualFile(StringBuffer(workingDir).Append(TCV("Second.txt")).data, bigFile.data, sizeof(tchar) * 3, sizeof(tchar) * 3))
+			return false;
+		if (!RunTestApp(logger, workingDir, runProcess, TC("-testRegisterVirtualFile")))
 			return false;
 		return true;
 	}
 
-	#if PLATFORM_MAC
+	bool TestRegisterVirtualFile(LoggerWithWriter& logger, const StringBufferBase& testRootDir)
+	{
+		Config config;
+		if (!RunLocal(logger, testRootDir, config, RunRegisterVirtualFileTest))
+			return false;
+		return true;
+	}
+
+	bool TestRemoteRegisterVirtualFile(LoggerWithWriter& logger, const StringBufferBase& testRootDir)
+	{
+		Config config;
+		if (!RunRemote(logger, testRootDir, config, RunRegisterVirtualFileTest))
+			return false;
+		return true;
+	}
+
+#if PLATFORM_MAC
 	bool RunXCodeSelect(LoggerWithWriter& logger, SessionServer& session, const tchar* workingDir, const RunProcessFunction& runProcess, void* extraData)
 	{
 		#if UBA_DEBUG // Failing on farm for some reason... need to revisit
diff --git a/Engine/Source/Programs/UnrealBuildAccelerator/Test/Private/UbaTestSharedMemory.cpp b/Engine/Source/Programs/UnrealBuildAccelerator/Test/Private/UbaTestSharedMemory.cpp
--- a/Engine/Source/Programs/UnrealBuildAccelerator/Test/Private/UbaTestSharedMemory.cpp
+++ b/Engine/Source/Programs/UnrealBuildAccelerator/Test/Private/UbaTestSharedMemory.cpp
@@ -4,11 +4,14 @@
 #include "UbaBinaryReaderWriter.h"
 #include "UbaSharedMemoryAllocator.h"
 #include "UbaSharedMemoryView.h"
+#include "UbaTest.h"
 
 namespace uba
 {
 	bool TestSharedMemory(LoggerWithWriter& logger, const StringBufferBase& rootDir)
 	{
+		TrackSystemUsage tsu(logger);
+
 		{
 			SharedMemoryAllocator mem(logger);
 			mem.Init(64ull * 1024 * 1024 * 1024);
@@ -27,19 +30,19 @@ namespace uba
 			{
 				{
 					StackBinaryWriter<1024> writer;
-					mem.ExtendMemory(writer, viewHandle1, 64 * 1024);
+					mem.ExtendMemory(writer, viewHandle1, PageSize);
 					BinaryReader reader(writer.GetData(), 0, writer.GetPosition());
 					view1.AddRequestedMemory(handle, reader, SharedMemoryMapType_ReadWrite, 0, TC(""), 0xff);
 				}
 				{
 					StackBinaryWriter<1024> writer;
-					mem.ExtendMemory(writer, viewHandle2, 64 * 1024);
+					mem.ExtendMemory(writer, viewHandle2, PageSize);
 					BinaryReader reader(writer.GetData(), 0, writer.GetPosition());
 					view2.AddRequestedMemory(handle, reader, SharedMemoryMapType_ReadWrite, 0, TC(""), 0xdd);
 				}
 			}
 
-			u8* ptr = mem.MapView(viewHandle1, SharedMemoryMapType_ReadOnly);
+			u8* ptr = mem.MapView(viewHandle1, TC(""), SharedMemoryMapType_ReadOnly);
 			mem.UnmapView(viewHandle1, ptr);
 
 			SharedMemoryView view3;
@@ -52,19 +55,18 @@ namespace uba
 			SharedMemoryHandle viewHandle4 = mem.CreateHandle(logger, TC(""));
 			SharedMemoryView view4;
 			view4.Init(handle, 1024 * 1024);
-			mem.ExtendMemory(view4, viewHandle4, 64 * 1024, TC(""), false, mem.GetReserved() / 2);
-			mem.ExtendMemory(view4, viewHandle4, 64 * 1024, TC(""), false, mem.GetReserved() / 2 + view4.GetMappedSize());
+			mem.ExtendMemory(view4, viewHandle4, PageSize, TC(""), false);
+			mem.ExtendMemory(view4, viewHandle4, PageSize, TC(""), false);
 		}
 
 		{
 			SharedMemoryAllocator mem(logger);
 			mem.Init(1024 * 1024);
-
 			SharedMemoryHandle viewHandle1 = mem.CreateHandle(logger, TC(""));
 			SharedMemoryHandle viewHandle2 = mem.CreateHandle(logger, TC(""));
 
 			StackBinaryWriter<1024> writer1;
-			mem.ExtendMemory(writer1, viewHandle1, 64 * 1024);
+			mem.ExtendMemory(writer1, viewHandle1, PageSize);
 
 			StackBinaryWriter<1024> writer2;
 			mem.ExtendMemory(writer2, viewHandle2, (1024-64) * 1024);
@@ -72,6 +74,7 @@ namespace uba
 			mem.CloseHandle(logger, viewHandle1, TC(""));
 			mem.CloseHandle(logger, viewHandle2, TC(""));
 		}
+
 		return true;
 	}
 }
diff --git a/Engine/Source/Programs/UnrealBuildAccelerator/Test/Private/UbaTestStorage.cpp b/Engine/Source/Programs/UnrealBuildAccelerator/Test/Private/UbaTestStorage.cpp
--- a/Engine/Source/Programs/UnrealBuildAccelerator/Test/Private/UbaTestStorage.cpp
+++ b/Engine/Source/Programs/UnrealBuildAccelerator/Test/Private/UbaTestStorage.cpp
@@ -61,8 +61,8 @@ namespace uba
 		if (originalSize != copySize)
 			return logger.Error(TC("Size mismatch between %s and %s (%llu vs %llu)"), detoursLib.data, detoursLibCopy.data, originalSize, copySize);
 
-		u8 originalBuffer[64 * 1024];
-		u8 copyBuffer[64 * 1024];
+		u8 originalBuffer[PageSize];
+		u8 copyBuffer[PageSize];
 		u64 left = originalSize;
 		while (left)
 		{
@@ -146,7 +146,7 @@ namespace uba
 
 		CasKey key;
 		bool storeCompressed = false;
-		if (!storageClient.StoreCasFileClient(key, ToStringKeyLower(fileName), fileName.data, FileMappingHandle{}, 0, 0, TC("UbaTestFile"), false, storeCompressed))
+		if (!storageClient.StoreCasFileClient(key, ToStringKeyLower(fileName), fileName.data, SharedMemoryHandle{}, 0, 0, TC("UbaTestFile"), false, storeCompressed))
 			return logger.Error(TC("Failed to store file %s"), fileName.data);
 
 		fileName.Clear().Append(testRootDir).Append(TCV("Uba")).Append(PathSeparator).Append(TCV("UbaTestFile"));
@@ -288,7 +288,7 @@ namespace uba
 		auto store = [&]()
 			{
 				bool storeCompressed = true;
-				if (!storageClient.StoreCasFileClient(key, ToStringKeyLower(fileName), fileName.data, FileMappingHandle{}, 0, 0, TC("UbaTestFile"), false, storeCompressed))
+				if (!storageClient.StoreCasFileClient(key, ToStringKeyLower(fileName), fileName.data, SharedMemoryHandle{}, 0, 0, TC("UbaTestFile"), false, storeCompressed))
 					return logger.Error(TC("Failed to store file %s"), fileName.data);
 				return true;
 			};
diff --git a/Engine/Source/Programs/UnrealBuildAccelerator/TestApp/Private/Windows/UbaTestAppWin.cpp b/Engine/Source/Programs/UnrealBuildAccelerator/TestApp/Private/Windows/UbaTestAppWin.cpp
--- a/Engine/Source/Programs/UnrealBuildAccelerator/TestApp/Private/Windows/UbaTestAppWin.cpp
+++ b/Engine/Source/Programs/UnrealBuildAccelerator/TestApp/Private/Windows/UbaTestAppWin.cpp
@@ -190,12 +190,35 @@ int wmain(int argc, wchar_t* argv[])
 			if (fh == INVALID_HANDLE_VALUE)
 				return LogError(L"Failed to open File4.out for read with write permissions");
 			DWORD bytesRead;
-			char data[1];
-			if (!ReadFile(fh, data, 1, &bytesRead, NULL) || bytesRead != 1 || data[0] != '0')
+			wchar_t data[1];
+			if (!ReadFile(fh, data, sizeof(data), &bytesRead, NULL) || bytesRead != sizeof(data) || data[0] != '0')
 				return LogError(L"Failed to read one byte from File4.out");
 			CloseHandle(fh);
 		}
 
+		auto createForRead = [](const wchar_t* fileName, wchar_t expectedData)
+			{
+				HANDLE fh = CreateFileW(fileName, GENERIC_READ, FILE_SHARE_READ, NULL, OPEN_EXISTING, 0, NULL);
+				if (fh == INVALID_HANDLE_VALUE)
+					return LogError(L"Failed to open File5.h for read with write permissions");
+				DWORD bytesRead;
+				wchar_t data[1];
+				if (!ReadFile(fh, data, sizeof(data), &bytesRead, NULL) || bytesRead != sizeof(data) || data[0] != expectedData)
+					return LogError(L"Failed to read one byte from %s", fileName);
+				CloseHandle(fh);
+				return 0;
+			};
+
+		if (createForRead(L"File4.obj", '9'))
+			return -1;
+		if (createForRead(L"File4.obj", '9'))
+			return -1;
+		if (createForRead(L"File4.lib", '9'))
+			return -1;
+
+		
+		if (WriteTestFile(L"Foo", 8, L"MemoryFile.out"))
+			return -1;
 
 		{
 			STARTUPINFOW si;
@@ -395,6 +418,18 @@ int wmain(int argc, wchar_t* argv[])
 		buf[0] = wchar_t('1' + readIndex);
 		return WriteTestFile(buf, 2, L"SpecialFile1");
 	}
+	else if (wcscmp(argv[1], L"-testRegisterVirtualFile") == 0)
+	{
+		wchar_t data[3];
+		if (ReadTestFile(data, 6, L"First.txt") == -1)
+			return -1;
+		if (memcmp(data, L"Foo", 6) != 0)
+			return -1;
+		if (ReadTestFile(data, 6, L"Second.txt") == -1)
+			return -1;
+		if (memcmp(data, L"Bar", 6) != 0)
+			return -1;
+	}
 	else
 	{
 		if (!detoursHandle)
