[UBA]
* Fixed so file mapping for virtual source files are closed

diff --git a/Engine/Source/Programs/UnrealBuildAccelerator/Common/Private/UbaSession.cpp b/Engine/Source/Programs/UnrealBuildAccelerator/Common/Private/UbaSession.cpp
--- a/Engine/Source/Programs/UnrealBuildAccelerator/Common/Private/UbaSession.cpp
+++ b/Engine/Source/Programs/UnrealBuildAccelerator/Common/Private/UbaSession.cpp
@@ -1868,7 +1868,13 @@ namespace uba
 		m_fileMappingGcThread.Wait();
 
 		for (auto& kv : m_virtualSourceFiles)
-			m_sharedMemory.CloseHandle(m_logger, kv.second.memoryHandle, TC("VirtualFile"));
+		{
+			VirtualSourceFile& file = kv.second;
+			if (file.mappingHandle.IsValid())
+				CloseFileMapping(m_logger, file.mappingHandle, TC("VirtualFile"));
+			if (file.memoryHandle.IsValid())
+				m_sharedMemory.CloseHandle(m_logger, file.memoryHandle, TC("VirtualFile"));
+		}
 
 		for (auto& i : m_fileMappingTableLookup)
 			if (i.second.canBeFreed)
@@ -3618,10 +3624,10 @@ namespace uba
 			FileHandle fileHandle = uba::CreateFileW(sourceFile, GENERIC_READ, FILE_SHARE_READ, OPEN_EXISTING, DefaultAttributes());
 			if (fileHandle == InvalidFileHandle)
 				return m_logger.Error(TC("[RegisterVirtualFileInternal] CreateFileW for %s failed (%s)"), sourceFile, LastErrorToText().data);
+			auto fg = MakeGuard([&]() { uba::CloseFile(sourceFile, fileHandle); });
 			u64 fileSize = 0;
 			if (!GetFileSizeEx(fileSize, fileHandle))
 				return m_logger.Error(TC("[RegisterVirtualFileInternal] GetFileSizeEx for %s failed (%s)"), sourceFile, LastErrorToText().data);
-			auto fg = MakeGuard([&]() { uba::CloseFile(sourceFile, fileHandle); });
 			virtualFile.mappingHandle = CreateFileMappingW(m_logger, fileHandle, PAGE_READONLY, fileSize, sourceFile);
 			virtualFile.size = fileSize;
 			insres.first->second = virtualFile;
