[UBA]
* Moved GetMappingString to SharedMemoryAllocator.h/cpp

diff --git a/Engine/Source/Programs/UnrealBuildAccelerator/Common/Private/UbaSession.cpp b/Engine/Source/Programs/UnrealBuildAccelerator/Common/Private/UbaSession.cpp
--- a/Engine/Source/Programs/UnrealBuildAccelerator/Common/Private/UbaSession.cpp
+++ b/Engine/Source/Programs/UnrealBuildAccelerator/Common/Private/UbaSession.cpp
@@ -409,7 +409,7 @@ namespace uba
 				if (!entry.isInvisible)
 				{
 					StringBuffer<> name;
-					Storage::GetMappingString(name, memoryHandle, 0, entry.canBeFreed);
+					GetMappingString(name, memoryHandle, 0, entry.canBeFreed);
 
 					SCOPED_WRITE_LOCK(m_fileMappingTableMemLock, lock);
 					BinaryWriter writer(m_fileMappingTableMem, m_fileMappingTableSize);
@@ -829,7 +829,7 @@ namespace uba
 			out.size = entry.contentSize;
 			if (entry.memoryHandle.IsValid())
 			{
-				Storage::GetMappingString(out.name, entry.memoryHandle, entry.memoryOffset, entry.canBeFreed);
+				GetMappingString(out.name, entry.memoryHandle, entry.memoryOffset, entry.canBeFreed);
 				updateRequestingProcess();
 			}
 			else
@@ -980,7 +980,7 @@ namespace uba
 			entry.canBeFreed = canBeFreed;
 			entry.memoryOffset = memoryOffset;
 			entry.memoryHandle = memoryHandle;
-			Storage::GetMappingString(out.name, memoryHandle, memoryOffset, entry.canBeFreed);
+			GetMappingString(out.name, memoryHandle, memoryOffset, entry.canBeFreed);
 
 			if (canBeFreed)
 			{
@@ -1039,7 +1039,7 @@ namespace uba
 				return false;
 			out.size = entry.contentSize;
 			if (entry.memoryHandle.IsValid())
-				Storage::GetMappingString(out.name, entry.memoryHandle, entry.memoryOffset, entry.canBeFreed);
+				GetMappingString(out.name, entry.memoryHandle, entry.memoryOffset, entry.canBeFreed);
 			else
 				out.name.Append(entry.isDir ? TC("$d") : TC("$f"));
 			return true;
@@ -1079,7 +1079,7 @@ namespace uba
 
 				entry.memoryOffset = mappedViewWrite.offset;
 				entry.memoryHandle = mappedViewWrite.handle;
-				Storage::GetMappingString(out.name, mappedViewWrite.handle, mappedViewWrite.offset, entry.canBeFreed);
+				GetMappingString(out.name, mappedViewWrite.handle, mappedViewWrite.offset, entry.canBeFreed);
 			}
 			mvrg.Execute();
 		}
@@ -1087,7 +1087,7 @@ namespace uba
 		{
 			UBA_ASSERT(mappedViewRead.memory == nullptr);
 			entry.memoryOffset = mappedViewRead.offset;
-			Storage::GetMappingString(out.name, mappedViewRead.handle, mappedViewRead.offset, entry.canBeFreed);
+			GetMappingString(out.name, mappedViewRead.handle, mappedViewRead.offset, entry.canBeFreed);
 			entry.memoryHandle = mappedViewRead.handle;
 			contentSize = mappedViewRead.size;
 		}
@@ -3010,7 +3010,7 @@ namespace uba
 				{
 					UBA_ASSERT(entry.isInvisible);
 					out.size = entry.contentSize;
-					Storage::GetMappingString(out.fileName, entry.memoryHandle, entry.memoryOffset, entry.canBeFreed);
+					GetMappingString(out.fileName, entry.memoryHandle, entry.memoryOffset, entry.canBeFreed);
 					return true;
 				}
 			}
@@ -3692,7 +3692,7 @@ namespace uba
 		entry.isInvisible = transient;
 
 		StringBuffer<> mappingName;
-		Storage::GetMappingString(mappingName, memoryHandle, mappingOffset, entry.canBeFreed);
+		GetMappingString(mappingName, memoryHandle, mappingOffset, entry.canBeFreed);
 
 		entry.success = true;
 
diff --git a/Engine/Source/Programs/UnrealBuildAccelerator/Common/Private/UbaSessionClient.cpp b/Engine/Source/Programs/UnrealBuildAccelerator/Common/Private/UbaSessionClient.cpp
--- a/Engine/Source/Programs/UnrealBuildAccelerator/Common/Private/UbaSessionClient.cpp
+++ b/Engine/Source/Programs/UnrealBuildAccelerator/Common/Private/UbaSessionClient.cpp
@@ -665,7 +665,7 @@ namespace uba
 
 					fileSize = entry.contentSize;
 					if (entry.memoryHandle.IsValid())
-						Storage::GetMappingString(newName, entry.memoryHandle, entry.memoryOffset, entry.canBeFreed);
+						GetMappingString(newName, entry.memoryHandle, entry.memoryOffset, entry.canBeFreed);
 					else
 						newName.Append(entry.isDir ? TC("$d") : TC("$f"));
 				}
@@ -2631,7 +2631,7 @@ namespace uba
 					{
 						StringBuffer<128> mappingName;
 						if (entry.mapping.IsValid())
-							Storage::GetMappingString(mappingName, entry.mapping, entry.mappingOffset);
+							GetMappingString(mappingName, entry.mapping, entry.mappingOffset);
 						else
 							mappingName.Append(TCV("Not valid"));
 						logStr.Appendf(TC(" Success: %s Size: %u IsDir: %s Mapping name: %s Mapping offset: %u"), ToString(entry.success), entry.size, ToString(entry.isDir), mappingName.data, entry.mappingOffset);
diff --git a/Engine/Source/Programs/UnrealBuildAccelerator/Common/Private/UbaSharedMemoryAllocator.cpp b/Engine/Source/Programs/UnrealBuildAccelerator/Common/Private/UbaSharedMemoryAllocator.cpp
--- a/Engine/Source/Programs/UnrealBuildAccelerator/Common/Private/UbaSharedMemoryAllocator.cpp
+++ b/Engine/Source/Programs/UnrealBuildAccelerator/Common/Private/UbaSharedMemoryAllocator.cpp
@@ -841,4 +841,14 @@ namespace uba
 	}
 
 	////////////////////////////////////////////////////////////////////////////////////////////////////
+
+	void GetMappingString(StringBufferBase& out, SharedMemoryHandle handle, u64 offset, bool canBeFreed)
+	{
+		out.Append(canBeFreed ? FreeableMemoryHandleChar : MemoryHandleChar);
+		out.AppendBase62(handle.internalHandle);
+		if (offset)
+			out.Append('-').AppendBase62(offset);
+	}
+
+	////////////////////////////////////////////////////////////////////////////////////////////////////
 }
diff --git a/Engine/Source/Programs/UnrealBuildAccelerator/Common/Private/UbaStorage.cpp b/Engine/Source/Programs/UnrealBuildAccelerator/Common/Private/UbaStorage.cpp
--- a/Engine/Source/Programs/UnrealBuildAccelerator/Common/Private/UbaStorage.cpp
+++ b/Engine/Source/Programs/UnrealBuildAccelerator/Common/Private/UbaStorage.cpp
@@ -22,14 +22,6 @@ namespace uba
 
     CasKey EmptyFileKey = []() { CasKeyHasher hasher; return ToCasKey(hasher, false); }();
 
-	void Storage::GetMappingString(StringBufferBase& out, SharedMemoryHandle handle, u64 offset, bool canBeFreed)
-	{
-		out.Append(canBeFreed ? FreeableMemoryHandleChar : MemoryHandleChar);
-		out.AppendBase62(handle.internalHandle);
-		if (offset)
-			out.Append('-').AppendBase62(offset);
-	}
-
 	StorageCreateInfo::StorageCreateInfo(const tchar* rootDir_, LogWriter& w, WorkManager& wm)
 	:	writer(w)
 	,	workManager(wm)
diff --git a/Engine/Source/Programs/UnrealBuildAccelerator/Common/Public/UbaSharedMemoryAllocator.h b/Engine/Source/Programs/UnrealBuildAccelerator/Common/Public/UbaSharedMemoryAllocator.h
--- a/Engine/Source/Programs/UnrealBuildAccelerator/Common/Public/UbaSharedMemoryAllocator.h
+++ b/Engine/Source/Programs/UnrealBuildAccelerator/Common/Public/UbaSharedMemoryAllocator.h
@@ -105,4 +105,8 @@ namespace uba
 	};
 
 	////////////////////////////////////////////////////////////////////////////////////////////////////
+
+	void GetMappingString(StringBufferBase& out, SharedMemoryHandle handle, u64 offset, bool canBeFreed);
+
+	////////////////////////////////////////////////////////////////////////////////////////////////////
 }
diff --git a/Engine/Source/Programs/UnrealBuildAccelerator/Common/Public/UbaStorage.h b/Engine/Source/Programs/UnrealBuildAccelerator/Common/Public/UbaStorage.h
--- a/Engine/Source/Programs/UnrealBuildAccelerator/Common/Public/UbaStorage.h
+++ b/Engine/Source/Programs/UnrealBuildAccelerator/Common/Public/UbaStorage.h
@@ -78,8 +78,6 @@ namespace uba
 		virtual void AddStats(StorageStats& stats) = 0;
 		virtual void TraceUpdate(Trace& trace, u32 startRow) = 0;
 
-		static void GetMappingString(StringBufferBase& out, SharedMemoryHandle handle, u64 offset, bool canBeFreed);
-
 		virtual void SetTrace(Trace* trace, bool detailed) {}
 		virtual void Ping() {}
 
diff --git a/Engine/Source/Programs/UnrealBuildAccelerator/Core/Private/UbaSharedMemoryView.cpp b/Engine/Source/Programs/UnrealBuildAccelerator/Core/Private/UbaSharedMemoryView.cpp
--- a/Engine/Source/Programs/UnrealBuildAccelerator/Core/Private/UbaSharedMemoryView.cpp
+++ b/Engine/Source/Programs/UnrealBuildAccelerator/Core/Private/UbaSharedMemoryView.cpp
@@ -235,5 +235,22 @@ namespace uba
 		return *name == FreeableMemoryHandleChar;
 	}
 
+	void GetMappingHandleAndOffset(const wchar_t* str, u64& outHandle, u64& outOffset)
+	{
+		const wchar_t* handleStr = str + 1;
+		const wchar_t* handleStrEnd = wcschr(handleStr, '-');
+		outOffset = 0;
+		if (!handleStrEnd)
+		{
+			handleStrEnd = handleStr + wcslen(handleStr);
+		}
+		else
+		{
+			const wchar_t* mappingOffsetStr = handleStrEnd + 1;
+			outOffset = StringToValueBase62(mappingOffsetStr, wcslen(mappingOffsetStr));
+		}
+		outHandle = StringToValueBase62(handleStr, handleStrEnd - handleStr);
+	};
+
 	////////////////////////////////////////////////////////////////////////////////////////////////////
 }
diff --git a/Engine/Source/Programs/UnrealBuildAccelerator/Core/Public/UbaSharedMemoryView.h b/Engine/Source/Programs/UnrealBuildAccelerator/Core/Public/UbaSharedMemoryView.h
--- a/Engine/Source/Programs/UnrealBuildAccelerator/Core/Public/UbaSharedMemoryView.h
+++ b/Engine/Source/Programs/UnrealBuildAccelerator/Core/Public/UbaSharedMemoryView.h
@@ -66,6 +66,7 @@ namespace uba
 
 	bool IsMemoryHandle(const tchar* name);
 	bool IsFreeableMemoryHandle(const tchar* name);
+	void GetMappingHandleAndOffset(const wchar_t* str, u64& outHandle, u64& outOffset);
 	inline constexpr tchar MemoryHandleChar = '^';
 	inline constexpr tchar FreeableMemoryHandleChar = '>';
 	inline constexpr tchar WrittenMemoryHandleChar = ':';
diff --git a/Engine/Source/Programs/UnrealBuildAccelerator/Detours/Private/Windows/UbaDetoursFunctionsNtDll.inl b/Engine/Source/Programs/UnrealBuildAccelerator/Detours/Private/Windows/UbaDetoursFunctionsNtDll.inl
--- a/Engine/Source/Programs/UnrealBuildAccelerator/Detours/Private/Windows/UbaDetoursFunctionsNtDll.inl
+++ b/Engine/Source/Programs/UnrealBuildAccelerator/Detours/Private/Windows/UbaDetoursFunctionsNtDll.inl
@@ -1192,28 +1192,11 @@ NTSTATUS NTAPI Shared_NtCreateFile(bool IsCreateFunc, PHANDLE hFileHandle, ACCES
 		return STATUS_SUCCESS;
 	}
 
-	auto GetHandleAndOffset = [&](u64& outHandle, u64& outOffset)
-		{
-			const wchar_t* handleStr = lpFileName + 1;
-			const wchar_t* handleStrEnd = wcschr(handleStr, '-');
-			outOffset = 0;
-			if (!handleStrEnd)
-			{
-				handleStrEnd = handleStr + wcslen(handleStr);
-			}
-			else
-			{
-				const wchar_t* mappingOffsetStr = handleStrEnd + 1;
-				outOffset = StringToValueBase62(mappingOffsetStr, wcslen(mappingOffsetStr));
-			}
-			outHandle = StringToValueBase62(handleStr, handleStrEnd - handleStr);
-		};
-
 	if (IsMemoryHandle(lpFileName)) // It is a HANDLE from session process
 	{
 		u64 mappingHandle;
 		u64 mappingOffset;
-		GetHandleAndOffset(mappingHandle, mappingOffset);
+		GetMappingHandleAndOffset(lpFileName, mappingHandle, mappingOffset);
 
 		if (mappingHandle & 0x8000'0000'0000'0000)
 		{
@@ -1277,9 +1260,11 @@ NTSTATUS NTAPI Shared_NtCreateFile(bool IsCreateFunc, PHANDLE hFileHandle, ACCES
 	{
 		if (!info.memoryFile)
 		{
-			const wchar_t* handleStr = lpFileName + 1;
-			UBA_ASSERT(!wcschr(handleStr, '-'));
-			auto memoryHandle = SharedMemoryHandle::FromU64(StringToValueBase62(handleStr, wcslen(handleStr)));
+			u64 mappingHandle;
+			u64 mappingOffset;
+			GetMappingHandleAndOffset(lpFileName, mappingHandle, mappingOffset);
+			UBA_ASSERT(mappingOffset == 0);
+			auto memoryHandle = SharedMemoryHandle::FromU64(mappingHandle);
 			UBA_ASSERT(info.size != InvalidValue);
 			u64 mappingHandleSize = info.size;
 
