[UBA]
* Fixed posix compile error

diff --git a/Engine/Source/Programs/UnrealBuildAccelerator/Core/Private/UbaSharedMemoryView.cpp b/Engine/Source/Programs/UnrealBuildAccelerator/Core/Private/UbaSharedMemoryView.cpp
--- a/Engine/Source/Programs/UnrealBuildAccelerator/Core/Private/UbaSharedMemoryView.cpp
+++ b/Engine/Source/Programs/UnrealBuildAccelerator/Core/Private/UbaSharedMemoryView.cpp
@@ -3,6 +3,7 @@
 #include "UbaSharedMemoryView.h"
 #include "UbaBinaryReaderWriter.h"
 #include "UbaHash.h"
+#include "UbaMemory.h"
 #include "UbaProcessStats.h"
 #include "UbaTimer.h"
 
@@ -235,19 +236,19 @@ namespace uba
 		return *name == FreeableMemoryHandleChar;
 	}
 
-	void GetMappingHandleAndOffset(const wchar_t* str, u64& outHandle, u64& outOffset)
+	void GetMappingHandleAndOffset(const tchar* str, u64& outHandle, u64& outOffset)
 	{
-		const wchar_t* handleStr = str + 1;
-		const wchar_t* handleStrEnd = wcschr(handleStr, '-');
+		const tchar* handleStr = str + 1;
+		const tchar* handleStrEnd = TStrchr(handleStr, '-');
 		outOffset = 0;
 		if (!handleStrEnd)
 		{
-			handleStrEnd = handleStr + wcslen(handleStr);
+			handleStrEnd = handleStr + TStrlen(handleStr);
 		}
 		else
 		{
-			const wchar_t* mappingOffsetStr = handleStrEnd + 1;
-			outOffset = StringToValueBase62(mappingOffsetStr, wcslen(mappingOffsetStr));
+			const tchar* mappingOffsetStr = handleStrEnd + 1;
+			outOffset = StringToValueBase62(mappingOffsetStr, TStrlen(mappingOffsetStr));
 		}
 		outHandle = StringToValueBase62(handleStr, handleStrEnd - handleStr);
 	};
diff --git a/Engine/Source/Programs/UnrealBuildAccelerator/Core/Public/UbaSharedMemoryView.h b/Engine/Source/Programs/UnrealBuildAccelerator/Core/Public/UbaSharedMemoryView.h
--- a/Engine/Source/Programs/UnrealBuildAccelerator/Core/Public/UbaSharedMemoryView.h
+++ b/Engine/Source/Programs/UnrealBuildAccelerator/Core/Public/UbaSharedMemoryView.h
@@ -66,7 +66,7 @@ namespace uba
 
 	bool IsMemoryHandle(const tchar* name);
 	bool IsFreeableMemoryHandle(const tchar* name);
-	void GetMappingHandleAndOffset(const wchar_t* str, u64& outHandle, u64& outOffset);
+	void GetMappingHandleAndOffset(const tchar* str, u64& outHandle, u64& outOffset);
 	inline constexpr tchar MemoryHandleChar = '^';
 	inline constexpr tchar FreeableMemoryHandleChar = '>';
 	inline constexpr tchar WrittenMemoryHandleChar = ':';
