[UBA]
* Fixed bug where external mappings are not properly released

diff --git a/Engine/Source/Programs/UnrealBuildAccelerator/Common/Private/UbaSharedMemoryAllocator.cpp b/Engine/Source/Programs/UnrealBuildAccelerator/Common/Private/UbaSharedMemoryAllocator.cpp
--- a/Engine/Source/Programs/UnrealBuildAccelerator/Common/Private/UbaSharedMemoryAllocator.cpp
+++ b/Engine/Source/Programs/UnrealBuildAccelerator/Common/Private/UbaSharedMemoryAllocator.cpp
@@ -368,6 +368,7 @@ namespace uba
 			// .. but only used in very special situations
 			FileMappingHandle file = FileMappingHandle::FromU64(handle.internalHandle & ~0x8000'0000'0000'0000);
 			u8* mem = MapViewOfFile(m_logger, file, FILE_MAP_READ, offset, size);
+			UBA_ASSERT(mem);
 			view.ReferenceExternalMemory(mem, size);
 			return true;
 		}
@@ -561,6 +562,18 @@ namespace uba
 		return h;
 	}
 
+	bool SharedMemoryAllocator::IsExternalMapping(SharedMemoryHandle handle)
+	{
+		return (handle.internalHandle & 0x8000'0000'0000'0000) != 0;
+	}
+
+	void SharedMemoryAllocator::UnmapExternalMapping(SharedMemoryHandle handle, u8* memory, u64 size)
+	{
+		UBA_ASSERT((handle.internalHandle & 0x8000'0000'0000'0000) != 0);
+		bool res = UnmapViewOfFile(m_logger, memory, size, TC("ExternalMapping"));
+		UBA_ASSERT(res); (void)res;
+	}
+
 	SharedMemoryAllocator::InternalHandle& SharedMemoryAllocator::GetHandle(SharedMemoryHandle h)
 	{
 		UBA_ASSERT(!(h.internalHandle & 0x8000'0000'0000'0000));
diff --git a/Engine/Source/Programs/UnrealBuildAccelerator/Common/Private/UbaStorage.cpp b/Engine/Source/Programs/UnrealBuildAccelerator/Common/Private/UbaStorage.cpp
--- a/Engine/Source/Programs/UnrealBuildAccelerator/Common/Private/UbaStorage.cpp
+++ b/Engine/Source/Programs/UnrealBuildAccelerator/Common/Private/UbaStorage.cpp
@@ -83,6 +83,8 @@ namespace uba
 		{
 			if (!mappedMem)
 				return;
+			if (allocator.IsExternalMapping(handle))
+				allocator.UnmapExternalMapping(handle, mappedMem, mapSize);
 			view.Reset();
 		}
 
diff --git a/Engine/Source/Programs/UnrealBuildAccelerator/Common/Public/UbaSharedMemoryAllocator.h b/Engine/Source/Programs/UnrealBuildAccelerator/Common/Public/UbaSharedMemoryAllocator.h
--- a/Engine/Source/Programs/UnrealBuildAccelerator/Common/Public/UbaSharedMemoryAllocator.h
+++ b/Engine/Source/Programs/UnrealBuildAccelerator/Common/Public/UbaSharedMemoryAllocator.h
@@ -60,6 +60,8 @@ namespace uba
 
 		// This is really ugly. But very contained for virtual file mappings
 		SharedMemoryHandle RegisterExternalMapping(FileMappingHandle file);
+		bool IsExternalMapping(SharedMemoryHandle handle);
+		void UnmapExternalMapping(SharedMemoryHandle handle, u8* memory, u64 size);
 
 	private:
 		struct InternalHandle;
