[UBA]
* Added support for threaded send in tcp backend.. where one thread handles all the socket sends

diff --git a/Engine/Source/Programs/UnrealBuildAccelerator/Common/Private/UbaNetworkBackendMemory.cpp b/Engine/Source/Programs/UnrealBuildAccelerator/Common/Private/UbaNetworkBackendMemory.cpp
--- a/Engine/Source/Programs/UnrealBuildAccelerator/Common/Private/UbaNetworkBackendMemory.cpp
+++ b/Engine/Source/Programs/UnrealBuildAccelerator/Common/Private/UbaNetworkBackendMemory.cpp
@@ -133,6 +133,10 @@ namespace uba
 	{
 	}
 
+	void NetworkBackendMemory::SetThreadedSend(void* connection)
+	{
+	}
+
 	void NetworkBackendMemory::SetDisconnectCallback(void* connection, void* context, DisconnectCallback* callback)
 	{
 		UBA_ASSERT(m_connection);
diff --git a/Engine/Source/Programs/UnrealBuildAccelerator/Common/Private/UbaNetworkBackendQuic.cpp b/Engine/Source/Programs/UnrealBuildAccelerator/Common/Private/UbaNetworkBackendQuic.cpp
--- a/Engine/Source/Programs/UnrealBuildAccelerator/Common/Private/UbaNetworkBackendQuic.cpp
+++ b/Engine/Source/Programs/UnrealBuildAccelerator/Common/Private/UbaNetworkBackendQuic.cpp
@@ -165,6 +165,10 @@ namespace uba
 	{
 	}
 
+	void NetworkBackendQuic::SetThreadedSend(void* connection)
+	{
+	}
+
 	void NetworkBackendQuic::SetDisconnectCallback(void* connection, void* context, DisconnectCallback* callback)
 	{
 		auto& conn = *(Connection*)connection;
diff --git a/Engine/Source/Programs/UnrealBuildAccelerator/Common/Private/UbaNetworkBackendTcp.cpp b/Engine/Source/Programs/UnrealBuildAccelerator/Common/Private/UbaNetworkBackendTcp.cpp
--- a/Engine/Source/Programs/UnrealBuildAccelerator/Common/Private/UbaNetworkBackendTcp.cpp
+++ b/Engine/Source/Programs/UnrealBuildAccelerator/Common/Private/UbaNetworkBackendTcp.cpp
@@ -26,6 +26,7 @@
 #include <arpa/inet.h>
 #include <ifaddrs.h>
 #include <poll.h>
+#include <sys/uio.h>
 #define TIMEVAL timeval
 #define SOCKET_ERROR -1
 #define SOCKET int
@@ -52,8 +53,6 @@
 #define UBA_USE_OVERLAPPED_SEND_WITH_LOCK 0
 #define UBA_USE_IOCP PLATFORM_WINDOWS
 
-#define UBA_EXPERIMENTAL_WRITE_QUEUE 0
-
 namespace uba
 {
 	constexpr u32 MaxHeaderSize = 33;
@@ -87,18 +86,33 @@ namespace uba
 
 	////////////////////////////////////////////////////////////////////////////////////////////////////
 
-	#if UBA_EXPERIMENTAL_WRITE_QUEUE
 	struct SendBlock
 	{
-		const void* data;
+		const u8* data;
 		u32 dataLen;
 		SendBlock* next;
+		Event sent;
+	};
+
+	struct ThreadedSendContext
+	{
+		SendBlock* firstBlock = nullptr;
+		SendBlock* lastBlock = nullptr;
+		Event blockQueueEvent;
+		Futex blockQueueLock;
+		Thread thread;
 	};
-	#endif
 
 	struct NetworkBackendTcp::Connection
 	{
 		Connection(Logger& l, SOCKET s) : logger(l), socket(s), ready(EventResetType_Manual) { CreateGuid(uid); }
+		~Connection()
+		{
+			delete threadedSendContext;
+			#if PLATFORM_WINDOWS
+			if (wsaEvent != WSA_INVALID_EVENT) WSACloseEvent(wsaEvent);
+			#endif
+		}
 
 		Logger& logger;
 		Atomic<SOCKET> socket;
@@ -144,14 +158,28 @@ namespace uba
 		bool receivingHeader = true;
 		#endif
 
-		#if UBA_EXPERIMENTAL_WRITE_QUEUE
-		SendBlock* firstBlock = nullptr;
-		SendBlock* lastBlock = nullptr;
-		Futex blockQueueLock;
-		#endif
+		ThreadedSendContext* threadedSendContext = nullptr;
 
 		Connection(const Connection&) = delete;
 		void operator=(const Connection&) = delete;
+
+		#if PLATFORM_WINDOWS
+		WSAEVENT wsaEvent = WSA_INVALID_EVENT;
+		#endif
+
+		void CreateWSAEvent(SOCKET s)
+		{
+			#if PLATFORM_WINDOWS
+			if (IsRunningWine())
+				return;
+			WSAEVENT ev = WSACreateEvent();
+			if (ev == WSA_INVALID_EVENT)
+				return;
+			if (WSAEventSelect(s, ev, FD_WRITE | FD_CLOSE) == SOCKET_ERROR)
+				return;
+			wsaEvent = ev;
+			#endif
+		}
 	};
 
 	////////////////////////////////////////////////////////////////////////////////////////////////////
@@ -282,9 +310,14 @@ namespace uba
 				continue;
 			SOCKET s = conn.socket;
 			conn.socket = INVALID_SOCKET;
+			auto context = conn.threadedSendContext;
+			if (context)
+				context->blockQueueEvent.Set();
 			ShutdownSocket(conn.logger, s, TC("Dtor"));
 			lock2.Leave();
 			conn.recvThread.Wait();
+			if (context)
+				context->thread.Wait();
 			CloseSocket(conn.logger, s, TC("Dtor"));
 		}
 		m_connections.clear();
@@ -365,6 +398,17 @@ namespace uba
 		conn.recvTimeoutCallback = callback;
 	}
 
+	void NetworkBackendTcp::SetThreadedSend(void* connection)
+	{
+		auto& conn = *(Connection*)connection;
+		if (conn.threadedSendContext)
+			return;
+		auto context = new ThreadedSendContext();
+		context->blockQueueEvent.Create(EventResetType_Auto);
+		context->thread.Start([this, connPtr = &conn] { ThreadSend(*connPtr); return 0; }, TC("UbaTcpSnd"));
+		conn.threadedSendContext = context;
+	}
+
 	void NetworkBackendTcp::SetDisconnectCallback(void* connection, void* context, DisconnectCallback* callback)
 	{
 		auto& conn = *(Connection*)connection;
@@ -599,6 +643,7 @@ namespace uba
 			SCOPED_FUTEX(m_connectionsLock, lock);
 			auto it = m_connections.emplace(m_connections.end(), logger, clientSocket);
 			auto& conn = *it;
+			conn.CreateWSAEvent(clientSocket);
 			#if UBA_USE_IOCP
 			if (m_iocpHandle)
 				CreateIoCompletionPort((HANDLE)clientSocket, m_iocpHandle, (ULONG_PTR)&conn, 0);
@@ -689,6 +734,10 @@ namespace uba
 			SCOPED_FUTEX(connection.sendLock, lock);
 			connection.socket = INVALID_SOCKET;
 		}
+
+		if (auto context = connection.threadedSendContext)
+			context->blockQueueEvent.Set();
+
 		if (auto cb = connection.disconnectCallback)
 		{
 			auto context = connection.disconnectContext;
@@ -866,6 +915,7 @@ namespace uba
 		SCOPED_FUTEX(m_connectionsLock, lock);
 		auto it = m_connections.emplace(m_connections.end(), logger, socketFd);
 		auto& conn = *it;
+		conn.CreateWSAEvent(socketFd);
 		#if UBA_USE_IOCP
 		if (m_iocpHandle)
 			CreateIoCompletionPort((HANDLE)socketFd, m_iocpHandle, (ULONG_PTR)&conn, 0);
@@ -880,6 +930,7 @@ namespace uba
 			ShutdownSocket(logger, conn.socket, TC("Connect"));
 			conn.ready.Set();
 			conn.recvThread.Wait();
+			UBA_ASSERT(!conn.threadedSendContext);
 			SCOPED_FUTEX(m_connectionsLock, lock2);
 			m_connections.erase(it);
 			return false;
@@ -942,104 +993,231 @@ namespace uba
 		}
 	}
 
-	bool NetworkBackendTcp::SendSocket(Connection& connection, Logger& logger, const void* b, u64 bufferLen, const tchar* hint)
+	void NetworkBackendTcp::ThreadSend(Connection& connection)
 	{
-		#if UBA_EMULATE_BAD_INTERNET
-		if ((rand() % 10000) == 0)
-		{
-			connection.logger.Info(TC("BAD INTERNET"));
-			Sleep(10000);
-		}
-		#endif
+		SetThreadPriority(GetCurrentThread(), THREAD_PRIORITY_ABOVE_NORMAL);
 
-#if UBA_EXPERIMENTAL_WRITE_QUEUE
-		if (true)
-		{
-			SendBlock newBlock { b, u32(bufferLen), nullptr };
+		auto& logger = connection.logger;
 
-			SCOPED_FUTEX(connection.blockQueueLock, queueLock);
-			if (connection.lastBlock)
-				connection.lastBlock->next = &newBlock;
-			else
-				connection.firstBlock = &newBlock;
-			connection.lastBlock = &newBlock;
-			queueLock.Leave();
+		if (!connection.ready.IsSet(60000)) // This should never time out!
+			return;
 
-			SCOPED_FUTEX(connection.sendLock, sendLock);
-			//if (newBlock.dataLen == 0)
-			//	return true;
+		SendBlock* first = nullptr;
+		SendBlock* last = nullptr;
 
-			queueLock.Enter();
-			SendBlock* first = connection.firstBlock;
-			if (!first)
-				return true;
-			connection.firstBlock = nullptr;
-			connection.lastBlock = nullptr;
+		const tchar* hint = TC("ThreadSend");
+
+		auto& context = *connection.threadedSendContext;
+
+		auto closeGuard = MakeGuard([&]()
+			{
+				for (;first;first=first->next)
+					first->sent.Set();
+				SCOPED_FUTEX(context.blockQueueLock, queueLock);
+				first = context.firstBlock;
+				context.firstBlock = (SendBlock*)~0ull;
+				context.lastBlock = nullptr;
+				queueLock.Leave();
+				for (; first; first = first->next)
+					first->sent.Set();
+			});
+
+		while (connection.socket != INVALID_SOCKET)
+		{
+			SCOPED_FUTEX(context.blockQueueLock, queueLock);
+			SendBlock* newFirst = context.firstBlock;
+			SendBlock* newLast = context.lastBlock;
+			context.firstBlock = nullptr;
+			context.lastBlock = nullptr;
 			queueLock.Leave();
 
-			while (true)
+			if (!first && !newFirst)
 			{
-				WSABUF buf[256];
-				u32 count = 0;
-				u32 totalSend = 0;
-				while (first)
+				context.blockQueueEvent.IsSet();
+				continue;
+			}
+
+			if (newFirst)
+			{
+				if (first)
 				{
-					if (totalSend + u32(first->dataLen) > 128*1024*1024)
-						break;
-					buf[count++] = { u32(first->dataLen), (char*)first->data };
-					totalSend += first->dataLen;
-					first = first->next;
-					if (count == 255)
-						break;
+					UBA_ASSERT(last);
+					last->next = newFirst;
+					last = newLast;
 				}
-
-				if (count > 20)
+				else
 				{
-					logger.Info(TC("Multisend: %u"), count);
+					first = newFirst;
+					last = newLast;
 				}
+			}
 
-				Event ev(true);
-				OVERLAPPED overlapped {};
-				overlapped.hEvent = ev.GetHandle();
+			SOCKET socket = connection.socket;
+			if (socket == INVALID_SOCKET)
+				break;
 
-				SOCKET socket = connection.socket;
-				if (socket == INVALID_SOCKET)
-					return false;
-				int sent = WSASend(socket, buf, count, NULL, 0, &overlapped, NULL);
-				if (sent == SOCKET_ERROR)
+		#if PLATFORM_WINDOWS
+			WSABUF buf[256];
+			u32 count = 0;
+			for (SendBlock* it = first; it && count < 255; it = it->next, ++count)
+				buf[count] = { u32(it->dataLen), (char*)it->data };
+			DWORD bytesSent = 0;
+			int res = WSASend(socket, buf, count, &bytesSent, 0, NULL, NULL);
+		#else
+			iovec buf[IOV_MAX];
+			u32 count = 0;
+			for (SendBlock* it = first; it && count < IOV_MAX; it = it->next, ++count)
+				buf[count] = { (void*)it->data, it->dataLen };
+			ssize_t res = writev(socket, buf, count);
+			u64 bytesSent = (u64)res;
+		#endif
+
+			if (res != SOCKET_ERROR)
+			{
+				while (bytesSent)
 				{
-					u32 lastError = WSAGetLastError();
-					if (lastError != WSA_IO_PENDING)
+					if (bytesSent >= first->dataLen)
 					{
-						#if UBA_LOG_SOCKET_ERRORS
-						logger.Info(TC("SendSocket - send returned an error for socket %i: %s (%s)"), socket, LastErrorToText(WSAGetLastError()).data, hint);
-						#endif
-						return false;
+						bytesSent -= first->dataLen;
+						auto temp = first;
+						first = first->next;
+						temp->dataLen = 0;
+						temp->sent.Set();
+					}
+					else
+					{
+						// only part of block was sent
+						first->dataLen -= bytesSent;
+						first->data += bytesSent;
+						bytesSent = 0;
 					}
 				}
+				continue;
+			}
 
-				ev.IsSet();
+			if (!HandleSendError(connection, logger, hint))
+				return;
+		}
+	}
 
-				DWORD bytesSent;
-				DWORD flags;
-				if (!WSAGetOverlappedResult(socket, &overlapped, &bytesSent, FALSE, &flags))
+	bool NetworkBackendTcp::HandleSendError(Connection& connection, Logger& logger, const tchar* hint)
+	{
+		SOCKET socket = connection.socket;
+		if (socket == INVALID_SOCKET)
+			return false;
+
+		u32 lastError = WSAGetLastError();
+		if (lastError != WSAEWOULDBLOCK)
+		{
+			#if UBA_LOG_SOCKET_ERRORS
+			logger.Info(TC("SendSocket - send returned an error for socket %i: %s (%s)"), socket, LastErrorToText(WSAGetLastError()).data, hint);
+			#endif
+			return false;
+		}
+
+		u32 timeoutMs = 40 * 1000; // 40 seconds timeout to begin with
+
+#if PLATFORM_WINDOWS
+		if (connection.wsaEvent != WSA_INVALID_EVENT)
+		{
+			for (u32 i = 0; i != 2; ++i)
+			{
+				DWORD wr = WSAWaitForMultipleEvents(1, &connection.wsaEvent, FALSE, timeoutMs, FALSE);
+				if (wr == WSA_WAIT_TIMEOUT)
+				{
+					if (i == 0)
+					{
+						logger.Info(TC("SendSocket - WSAWaitForMultipleEvents returned timeout after 40 seconds for socket %i (%s)"), socket, hint);
+						timeoutMs = DefaultNetworkSendTimeoutSeconds * 1000; // 10 minutes send timeout
+						continue;
+					}
+					logger.Info(TC("SendSocket - WSAWaitForMultipleEvents returned timeout after 10 minutes for socket %i (%s)"), socket, hint);
+					return false;
+				}
+				if (wr == WSA_WAIT_FAILED)
 				{
 					#if UBA_LOG_SOCKET_ERRORS
-					logger.Info(TC("WSAGetOverlappedResult - error for socket %i: %s (%s)"), socket, LastErrorToText().data, hint);
+					logger.Info(TC("SendSocket - WSAWaitForMultipleEvents returned an error for socket %i: %s (%s)"), socket, LastErrorToText(WSAGetLastError()).data, hint);
 					#endif
-					return false;
 				}
-
-				if (bytesSent != totalSend)
+				WSANETWORKEVENTS ne{};
+				if (WSAEnumNetworkEvents(socket, connection.wsaEvent, &ne) == SOCKET_ERROR)
 				{
-					logger.Info(TC("ERROR TRIED TO SEND %llu but sent %u"), totalSend, bytesSent);
+					#if UBA_LOG_SOCKET_ERRORS
+					logger.Info(TC("SendSocket - WSAEnumNetworkEvents returned an error for socket %i: %s (%s)"), socket, LastErrorToText(WSAGetLastError()).data, hint);
+					#endif
+					return false;
 				}
-				if (!first)
-					return true;
-				ev.Reset();
+				if (ne.lNetworkEvents & FD_CLOSE)
+					return false;
 			}
 		}
+		else
 #endif
+		{
+			for (u32 i = 0; i != 2; ++i)
+			{
+				WSAPOLLFD p;
+				p.fd = socket;
+				p.revents = 0;
+				p.events = POLLWRNORM;
+				int res = WSAPoll(&p, 1, int(timeoutMs));
+				if (!res)
+				{
+					if (i == 0)
+					{
+						logger.Info(TC("SendSocket - WSAPoll returned timeout after 40 seconds for socket %i (%s)"), socket, hint);
+						timeoutMs = DefaultNetworkSendTimeoutSeconds * 1000; // 10 minutes send timeout
+						continue;
+					}
+					logger.Info(TC("SendSocket - WSAPoll returned timeout after 10 minutes for socket %i (%s)"), socket, hint);
+					return false;
+				}
+				if (res == SOCKET_ERROR)
+				{
+					#if UBA_LOG_SOCKET_ERRORS
+					logger.Info(TC("SendSocket - WSAPoll returned an error for socket %i: %s (%s)"), socket, LastErrorToText(WSAGetLastError()).data, hint);
+					#endif
+					return false;
+				}
+				break;
+			}
+		}
+		return true;
+	}
+
+	bool NetworkBackendTcp::SendSocket(Connection& connection, Logger& logger, const void* b, u64 bufferLen, const tchar* hint)
+	{
+		#if UBA_EMULATE_BAD_INTERNET
+		if ((rand() % 10000) == 0)
+		{
+			connection.logger.Info(TC("BAD INTERNET"));
+			Sleep(10000);
+		}
+		#endif
+
+		if (auto context = connection.threadedSendContext)
+		{
+			SendBlock newBlock { (const u8*)b, u32(bufferLen), nullptr };
+
+			newBlock.sent.Create(EventResetType_Manual);
+
+			SCOPED_FUTEX(context->blockQueueLock, queueLock);
+			if (context->lastBlock)
+				context->lastBlock->next = &newBlock;
+			else
+			{
+				if (context->firstBlock == (SendBlock*)~0ull)
+					return false;
+				context->firstBlock = &newBlock;
+				context->blockQueueEvent.Set();
+			}
+			context->lastBlock = &newBlock;
+			queueLock.Leave();
+
+			newBlock.sent.IsSet();
+			return newBlock.dataLen == 0;
+		}
 
 
 #if UBA_USE_OVERLAPPED_SEND
@@ -1152,50 +1330,15 @@ namespace uba
 		while (left)
 		{
 			int sent = (int)send(socket, (char*)buffer, u32(left), 0);
-			if (sent == SOCKET_ERROR)
+			if (sent != SOCKET_ERROR)
 			{
-				u32 lastError = WSAGetLastError();
-				if (lastError == WSAEWOULDBLOCK)
-				{
-					u64 timeoutMs = 40 * 1000; // 40 seconds timeout to begin with
-
-					for (u32 i=0; i!=2; ++i)
-					{
-						WSAPOLLFD p;
-						p.fd = socket;
-						p.revents = 0;
-						p.events = POLLWRNORM;
-						int res = WSAPoll(&p, 1, int(timeoutMs));
-						if (!res)
-						{
-							if (i == 0)
-							{
-								logger.Info(TC("SendSocket - WSAPoll returned timeout after 40 seconds for socket %i (%s)"), socket, hint);
-								timeoutMs = DefaultNetworkSendTimeoutSeconds * 1000; // 10 minutes send timeout
-								continue;
-							}
-							logger.Info(TC("SendSocket - WSAPoll returned timeout after 10 minutes for socket %i (%s)"), socket, hint);
-							return false;
-						}
-						if (res == SOCKET_ERROR)
-						{
-							#if UBA_LOG_SOCKET_ERRORS
-							logger.Info(TC("SendSocket - WSAPoll returned an error for socket %i: %s (%s)"), socket, LastErrorToText(WSAGetLastError()).data, hint);
-							#endif
-							return false;
-						}
-						break;
-					}
-					continue;
-				}
-				#if UBA_LOG_SOCKET_ERRORS
-				logger.Info(TC("SendSocket - send returned an error for socket %i: %s (%s)"), socket, LastErrorToText(WSAGetLastError()).data, hint);
-				#endif
-				return false;
+				buffer += sent;
+				left -= sent;
+				continue;
 			}
 
-			buffer += sent;
-			left -= sent;
+			if (!HandleSendError(connection, logger, hint))
+				return false;
 		}
 		return true;
 	}
@@ -1950,7 +2093,7 @@ namespace uba
 		optname = SO_RCVBUFFORCE;
 		#endif
 		if (setsockopt(socket, SOL_SOCKET, optname, (const char*)&windowSize, sizeof(windowSize)) == SOCKET_ERROR)
-			return logger.Error(TC("setsockopt TCP_NODELAY error: (error: %s)"), LastErrorToText(WSAGetLastError()).data);
+			return logger.Error(TC("setsockopt SO_RCVBUF error: (error: %s)"), LastErrorToText(WSAGetLastError()).data);
 		return true;
 	}
 
@@ -1961,7 +2104,7 @@ namespace uba
 		optname = SO_SNDBUFFORCE;
 		#endif
 		if (setsockopt(socket, SOL_SOCKET, optname, (const char*)&windowSize, sizeof(windowSize)) == SOCKET_ERROR)
-			return logger.Error(TC("setsockopt TCP_NODELAY error: (error: %s)"), LastErrorToText(WSAGetLastError()).data);
+			return logger.Error(TC("setsockopt SO_SNDBUF error: (error: %s)"), LastErrorToText(WSAGetLastError()).data);
 		return true;
 	}
 
diff --git a/Engine/Source/Programs/UnrealBuildAccelerator/Common/Private/UbaNetworkServer.cpp b/Engine/Source/Programs/UnrealBuildAccelerator/Common/Private/UbaNetworkServer.cpp
--- a/Engine/Source/Programs/UnrealBuildAccelerator/Common/Private/UbaNetworkServer.cpp
+++ b/Engine/Source/Programs/UnrealBuildAccelerator/Common/Private/UbaNetworkServer.cpp
@@ -24,6 +24,7 @@ namespace uba
 		table.GetValueAsBool(logConnections, TC("LogConnections"));
 		table.GetValueAsBool(useKeepAlive, TC("UseKeepAlive"));
 		table.GetValueAsBool(useMessagePriority, TC("UseMessagePriority"));
+		table.GetValueAsBool(useThreadedSend, TC("UseThreadedSend"));
 	}
 
 	struct NetworkServer::Client
@@ -795,6 +796,7 @@ namespace uba
 		m_logConnections = info.logConnections;
 		m_useKeepAlive = info.useKeepAlive;
 		m_useMessagePriority = info.useMessagePriority;
+		m_useThreadedSend = info.useThreadedSend;
 
 		m_recvBufferSize = info.recvBufferSize;
 		m_sendBufferSize = info.sendBufferSize;
@@ -1578,6 +1580,8 @@ namespace uba
 				backend.SetRecvBufferSize(backendConnection, m_recvBufferSize);
 			if (m_sendBufferSize)
 				backend.SetSendBufferSize(backendConnection, m_sendBufferSize);
+			if (m_useThreadedSend)
+				backend.SetThreadedSend(backendConnection);
 		}
 
 		m_connections.emplace_back(*this, backend, backendConnection, remoteSocketAddr, requiresCrypto, cryptoKey, m_connectionIdCounter++);
diff --git a/Engine/Source/Programs/UnrealBuildAccelerator/Common/Public/UbaNetworkBackend.h b/Engine/Source/Programs/UnrealBuildAccelerator/Common/Public/UbaNetworkBackend.h
--- a/Engine/Source/Programs/UnrealBuildAccelerator/Common/Public/UbaNetworkBackend.h
+++ b/Engine/Source/Programs/UnrealBuildAccelerator/Common/Public/UbaNetworkBackend.h
@@ -49,6 +49,9 @@ namespace uba
 		using RecvTimeoutCallback = bool(void* context, u32 timeoutMs, const tchar* recvHint, const tchar* hint);
 		virtual void SetRecvTimeout(void* connection, u32 timeoutMs, void* context, RecvTimeoutCallback* callback) = 0;
 
+		// Change so send is handled in a separate thread instead of directly on calling thread
+		virtual void SetThreadedSend(void* connection) = 0;
+
 		// Disconnect callback. This is called as soon as connection is interrupted from send, recv or shutdown.
 		using DisconnectCallback = void(void* context, const Guid& connectionUid, void* connection);
 		virtual void SetDisconnectCallback(void* connection, void* context, DisconnectCallback* callback) = 0;
diff --git a/Engine/Source/Programs/UnrealBuildAccelerator/Common/Public/UbaNetworkBackendMemory.h b/Engine/Source/Programs/UnrealBuildAccelerator/Common/Public/UbaNetworkBackendMemory.h
--- a/Engine/Source/Programs/UnrealBuildAccelerator/Common/Public/UbaNetworkBackendMemory.h
+++ b/Engine/Source/Programs/UnrealBuildAccelerator/Common/Public/UbaNetworkBackendMemory.h
@@ -18,6 +18,7 @@ namespace uba
 		virtual void SetDataSentCallback(void* connection, void* context, DataSentCallback* callback) override;
 		virtual void SetRecvCallbacks(void* connection, void* context, u32 headerSize, RecvHeaderCallback* h, RecvBodyCallback* b, const tchar* recvHint) override;
 		virtual void SetRecvTimeout(void* connection, u32 timeoutMs, void* context, RecvTimeoutCallback* callback) override;
+		virtual void SetThreadedSend(void* connection) override;
 		virtual void SetDisconnectCallback(void* connection, void* context, DisconnectCallback* callback) override;
 		virtual void SetAllowLessThanBodySize(void* connection, bool allow) override;
 		virtual void SetPriority(void* connection) override;
diff --git a/Engine/Source/Programs/UnrealBuildAccelerator/Common/Public/UbaNetworkBackendQuic.h b/Engine/Source/Programs/UnrealBuildAccelerator/Common/Public/UbaNetworkBackendQuic.h
--- a/Engine/Source/Programs/UnrealBuildAccelerator/Common/Public/UbaNetworkBackendQuic.h
+++ b/Engine/Source/Programs/UnrealBuildAccelerator/Common/Public/UbaNetworkBackendQuic.h
@@ -35,6 +35,7 @@ namespace uba
 		virtual void SetDataSentCallback(void* connection, void* context, DataSentCallback* callback) override;
 		virtual void SetRecvCallbacks(void* connection, void* context, u32 headerSize, RecvHeaderCallback* h, RecvBodyCallback* b, const wchar_t* recvHint) override;
 		virtual void SetRecvTimeout(void* connection, u32 timeoutMs) override;
+		virtual void SetThreadedSend(void* connection) override;
 		virtual void SetDisconnectCallback(void* connection, void* context, DisconnectCallback* callback) override;
 
 		virtual bool StartListen(Logger& logger, u16 port, const wchar_t* ip, const ListenConnectedFunc& connectedFunc) override;
diff --git a/Engine/Source/Programs/UnrealBuildAccelerator/Common/Public/UbaNetworkBackendTcp.h b/Engine/Source/Programs/UnrealBuildAccelerator/Common/Public/UbaNetworkBackendTcp.h
--- a/Engine/Source/Programs/UnrealBuildAccelerator/Common/Public/UbaNetworkBackendTcp.h
+++ b/Engine/Source/Programs/UnrealBuildAccelerator/Common/Public/UbaNetworkBackendTcp.h
@@ -39,6 +39,7 @@ namespace uba
 		virtual void SetDataSentCallback(void* connection, void* context, DataSentCallback* callback) override;
 		virtual void SetRecvCallbacks(void* connection, void* context, u32 headerSize, RecvHeaderCallback* h, RecvBodyCallback* b, const tchar* recvHint) override;
 		virtual void SetRecvTimeout(void* connection, u32 timeoutMs, void* context, RecvTimeoutCallback* callback) override;
+		virtual void SetThreadedSend(void* connection) override;
 		virtual void SetDisconnectCallback(void* connection, void* context, DisconnectCallback* callback) override;
 		virtual void SetAllowLessThanBodySize(void* connection, bool allow) override;
 		virtual void SetPriority(void* connection) override;
@@ -65,6 +66,8 @@ namespace uba
 		struct RecvCache;
 		bool ThreadListen(Logger& logger, ListenEntry& entry);
 		void ThreadRecv(Connection& connection);
+		void ThreadSend(Connection& connection);
+		bool HandleSendError(Connection& connection, Logger& logger, const tchar* hint);
 		bool SendSocket(Connection& connection, Logger& logger, const void* b, u64 bufferLen, const tchar* hint);
 		bool RecvSocket(Connection& connection, RecvCache& recvCache, void* b, u32& bufferLen, const tchar* hint, bool isFirstCall, bool allowLess);
 
diff --git a/Engine/Source/Programs/UnrealBuildAccelerator/Common/Public/UbaNetworkServer.h b/Engine/Source/Programs/UnrealBuildAccelerator/Common/Public/UbaNetworkServer.h
--- a/Engine/Source/Programs/UnrealBuildAccelerator/Common/Public/UbaNetworkServer.h
+++ b/Engine/Source/Programs/UnrealBuildAccelerator/Common/Public/UbaNetworkServer.h
@@ -42,6 +42,7 @@ namespace uba
 		bool logConnections = true;
 		bool useKeepAlive = false;
 		bool useMessagePriority = false;
+		bool useThreadedSend = false;
 
 		u32 recvBufferSize = 0;
 		u32 sendBufferSize = 0;
@@ -183,6 +184,7 @@ namespace uba
 		bool m_logConnections = true;
 		bool m_useKeepAlive = false;
 		bool m_useMessagePriority = false;
+		bool m_useThreadedSend = false;
 		bool m_allClientsDisconnected = false;
 
 		Futex m_addConnectionsLock;
diff --git a/Engine/Source/Programs/UnrealBuildAccelerator/Host/Private/UbaExports.cpp b/Engine/Source/Programs/UnrealBuildAccelerator/Host/Private/UbaExports.cpp
--- a/Engine/Source/Programs/UnrealBuildAccelerator/Host/Private/UbaExports.cpp
+++ b/Engine/Source/Programs/UnrealBuildAccelerator/Host/Private/UbaExports.cpp
@@ -254,6 +254,7 @@ extern "C"
 		info.workerCount = workerCount;
 		info.sendSize = sendSize;
 		info.receiveTimeoutSeconds = receiveTimeoutSeconds;
+		info.useThreadedSend = true;
 		info.Apply(GetConfig());
 
 		bool success = true;
