[UBA]
* Changed so process finished message is not waited on by client, just fire and forget (preparation for local process racing)

diff --git a/Engine/Source/Programs/UnrealBuildAccelerator/Common/Private/UbaNetworkClient.cpp b/Engine/Source/Programs/UnrealBuildAccelerator/Common/Private/UbaNetworkClient.cpp
--- a/Engine/Source/Programs/UnrealBuildAccelerator/Common/Private/UbaNetworkClient.cpp
+++ b/Engine/Source/Programs/UnrealBuildAccelerator/Common/Private/UbaNetworkClient.cpp
@@ -912,6 +912,12 @@ namespace uba
 		return m_client->Send(*this, nullptr, 0, false);
 	}
 
+	bool NetworkMessage::Send(Timer& outTimer)
+	{
+		TimerScope ts(outTimer);
+		return Send();
+	}
+
 	bool NetworkMessage::Send(BinaryReader& response)
 	{
 		if (!m_client->Send(*this, (u8*)response.GetPositionData(), u32(response.GetLeft()), false))
@@ -923,8 +929,7 @@ namespace uba
 	bool NetworkMessage::Send(BinaryReader& response, Timer& outTimer)
 	{
 		TimerScope ts(outTimer);
-		bool res = Send(response);
-		return res;
+		return Send(response);
 	}
 
 	bool NetworkMessage::SendAsync(BinaryReader& response, DoneFunc* func, void* userData)
diff --git a/Engine/Source/Programs/UnrealBuildAccelerator/Common/Private/UbaProcess.cpp b/Engine/Source/Programs/UnrealBuildAccelerator/Common/Private/UbaProcess.cpp
--- a/Engine/Source/Programs/UnrealBuildAccelerator/Common/Private/UbaProcess.cpp
+++ b/Engine/Source/Programs/UnrealBuildAccelerator/Common/Private/UbaProcess.cpp
@@ -44,8 +44,6 @@ extern char **environ;
 
 //////////////////////////////////////////////////////////////////////////////
 
-#define UBA_EXIT_CODE(x) (9000 + x)
-
 namespace uba
 {
 	static constexpr const tchar g_extractExportsStr[] = TC("/extractexports");
diff --git a/Engine/Source/Programs/UnrealBuildAccelerator/Common/Private/UbaSessionClient.cpp b/Engine/Source/Programs/UnrealBuildAccelerator/Common/Private/UbaSessionClient.cpp
--- a/Engine/Source/Programs/UnrealBuildAccelerator/Common/Private/UbaSessionClient.cpp
+++ b/Engine/Source/Programs/UnrealBuildAccelerator/Common/Private/UbaSessionClient.cpp
@@ -1623,8 +1623,7 @@ namespace uba
 
 		exitTime = oldExitTime;
 
-		StackBinaryReader<16> reader;
-		if (!msg.Send(reader, m_stats.procFinishedMsg) && m_loop)
+		if (!msg.Send(m_stats.procFinishedMsg) && m_loop)
 			return m_logger.Error(TC("Failed to send ProcessFinished message!"));
 		return true;
 	}
@@ -1824,7 +1823,7 @@ namespace uba
 		dest.Append(TCV("<log>")).Append(logFile);
 		f.name = dest.data;
 		f.key = ToStringKeyLower(dest);
-		SendFile(f, pi.GetId(), false, true);
+		SendFile(f, 0, false, true);
 		for (auto& child : pi.m_childProcesses)
 			SendLogFileToServer(*(ProcessImpl*)child.m_process);
 	}
diff --git a/Engine/Source/Programs/UnrealBuildAccelerator/Common/Private/UbaSessionServer.cpp b/Engine/Source/Programs/UnrealBuildAccelerator/Common/Private/UbaSessionServer.cpp
--- a/Engine/Source/Programs/UnrealBuildAccelerator/Common/Private/UbaSessionServer.cpp
+++ b/Engine/Source/Programs/UnrealBuildAccelerator/Common/Private/UbaSessionServer.cpp
@@ -12,6 +12,15 @@
 
 namespace uba
 {
+	struct SessionServer::ReceivedFile
+	{
+		u32 processId;
+		CasKey casKey;
+		StringKey destinationKey;
+		TString destination;
+		u32 attributes;
+	};
+
 	class SessionServer::RemoteProcess final : public Process
 	{
 	public:
@@ -95,6 +104,7 @@ namespace uba
 		u64 m_peakMemory = 0;
 		Event m_done;
 		Vector<ProcessLogLine> m_logLines;
+		Vector<ReceivedFile> m_receivedFiles;
 		Vector<u8> m_trackedInputs;
 		Vector<u8> m_trackedOutputs;
 		bool m_cancelled = false;
@@ -122,6 +132,8 @@ namespace uba
 			table->GetValueAsBool(remoteTraceEnabled, TC("RemoteTraceEnabled"));
 			table->GetValueAsBool(nameToHashTableEnabled, TC("NameToHashTableEnabled"));
 			table->GetValueAsBool(traceIOEnabled, TC("TraceIOEnabled"));
+			table->GetValueAsBool(customCasKeysEnabled, TC("CustomCasKeysEnabled"));
+
 		}
 	}
 
@@ -185,6 +197,7 @@ namespace uba
 		m_remoteLogEnabled = info.remoteLogEnabled;
 		m_remoteTraceEnabled = info.remoteTraceEnabled;
 		m_traceIOEnabled = info.traceIOEnabled;
+		m_customCasKeysEnabled = info.customCasKeysEnabled;
 
 		if (m_resetCas)
 			m_storage.Reset();
@@ -457,6 +470,11 @@ namespace uba
 
 	void SessionServer::SetCustomCasKeyFromTrackedInputs(const tchar* fileName_, const tchar* workingDir_, const u8* trackedInputs, u32 trackedInputsBytes)
 	{
+		if (!m_customCasKeysEnabled)
+		{
+			m_logger.Error(TC("Custom cas keys not enabled. Enable with config option CustomCasKeysEnabled"));
+			return;
+		}
 		StringBuffer<> workingDir;
 		FixFileName(workingDir, workingDir_, nullptr);
 		if (workingDir[workingDir.count - 1] != '\\')
@@ -1148,34 +1166,17 @@ namespace uba
 		#endif
 	}
 
-	bool SessionServer::HandleSendFileToServer(const ConnectionInfo& connectionInfo, const WorkContext& workContext, BinaryReader& reader, BinaryWriter& writer)
+	bool SessionServer::WriteReceivedFile(ReceivedFile& file, RemoteProcess* process, const ConnectionInfo& connectionInfo)
 	{
 		u32 clientId = connectionInfo.GetId();
-		u32 processId = reader.ReadU32();
-		StringBuffer<> destination;
-		reader.ReadString(destination);
-		StringKey destinationKey = reader.ReadStringKey();
-		u32 attributes = reader.ReadU32();
-		UBA_ASSERT(attributes);
-		CasKey casKey = reader.ReadCasKey();
-		Storage::RetrieveResult res;
-		bool success = m_storage.RetrieveCasFile(res, casKey, destination.data, nullptr, 1, true, clientId);
-		casKey = res.casKey;
-
-		auto writeResponse = MakeGuard([&]() { writer.WriteBool(success); });
-
-		if (!success)
-		{
-			auto logType = connectionInfo.ShouldDisconnect() ? LogEntryType_Info : LogEntryType_Warning;
-			m_logger.Logf(logType, TC("Failed to retrieve cas for %s from client %u (Needed to write %s)"), CasKeyString(casKey).str, clientId, destination.data);
-			return true;
-		}
+		StringView destination = file.destination;
+		CasKey casKey = file.casKey;
 
 		if (destination.StartsWith(TC("<log>")))
 		{
 			StringBuffer<> logPath;
 			logPath.Append(m_sessionLogDir).Append(destination.data + 5);
-			if (!m_storage.CopyOrLink(casKey, logPath.data, attributes, false, clientId))
+			if (!m_storage.CopyOrLink(casKey, logPath.data, file.attributes, false, clientId))
 				m_logger.Error(TC("Failed to copy cas from %s to %s"), CasKeyString(casKey).str, logPath.data);
 			else if (!m_storage.DropCasFile(casKey, false, logPath.data))
 				m_logger.Error(TC("Failed to drop cas %s"), CasKeyString(casKey).str);
@@ -1197,32 +1198,26 @@ namespace uba
 				ubaPath.Append(TCV("Connection")).AppendValue(clientId);
 
 			ubaPath.Append(TCV(".uba"));
-			m_storage.CopyOrLink(casKey, ubaPath.data, attributes, false, clientId);
+			m_storage.CopyOrLink(casKey, ubaPath.data, file.attributes, false, clientId);
 			m_storage.DropCasFile(casKey, false, ubaPath.data);
 			return true;
 		}
 
-		ProcessHandle h = GetProcess(processId);
-		if (!h.IsValid())
-		{
-			m_logger.Info(TC("Failed to find process for id %u when receiving SendFileToServer message"), processId);
-			success = false;
-			return false;
-		}
+		if (!process)
+			return m_logger.Info(TC("Failed to find process for id %u when receiving SendFileToServer message for file %s"), file.processId, destination.data).ToFalse();
 
-		auto& process = *(RemoteProcess*)h.m_process;
-		auto& startInfo = process.GetStartInfo();
+		auto& startInfo = process->GetStartInfo();
 		auto& rules = *startInfo.rules;
 
 		bool shouldWriteToDisk = ShouldWriteToDisk(destination);
 
 		{
-			SCOPED_WRITE_LOCK(process.m_writtenFilesLock, lock);
-			process.m_writtenFiles.try_emplace(destination.ToString(), casKey);
+			SCOPED_WRITE_LOCK(process->m_writtenFilesLock, lock);
+			process->m_writtenFiles.try_emplace(destination.ToString(), casKey);
 		}
 
 		// We need to report before writing to flush out potential deferred cas creations
-		m_storage.ReportFileWrite(destinationKey, destination.data);
+		m_storage.ReportFileWrite(file.destinationKey, destination.data);
 
 		if (shouldWriteToDisk)
 		{
@@ -1244,25 +1239,20 @@ namespace uba
 				writeCompressed = g_globalRules.FileCanBeCompressed(destination);
 			}
 
-			success = m_storage.CopyOrLink(casKey, destination.data, attributes, writeCompressed, clientId, formattingFunc);
-			if (!success)
-				m_logger.Error(TC("Failed to copy cas from %s to %s (%s)"), CasKeyString(casKey).str, destination.data, GetProcessDescription(processId).c_str());
-			else
-				TraceWrittenFile(processId, destination);
+			if (!m_storage.CopyOrLink(casKey, destination.data, file.attributes, writeCompressed, clientId, formattingFunc))
+				return m_logger.Error(TC("Failed to copy cas from %s to %s (%s)"), CasKeyString(casKey).str, destination.data, startInfo.GetDescription());
+			TraceWrittenFile(file.processId, destination);
 		}
 		else
 		{
 			if (!rules.IsInvisible(destination))
-			{
-				success = m_storage.FakeCopy(casKey, destination.data, 0, 0, attributes);
-				if (!success)
-					m_logger.Error(TC("Failed to fake copy cas from %s to %s (%s)"), CasKeyString(casKey).str, destination.data, GetProcessDescription(processId).c_str());
-			}
+				if (!m_storage.FakeCopy(casKey, destination.data, 0, 0, file.attributes))
+					return m_logger.Error(TC("Failed to fake copy cas from %s to %s (%s)"), CasKeyString(casKey).str, destination.data, startInfo.GetDescription());
 			SCOPED_WRITE_LOCK(m_receivedFilesLock, lock);
-			m_receivedFiles.try_emplace(destinationKey, casKey);
+			m_receivedFiles.try_emplace(file.destinationKey, casKey);
 		}
 
-		if (success && !rules.IsInvisible(destination))
+		if (!rules.IsInvisible(destination))
 		{
 			bool invalidateStorage = false; // No need, already handled in m_storage.CopyOrLink
 			RegisterCreateFileForWrite(StringKeyZero, destination, shouldWriteToDisk, 0, 0, invalidateStorage);
@@ -1270,15 +1260,57 @@ namespace uba
 			if (startInfo.trackInputs)
 			{
 				u64 bytes = GetStringWriteSize(destination.data, destination.count);
-				u64 prevSize = process.m_trackedOutputs.size();
-				process.m_trackedOutputs.resize(prevSize + bytes);
-				BinaryWriter w2(process.m_trackedOutputs.data(), prevSize, prevSize + bytes);
+				u64 prevSize = process->m_trackedOutputs.size();
+				process->m_trackedOutputs.resize(prevSize + bytes);
+				BinaryWriter w2(process->m_trackedOutputs.data(), prevSize, prevSize + bytes);
 				w2.WriteString(destination);
 			}
 		}
 		return true;
 	}
 
+	bool SessionServer::HandleSendFileToServer(const ConnectionInfo& connectionInfo, const WorkContext& workContext, BinaryReader& reader, BinaryWriter& writer)
+	{
+		u32 clientId = connectionInfo.GetId();
+
+		ReceivedFile file;
+		file.processId = reader.ReadU32();
+		file.destination = reader.ReadString();
+		file.destinationKey = reader.ReadStringKey();
+		file.attributes = reader.ReadU32();
+		CasKey casKey = reader.ReadCasKey();
+
+		Storage::RetrieveResult res;
+		bool success = m_storage.RetrieveCasFile(res, casKey, file.destination.c_str(), nullptr, 1, true, clientId);
+		auto writeResponse = MakeGuard([&]() { writer.WriteBool(success); });
+
+		if (!success)
+		{
+			auto logType = connectionInfo.ShouldDisconnect() ? LogEntryType_Info : LogEntryType_Warning;
+			m_logger.Logf(logType, TC("Failed to retrieve cas for %s from client %u (Needed to write %s)"), CasKeyString(casKey).str, clientId, file.destination.c_str());
+			return true;
+		}
+		file.casKey = res.casKey;
+
+		if (!file.processId)
+		{
+			success = WriteReceivedFile(file, nullptr, connectionInfo);
+			return true;
+		}
+
+		ProcessHandle h = GetProcess(file.processId);
+		if (!h.IsValid())
+		{
+			m_logger.Info(TC("Failed to find process for id %u when receiving SendFileToServer message"), file.processId);
+			success = false;
+			return false;
+		}
+
+		auto& process = *(RemoteProcess*)h.m_process;
+		process.m_receivedFiles.push_back(file);
+		return true;
+	}
+
 	bool SessionServer::HandleDeleteFile(const ConnectionInfo& connectionInfo, const WorkContext& workContext, BinaryReader& reader, BinaryWriter& writer)
 	{
 		StringKey fileNameKey = reader.ReadStringKey();
@@ -1601,6 +1633,22 @@ namespace uba
 			for (auto& line : process.m_logLines)
 				func(process.m_startInfo.logLineUserData, line.text.c_str(), u32(line.text.size()), line.type);
 
+		Timer writeFilesTimer;
+		if (exitCode == 0)
+		{
+			TimerScope ts(writeFilesTimer);
+			Atomic<bool> writeSuccess = true;
+			m_server.ParallelFor(u32(process.m_receivedFiles.size()), process.m_receivedFiles, [&](const WorkContext&, auto& it)
+				{
+					if (!WriteReceivedFile(*it, &process, connectionInfo))
+						writeSuccess = false;
+					++writeFilesTimer.count;
+				}, AsView(TC("WriteReceivedFiles")));
+			if (!writeSuccess)
+				exitCode = UBA_EXIT_CODE(20);
+			--writeFilesTimer.count;
+		}
+
 		u32 id = process.m_processId;
 		Vector<ProcessLogLine> emptyLines;
 		auto& logLines = (exitCode != 0 || m_detailedTrace) ? process.m_logLines : emptyLines;
@@ -1609,11 +1657,7 @@ namespace uba
 		ProcessStats processStats;
 		processStats.Read(reader, ~0u);
 
-		//SessionStats sessionStats;
-		//sessionStats.Read(reader);
-		//StorageStats storageStats;
-		//storageStats.Read(reader);
-
+		processStats.writeFiles = writeFilesTimer; // The writeFiles stats from client session is not really useful since it is a noop
 		process.m_processorTime = processStats.cpuTime;
 		process.m_wallTime = processStats.wallTime;
 		process.m_peakMemory = processStats.peakJobMemory > 0 ? processStats.peakJobMemory : processStats.peakMemory;
@@ -1913,6 +1957,7 @@ namespace uba
 		CasKey casKeyOverride = CasKeyZero;
 
 		bool deferCreation = true;
+		if (m_customCasKeysEnabled)
 		{
 			SCOPED_FUTEX(m_customCasKeysLock, lock);
 			auto findIt = m_customCasKeys.find(fileNameKey);
diff --git a/Engine/Source/Programs/UnrealBuildAccelerator/Common/Public/UbaNetworkMessage.h b/Engine/Source/Programs/UnrealBuildAccelerator/Common/Public/UbaNetworkMessage.h
--- a/Engine/Source/Programs/UnrealBuildAccelerator/Common/Public/UbaNetworkMessage.h
+++ b/Engine/Source/Programs/UnrealBuildAccelerator/Common/Public/UbaNetworkMessage.h
@@ -20,6 +20,7 @@ namespace uba
 		bool Send(BinaryReader& response);
 		bool Send(BinaryReader& response, Timer& outTimer);
 		bool Send();
+		bool Send(Timer& outTimer);
 
 		using DoneFunc = void(bool error, void* userData);
 		bool SendAsync(BinaryReader& response, DoneFunc* func, void* userData); // Try to minimize work in GetResponseFunc since it is running on receiving thread
diff --git a/Engine/Source/Programs/UnrealBuildAccelerator/Common/Public/UbaProcess.h b/Engine/Source/Programs/UnrealBuildAccelerator/Common/Public/UbaProcess.h
--- a/Engine/Source/Programs/UnrealBuildAccelerator/Common/Public/UbaProcess.h
+++ b/Engine/Source/Programs/UnrealBuildAccelerator/Common/Public/UbaProcess.h
@@ -7,6 +7,8 @@
 #include "UbaSession.h"
 #include "UbaStats.h"
 
+#define UBA_EXIT_CODE(x) (9000 + x)
+
 namespace uba
 {
 	class Process
diff --git a/Engine/Source/Programs/UnrealBuildAccelerator/Common/Public/UbaSessionServer.h b/Engine/Source/Programs/UnrealBuildAccelerator/Common/Public/UbaSessionServer.h
--- a/Engine/Source/Programs/UnrealBuildAccelerator/Common/Public/UbaSessionServer.h
+++ b/Engine/Source/Programs/UnrealBuildAccelerator/Common/Public/UbaSessionServer.h
@@ -123,6 +123,9 @@ namespace uba
 		virtual bool RunSpecialProgram(ProcessImpl& process, BinaryReader& reader, BinaryWriter& writer) override final;
 		virtual void PrintSessionStats(Logger& logger) override final;
 		
+		struct ReceivedFile;
+		bool WriteReceivedFile(ReceivedFile& file, RemoteProcess* process, const ConnectionInfo& connectionInfo);
+
 		void TraceSessionUpdate();
 
 		void WriteRemoteEnvironmentVariables(BinaryWriter& writer);
@@ -216,6 +219,7 @@ namespace uba
 		bool m_remoteLogEnabled = false;
 		bool m_remoteTraceEnabled = false;
 		bool m_traceIOEnabled = false;
+		bool m_customCasKeysEnabled = false;
 
 		Scheduler* m_outerScheduler = nullptr;
 
diff --git a/Engine/Source/Programs/UnrealBuildAccelerator/Common/Public/UbaSessionServerCreateInfo.h b/Engine/Source/Programs/UnrealBuildAccelerator/Common/Public/UbaSessionServerCreateInfo.h
--- a/Engine/Source/Programs/UnrealBuildAccelerator/Common/Public/UbaSessionServerCreateInfo.h
+++ b/Engine/Source/Programs/UnrealBuildAccelerator/Common/Public/UbaSessionServerCreateInfo.h
@@ -21,5 +21,6 @@ namespace uba
 		bool remoteLogEnabled = false; // If Uba is built in debug, then the logs will be sent back to server
 		bool remoteTraceEnabled = false; // If this is true, the agents will run trace and send the .uba file back to server
 		bool traceIOEnabled = true; // Enable IO information for storage 
+		bool customCasKeysEnabled = false; // Need to set this to true to be able to use SetCustomCasKeyFromTrackedInputs
 	};
 }
