[UBA]
* Removed accidentally added line

diff --git a/Engine/Source/Programs/UnrealBuildAccelerator/Common/Private/UbaNetworkBackendTcp.cpp b/Engine/Source/Programs/UnrealBuildAccelerator/Common/Private/UbaNetworkBackendTcp.cpp
--- a/Engine/Source/Programs/UnrealBuildAccelerator/Common/Private/UbaNetworkBackendTcp.cpp
+++ b/Engine/Source/Programs/UnrealBuildAccelerator/Common/Private/UbaNetworkBackendTcp.cpp
@@ -995,8 +995,6 @@ namespace uba
 
 	void NetworkBackendTcp::ThreadSend(Connection& connection)
 	{
-		SetThreadPriority(GetCurrentThread(), THREAD_PRIORITY_ABOVE_NORMAL);
-
 		auto& logger = connection.logger;
 
 		if (!connection.ready.IsSet(60000)) // This should never time out!
