[UBA]
* Fixed bug in wsaevent path in tcp backend that could cause hangs in networking

diff --git a/Engine/Source/Programs/UnrealBuildAccelerator/Common/Private/UbaNetworkBackendTcp.cpp b/Engine/Source/Programs/UnrealBuildAccelerator/Common/Private/UbaNetworkBackendTcp.cpp
--- a/Engine/Source/Programs/UnrealBuildAccelerator/Common/Private/UbaNetworkBackendTcp.cpp
+++ b/Engine/Source/Programs/UnrealBuildAccelerator/Common/Private/UbaNetworkBackendTcp.cpp
@@ -1148,6 +1148,7 @@ namespace uba
 				}
 				if (ne.lNetworkEvents & FD_CLOSE)
 					return false;
+				break;
 			}
 		}
 		else
