[UBA]
* Fixed race condition in code where lookup table was not protected with a lock. Caused a hang in ld-lld.exe

#Synced-CL 48599271

diff --git a/Engine/Source/Programs/UnrealBuildAccelerator/Detours/Private/UbaDetoursShared.cpp b/Engine/Source/Programs/UnrealBuildAccelerator/Detours/Private/UbaDetoursShared.cpp
--- a/Engine/Source/Programs/UnrealBuildAccelerator/Detours/Private/UbaDetoursShared.cpp
+++ b/Engine/Source/Programs/UnrealBuildAccelerator/Detours/Private/UbaDetoursShared.cpp
@@ -586,7 +586,7 @@ namespace uba
 			{
 				StringBuffer<> temp;
 				u32 closeId;
-				Rpc_CreateFileW(fileName, fileNameKey, AccessFlag_Read, temp.data, temp.capacity, fileSize, closeId, false);
+				Rpc_CreateFileW(fileName, fileNameKey, AccessFlag_Read, temp.data, temp.capacity, fileSize, closeId, true);
 			}
 		}
 
