[UBA]
* Added speed test to uba
* Added so speed test can run with -speedtest in UbaCli

diff --git a/Engine/Source/Programs/UnrealBuildAccelerator/Cli/Private/UbaCli.cpp b/Engine/Source/Programs/UnrealBuildAccelerator/Cli/Private/UbaCli.cpp
--- a/Engine/Source/Programs/UnrealBuildAccelerator/Cli/Private/UbaCli.cpp
+++ b/Engine/Source/Programs/UnrealBuildAccelerator/Cli/Private/UbaCli.cpp
@@ -102,6 +102,7 @@ namespace uba
 		logger.Info(TC("   -cachecommand=<cmd>     Send command to cache server. Will output result in log"));
 		logger.Info(TC("   -writecachesummary      Write cache summary file about connected cache server"));
 		logger.Info(TC("   -decompress=<file>      Decompress a compressed file in-place"));
+		logger.Info(TC("   -speedtest              Run a network speedtest"));
 		logger.Info(TC(""));
 		logger.Info(TC("  CoordinatorOptions (if coordinator set):"));
 		logger.Info(TC("   -uri=<address>          Uri to coordinator"));
@@ -184,6 +185,7 @@ namespace uba
 		bool useHackVfs = false;
 		bool useScheduler = false;
 		bool decompressInPlace = false;
+		bool speedTest = false;
 		TString checkFileTable;
 		TString cacheFilterString;
 		TString cacheCommand;
@@ -465,6 +467,10 @@ namespace uba
 				{
 					storeCompressed = false;
 				}
+				else if (name.Equals(TCV("-speedtest")))
+				{
+					speedTest = true;
+				}
 				else if (name.Equals(TCV("-dir")))
 				{
 					if (value.IsEmpty())
@@ -953,6 +959,13 @@ namespace uba
 			return true;
 		}
 
+		if (speedTest)
+		{
+			if (commandType == CommandType_NotSet)
+				commandType = CommandType_Agent;
+			application = TC("speedtest");
+		}
+
 		if (exit)
 			return true;
 
@@ -971,7 +984,7 @@ namespace uba
 			if (application.empty())
 				return PrintHelp(TC("No executable provided"));
 
-			if (!IsAbsolutePath(application.c_str()))
+			if (!IsAbsolutePath(application.c_str()) && !speedTest)
 			{
 				StringBuffer<> fullApplicationName;
 				if (!SearchPathForFile(logger, fullApplicationName, application.c_str(), currentDir, {}))
@@ -1087,8 +1100,11 @@ namespace uba
 		signal(SIGTERM, ConsoleHandler);
 		#endif
 
-		NetworkBackendTcp networkBackend(logWriter);
+		NetworkBackendTcpCreateInfo nbtci(logWriter);
+		NetworkBackendTcp networkBackend(nbtci);
+
 		NetworkServerCreateInfo nsci(logWriter);
+		nsci.useThreadedSend = true;
 		nsci.Apply(config);
 
 		//nsci.workerCount = 4;
@@ -1326,6 +1342,29 @@ namespace uba
 
 		CoordinatorWrapper coordinator;
 
+		auto CreateCoordinator = [&](Scheduler* scheduler)
+			{
+				if (coordinatorName.empty())
+					return true;
+				StringBuffer<512> coordinatorWorkDir(g_rootDir);
+				coordinatorWorkDir.EnsureEndsWithSlash().Append(coordinatorName);
+				StringBuffer<512> binariesDir;
+				if (!GetDirectoryOfCurrentModule(logger, binariesDir))
+					return false;
+
+				CoordinatorCreateInfo cinfo;
+				cinfo.workDir = coordinatorWorkDir.data;
+				cinfo.binariesDir = binariesDir.data;
+
+				// TODO: This is very horde specific.. maybe all these parameters should be a string or something
+				cinfo.pool = coordinatorPool.c_str();
+				cinfo.maxCoreCount = 1;
+				cinfo.logging = false;
+				if (!coordinator.Create(logger, coordinatorName.c_str(), cinfo, networkBackend, networkServer, scheduler))
+					return false;
+				return true;
+			};
+
 		auto RunScheduler = [&](const TString& app, const TString& arg)
 		{
 			/*
@@ -1366,6 +1405,11 @@ namespace uba
 			info.writeToCache = populateCache;
 			Scheduler scheduler(info);
 
+
+			if (!CreateCoordinator(&scheduler))
+				return false;
+
+
 			if (StringView(app).EndsWith(TCV(".yaml")))
 			{
 				if (!scheduler.EnqueueFromFile(app.c_str(), [&](EnqueueProcessInfo& epi)
@@ -1444,26 +1488,6 @@ namespace uba
 				return RunQueue();
 		};
 
-
-		if (!coordinatorName.empty())
-		{
-			StringBuffer<512> coordinatorWorkDir(g_rootDir);
-			coordinatorWorkDir.EnsureEndsWithSlash().Append(coordinatorName);
-			StringBuffer<512> binariesDir;
-			if (!GetDirectoryOfCurrentModule(logger, binariesDir))
-				return false;
-
-			CoordinatorCreateInfo cinfo;
-			cinfo.workDir = coordinatorWorkDir.data;
-			cinfo.binariesDir = binariesDir.data;
-
-			// TODO: This is very horde specific.. maybe all these parameters should be a string or something
-			cinfo.pool = coordinatorPool.c_str();
-			cinfo.maxCoreCount = coordinatorMaxCoreCount;
-			cinfo.logging = true;
-			if (!coordinator.Create(logger, coordinatorName.c_str(), cinfo, networkBackend, networkServer))
-				return false;
-		}
 		auto cg = MakeGuard([&]() { coordinator.Destroy(); });
 
 #if PLATFORM_WINDOWS	// Annoying that link.exe/lld-link.exe needs path to windows folder.. 
@@ -1492,6 +1516,9 @@ namespace uba
 			}
 			else
 			{
+				if (!CreateCoordinator(nullptr))
+					return false;
+
 				switch (commandType)
 				{
 				case CommandType_Native:
diff --git a/Engine/Source/Programs/UnrealBuildAccelerator/Common/Private/UbaSessionClient.cpp b/Engine/Source/Programs/UnrealBuildAccelerator/Common/Private/UbaSessionClient.cpp
--- a/Engine/Source/Programs/UnrealBuildAccelerator/Common/Private/UbaSessionClient.cpp
+++ b/Engine/Source/Programs/UnrealBuildAccelerator/Common/Private/UbaSessionClient.cpp
@@ -1757,7 +1757,7 @@ namespace uba
 		}
 	}
 
-	void SessionClient::SendNotification(const StringView& text)
+	void SessionClient::SendNotification(StringView text)
 	{
 		StackBinaryWriter<1024> writer;
 		NetworkMessage msg(m_client, ServiceId, SessionMessageType_Notification, writer);
@@ -1766,6 +1766,16 @@ namespace uba
 		msg.Send();
 	}
 
+	void SessionClient::SendLogEntry(StringView text, LogEntryType type)
+	{
+		StackBinaryWriter<1024> writer;
+		NetworkMessage msg(m_client, ServiceId, SessionMessageType_LogEntry, writer);
+		writer.WriteByte(type);
+		writer.WriteString(text);
+		StackBinaryReader<16> reader;
+		msg.Send(reader);
+	}
+
 	bool SessionClient::SendRootsHandle(RootsHandle rootsHandle)
 	{
 		SCOPED_FUTEX(m_rootsLookupLock, rootsLock);
@@ -2028,6 +2038,14 @@ namespace uba
 
 				for (InternalProcessStartInfo& startInfo : startInfos)
 				{
+					if (startInfo.applicationStr == TC("speedtest"))
+					{
+						RunSpeedTest();
+						SendReturnProcess(startInfo.processId, TC("speedtest"));
+						continue;
+					}
+
+
 					startInfo.uiLanguage = int(m_uiLanguage);
 					startInfo.priorityClass = m_defaultPriorityClass;
 					startInfo.useCustomAllocator = !m_disableCustomAllocator;
@@ -2336,6 +2354,45 @@ namespace uba
 		return true;
 	}
 
+	void SessionClient::RunSpeedTest()
+	{
+		u64 dataSize = 300ull * 1024 * 1024;
+		u8* dest = new u8[dataSize];
+
+		u64 messageCount = dataSize / (SendMaxSize - 10);
+
+		struct Entry
+		{
+			Entry(NetworkClient& c, BinaryWriter& w, u8* d) : msg(c, ServiceId, SessionMessageType_SpeedTest, w), ev(EventResetType_Manual), reader(d) {}
+			NetworkMessage msg;
+			Event ev;
+			BinaryReader reader;
+		};
+
+		for (u32 j=0; j!=10; ++j)
+		{
+			StringBuffer<> note;
+			note.Appendf(TC("Receiving 300mb (%i)"), j);
+			SendLogEntry(note);
+
+			List<Entry> entries;
+
+			u64 start = GetTime();
+			for (u64 i=0; i!=messageCount;++i)
+			{
+				StackBinaryWriter<16> writer;
+				auto& entry = entries.emplace_back(m_client, writer, dest + (SendMaxSize - 10) * i);
+				writer.WriteByte(0);
+				entry.msg.SendAsync(entry.reader, [](bool err, void* userData) { ((Event*)userData)->Set(); }, &entry.ev);
+			}
+			for (auto& entry : entries)
+				entry.ev.IsSet();
+
+			note.Clear().Appendf(TC(".. took (%s)"), TimeToText(GetTime() - start).str);
+			SendLogEntry(note);
+		}
+	}
+
 	bool SessionClient::AllocFailed(Process& process, const tchar* allocType, u32 error)
 	{
 		//StackBinaryWriter<32> writer;
@@ -2722,8 +2779,8 @@ namespace uba
 		// send and recv are swapped on purpose because that is how visualizer is visualizing 
 		m_trace.SessionUpdate(0, 0, send, recv, m_lastPing, m_memAvail, m_memTotal, m_cpuUsage);
 
-		m_sharedMemory.TraceStats(m_trace, 30);
-		m_fileMappingBuffer.TraceStatus(m_trace, 20, TC("Session"), 1);
-		m_storage.TraceUpdate(m_trace, 300);
+		//m_sharedMemory.TraceStats(m_trace, 30);
+		//m_fileMappingBuffer.TraceStatus(m_trace, 20, TC("Session"), 1);
+		//m_storage.TraceUpdate(m_trace, 300);
 	}
 }
diff --git a/Engine/Source/Programs/UnrealBuildAccelerator/Common/Private/UbaSessionServer.cpp b/Engine/Source/Programs/UnrealBuildAccelerator/Common/Private/UbaSessionServer.cpp
--- a/Engine/Source/Programs/UnrealBuildAccelerator/Common/Private/UbaSessionServer.cpp
+++ b/Engine/Source/Programs/UnrealBuildAccelerator/Common/Private/UbaSessionServer.cpp
@@ -1696,6 +1696,16 @@ namespace uba
 		if (session.enabled)
 			++m_availableRemoteSlotCount;
 
+		if (process->m_startInfo.applicationStr == TC("speedtest"))
+		{
+			m_trace.ProcessReturned(process->m_processId, TCV("speedtest"));
+			process->m_exitCode = 0;
+			process->m_server = nullptr;
+			process->m_done.Set();
+			process->CallProcessExit(h);
+			return true;
+		}
+
 		m_logger.Detail(TC("Client %s returned process %u to queue (%s)"), session.name.c_str(), processId, reason.data);
 		++m_returnedRemoteProcessCount;
 
@@ -2587,4 +2597,18 @@ namespace uba
 
 		return true;
 	}
+
+	bool SessionServer::HandleSpeedTest(const ConnectionInfo& connectionInfo, const WorkContext& workContext, BinaryReader& reader, BinaryWriter& writer)
+	{
+		writer.AllocWrite(SendMaxSize-10);
+		return true;
+	}
+
+	bool SessionServer::HandleLogEntry(const ConnectionInfo& connectionInfo, const WorkContext& workContext, BinaryReader& reader, BinaryWriter& writer)
+	{
+		auto logType = (LogEntryType)reader.ReadByte();
+		TString str = reader.ReadString();
+		m_logger.Log(logType, str.c_str(), u32(str.size()));
+		return true;
+	}
 }
diff --git a/Engine/Source/Programs/UnrealBuildAccelerator/Common/Public/UbaSessionClient.h b/Engine/Source/Programs/UnrealBuildAccelerator/Common/Public/UbaSessionClient.h
--- a/Engine/Source/Programs/UnrealBuildAccelerator/Common/Public/UbaSessionClient.h
+++ b/Engine/Source/Programs/UnrealBuildAccelerator/Common/Public/UbaSessionClient.h
@@ -103,13 +103,15 @@ namespace uba
 		bool SendProcessInputs(ProcessImpl& process);
 		bool SendProcessFinished(ProcessImpl& process, u32 exitCode);
 		void SendPing(u64 memAvail, u64 memTotal);
-		void SendNotification(const StringView& text);
+		void SendNotification(StringView text);
+		void SendLogEntry(StringView text, LogEntryType type = LogEntryType_Info);
 		bool SendRootsHandle(RootsHandle rootsHandle);
 		void SendLogFileToServer(ProcessImpl& pi);
 		void GetLogFileName(StringBufferBase& out, const tchar* logFile, const tchar* arguments, u32 processId);
 		u32 WriteLogLines(BinaryWriter& writer, ProcessImpl& process);
 		bool ParseDirectoryTable();
 		bool EntryExists(const StringView& path, u32& outTableOffset);
+		void RunSpeedTest();
 
 		void ThreadCreateProcessLoop();
 
