[UBA]
* Missed file in previous submit

diff --git a/Engine/Source/Programs/UnrealBuildAccelerator/Common/Public/UbaNetwork.h b/Engine/Source/Programs/UnrealBuildAccelerator/Common/Public/UbaNetwork.h
--- a/Engine/Source/Programs/UnrealBuildAccelerator/Common/Public/UbaNetwork.h
+++ b/Engine/Source/Programs/UnrealBuildAccelerator/Common/Public/UbaNetwork.h
@@ -75,6 +75,8 @@ namespace uba
 		UBA_SESSION_MESSAGE(DebugFileNotFoundError) \
 		UBA_SESSION_MESSAGE(HostRun) \
 		UBA_SESSION_MESSAGE(GetSymbols) \
+		UBA_SESSION_MESSAGE(SpeedTest) \
+		UBA_SESSION_MESSAGE(LogEntry) \
 
 	#define UBA_CACHE_MESSAGES \
 		UBA_CACHE_MESSAGE(Connect) \
