[UBA]
* Fixed uba coordinator horde code for mac

diff --git a/Engine/Source/Programs/UnrealBuildAccelerator/Common/Public/UbaCoordinatorWrapper.h b/Engine/Source/Programs/UnrealBuildAccelerator/Common/Public/UbaCoordinatorWrapper.h
--- a/Engine/Source/Programs/UnrealBuildAccelerator/Common/Public/UbaCoordinatorWrapper.h
+++ b/Engine/Source/Programs/UnrealBuildAccelerator/Common/Public/UbaCoordinatorWrapper.h
@@ -35,6 +35,8 @@ namespace uba
 
 			#if PLATFORM_WINDOWS
 			coordinatorBin.Append(TCV("UbaCoordinator")).Append(coordinatorType).Append(TCV(".dll"));
+			#elif PLATFORM_MAC
+			coordinatorBin.Append(TCV("libUbaCoordinator")).Append(coordinatorType).Append(TCV(".dylib"));
 			#else
 			coordinatorBin.Append(TCV("libUbaCoordinator")).Append(coordinatorType).Append(TCV(".so"));
 			#endif
diff --git a/Engine/Source/Programs/UnrealBuildAccelerator/UbaCoordinatorHorde.Target.cs b/Engine/Source/Programs/UnrealBuildAccelerator/UbaCoordinatorHorde.Target.cs
--- a/Engine/Source/Programs/UnrealBuildAccelerator/UbaCoordinatorHorde.Target.cs
+++ b/Engine/Source/Programs/UnrealBuildAccelerator/UbaCoordinatorHorde.Target.cs
@@ -51,6 +51,15 @@ public class UbaCoordinatorHordeTarget : TargetRules
 			GlobalDefinitions.Add("PLATFORM_WINDOWS=0");
 		}
 
+		if (Target.Platform.IsInGroup(UnrealPlatformGroup.Apple))
+		{
+			OutputFile = Path.Combine(BinariesFolder, $"lib{LaunchModuleName}.dylib");
+		}
+		else
+		{
+			GlobalDefinitions.Add("PLATFORM_MAC=0");
+		}
+
 		if (Target.Platform.IsInGroup(UnrealPlatformGroup.Linux))
 		{
 			OutputFile = Path.Combine(BinariesFolder, $"lib{LaunchModuleName}.so");
