[UBA]
* Fixed bug when compiling with msvc static analysis.. the new code skips unmapping files and in this case it was not ok

diff --git a/Engine/Source/Programs/UnrealBuildAccelerator/Detours/Private/Windows/UbaDetoursFunctionsKernelBase.inl b/Engine/Source/Programs/UnrealBuildAccelerator/Detours/Private/Windows/UbaDetoursFunctionsKernelBase.inl
--- a/Engine/Source/Programs/UnrealBuildAccelerator/Detours/Private/Windows/UbaDetoursFunctionsKernelBase.inl
+++ b/Engine/Source/Programs/UnrealBuildAccelerator/Detours/Private/Windows/UbaDetoursFunctionsKernelBase.inl
@@ -2520,6 +2520,9 @@ HANDLE Detoured_OpenFileMappingW(DWORD dwDesiredAccess, BOOL bInheritHandle, LPC
 	return mappingHandle;
 }
 
+// This is ugly and will break if we have multiple.. but right now it is only used by pch in msvc
+SharedMemoryView g_specialView;
+
 LPVOID Detoured_MapViewOfFileEx(HANDLE hFileMappingObject, DWORD dwDesiredAccess, DWORD dwFileOffsetHigh, DWORD dwFileOffsetLow, SIZE_T dwNumberOfBytesToMap, LPVOID lpBaseAddress)
 {
 	DETOURED_CALL(MapViewOfFileEx);
@@ -2538,7 +2541,8 @@ LPVOID Detoured_MapViewOfFileEx(HANDLE hFileMappingObject, DWORD dwDesiredAccess
 				// This is where cl.exe "continues" its compile from the pch
 				UBA_ASSERT(dwDesiredAccess == FILE_MAP_COPY);
 
-				SharedMemoryView& view = *new SharedMemoryView();
+				UBA_ASSERT(!g_specialView.GetMemory());
+				SharedMemoryView& view = g_specialView;
 				RPC_MESSAGE(GetSharedMemory, getSharedMemory)
 				writer.WriteU64(fi.trueFileMapHandle.ToU64());
 				writer.WriteU64(offset + fi.trueFileMapOffset);
@@ -2620,6 +2624,13 @@ BOOL Detoured_UnmapViewOfFileEx(PVOID lpBaseAddress, ULONG UnmapFlags)
 		}
 	}
 
+	if (g_specialView.GetMemory() == lpBaseAddress)
+	{
+		g_specialView.Reset();
+		SetLastError(ERROR_SUCCESS);
+		return true;
+	}
+
 	auto res = True_UnmapViewOfFileEx(lpBaseAddress, UnmapFlags); (void)res;
 	DEBUG_LOG_TRUE(L"UnmapViewOfFileEx", L"0x%llx -> %ls", uintptr_t(lpBaseAddress), ToString(res));
 
