[UBA]
* Fixed bug when compiling with msvc static analysis.. the new code skips unmapping files and in this case it was not ok

diff --git a/Engine/Source/Programs/UnrealBuildAccelerator/Detours/Private/Windows/UbaDetoursFunctionsKernelBase.inl b/Engine/Source/Programs/UnrealBuildAccelerator/Detours/Private/Windows/UbaDetoursFunctionsKernelBase.inl
--- a/Engine/Source/Programs/UnrealBuildAccelerator/Detours/Private/Windows/UbaDetoursFunctionsKernelBase.inl
+++ b/Engine/Source/Programs/UnrealBuildAccelerator/Detours/Private/Windows/UbaDetoursFunctionsKernelBase.inl
@@ -2520,8 +2520,7 @@ HANDLE Detoured_OpenFileMappingW(DWORD dwDesiredAccess, BOOL bInheritHandle, LPC
 	return mappingHandle;
 }
 
-// This is ugly and will break if we have multiple.. but right now it is only used by pch in msvc
-SharedMemoryView g_specialView;
+List<SharedMemoryView> g_specialViews;
 
 LPVOID Detoured_MapViewOfFileEx(HANDLE hFileMappingObject, DWORD dwDesiredAccess, DWORD dwFileOffsetHigh, DWORD dwFileOffsetLow, SIZE_T dwNumberOfBytesToMap, LPVOID lpBaseAddress)
 {
@@ -2541,8 +2540,7 @@ LPVOID Detoured_MapViewOfFileEx(HANDLE hFileMappingObject, DWORD dwDesiredAccess
 				// This is where cl.exe "continues" its compile from the pch
 				UBA_ASSERT(dwDesiredAccess == FILE_MAP_COPY);
 
-				UBA_ASSERT(!g_specialView.GetMemory());
-				SharedMemoryView& view = g_specialView;
+				SharedMemoryView& view = g_specialViews.emplace_back();
 				RPC_MESSAGE(GetSharedMemory, getSharedMemory)
 				writer.WriteU64(fi.trueFileMapHandle.ToU64());
 				writer.WriteU64(offset + fi.trueFileMapOffset);
@@ -2624,9 +2622,11 @@ BOOL Detoured_UnmapViewOfFileEx(PVOID lpBaseAddress, ULONG UnmapFlags)
 		}
 	}
 
-	if (g_specialView.GetMemory() == lpBaseAddress)
+	for (auto i=g_specialViews.begin(), e=g_specialViews.end(); i!=e; ++i)
 	{
-		g_specialView.Reset();
+		if (i->GetMemory() != lpBaseAddress)
+			continue;
+		g_specialViews.erase(i);
 		SetLastError(ERROR_SUCCESS);
 		return true;
 	}
