[UBA]
* Added asserts in a few places to try to help track down issue on mac

diff --git a/Engine/Source/Programs/UnrealBuildAccelerator/Common/Private/UbaSession.cpp b/Engine/Source/Programs/UnrealBuildAccelerator/Common/Private/UbaSession.cpp
--- a/Engine/Source/Programs/UnrealBuildAccelerator/Common/Private/UbaSession.cpp
+++ b/Engine/Source/Programs/UnrealBuildAccelerator/Common/Private/UbaSession.cpp
@@ -2990,6 +2990,7 @@ namespace uba
 				bool canBeFreed = canBeCompressed;
 				if (GetOrCreateMemoryMapFromFile(map, fileNameKey, fileName.data, false, alignment, TC("CreateFile"), &process, canBeFreed))
 				{
+					UBA_ASSERT(!map.name.IsEmpty());
 					out.size = map.size;
 					out.fileName.Append(map.name);
 				}
diff --git a/Engine/Source/Programs/UnrealBuildAccelerator/Common/Private/UbaSharedMemoryAllocator.cpp b/Engine/Source/Programs/UnrealBuildAccelerator/Common/Private/UbaSharedMemoryAllocator.cpp
--- a/Engine/Source/Programs/UnrealBuildAccelerator/Common/Private/UbaSharedMemoryAllocator.cpp
+++ b/Engine/Source/Programs/UnrealBuildAccelerator/Common/Private/UbaSharedMemoryAllocator.cpp
@@ -511,13 +511,14 @@ namespace uba
 			ReleaseHandle(handle);
 			return;
 		}
-		u64 pos = 0;
-		u32 res; (void)res;
 
 		{
 			auto& kernelStats = KernelStats::GetCurrent();
 
 			#if PLATFORM_WINDOWS
+			u64 pos = 0;
+			u32 res; (void)res;
+
 			for (Slice slice : h.slices)
 			{
 				TimerScope ts(kernelStats.unmapViewOfFile2);
@@ -539,7 +540,8 @@ namespace uba
 			u64 committed = 0;
 			for (Slice slice : h.slices)
 				committed += slice.size;
-			munmap(memory, committed);
+			int res = munmap(memory, committed);
+			UBA_ASSERT(res == 0);(void)res;
 			#endif
 		}
 		ReleaseHandle(handle);
diff --git a/Engine/Source/Programs/UnrealBuildAccelerator/Core/Private/UbaSharedMemoryView.cpp b/Engine/Source/Programs/UnrealBuildAccelerator/Core/Private/UbaSharedMemoryView.cpp
--- a/Engine/Source/Programs/UnrealBuildAccelerator/Core/Private/UbaSharedMemoryView.cpp
+++ b/Engine/Source/Programs/UnrealBuildAccelerator/Core/Private/UbaSharedMemoryView.cpp
@@ -58,7 +58,6 @@ namespace uba
 
 		#if PLATFORM_WINDOWS
 		bool res; (void)res;
-
 		if (m_mappedSize)
 		{
 			auto& kernelStats = KernelStats::GetCurrent();
@@ -89,7 +88,8 @@ namespace uba
 		res = VirtualFree(m_memory, 0, MEM_RELEASE);
 		UBA_ASSERTF(res, TC("VirtualFree failed (%u)"), GetLastError());
 		#else
-		munmap(m_memory, m_capacity); 
+		int res = munmap(m_memory, m_capacity);
+		UBA_ASSERT(res == 0); (void)res;
 		#endif
 
 		m_memory = nullptr;
