[UBA]
* Changed code casting address to *(u64*)ptr = value; to instead be memcpy. Line is causing sigbus on mac and I can't figure out why, my only theory is memory alignment but I can't see how it can't be aligned (and in the debugger I can see it is properly aligned when sigbus happen)

diff --git a/Engine/Source/Programs/UnrealBuildAccelerator/Common/Private/UbaStorageServer.cpp b/Engine/Source/Programs/UnrealBuildAccelerator/Common/Private/UbaStorageServer.cpp
--- a/Engine/Source/Programs/UnrealBuildAccelerator/Common/Private/UbaStorageServer.cpp
+++ b/Engine/Source/Programs/UnrealBuildAccelerator/Common/Private/UbaStorageServer.cpp
@@ -1093,7 +1093,8 @@ namespace uba
 
 				firstStore = &tempStore;
 
-				*(u64*)mappedView.memory = fileSize;
+				//*(u64*)mappedView.memory = fileSize; // This code cause sigbus on mac. don't know why because address is always aligned.
+				memcpy(mappedView.memory, &fileSize, sizeof(u64)); 
 
 				firstStore->casEntry  = &casEntry;
 				firstStore->fileSize = fileSize;
