[UBA]
* Improved asserts in shared memory allocator
* Changed PAGE_READ/PAGE_READWRITE to match correct protection on posix

diff --git a/Engine/Source/Programs/UnrealBuildAccelerator/Common/Private/UbaSharedMemoryAllocator.cpp b/Engine/Source/Programs/UnrealBuildAccelerator/Common/Private/UbaSharedMemoryAllocator.cpp
--- a/Engine/Source/Programs/UnrealBuildAccelerator/Common/Private/UbaSharedMemoryAllocator.cpp
+++ b/Engine/Source/Programs/UnrealBuildAccelerator/Common/Private/UbaSharedMemoryAllocator.cpp
@@ -10,7 +10,7 @@
 #else
 #define ZeroMemory(mem, size) memset(mem, 0, size)
 #define DiscardVirtualMemory(...)
-#define PAGE_WRITECOPY 0x08
+#define PAGE_WRITECOPY PROT_NONE
 #endif
 
 #define UBA_DEBUG_HINTS (UBA_DEBUG || UBA_TRACK_SHARED_MEMORY_ALLOCATIONS)
@@ -489,7 +489,7 @@ namespace uba
 			void* viewA = mmap(targetAddress, slice.size, protect, MAP_SHARED | MAP_FIXED, m_fileMapping.shmFd, slice.offset);
 			#endif
 
-			UBA_ASSERTF(viewA, TC("MapViewOfFile3 failed in MapView 0x%p %llu for % s (%s)"), memory, pos, h.GetHint(), LastErrorToText().data); (void)viewA;
+			UBA_ASSERTF(viewA == targetAddress, TC("MapViewOfFile3 failed in MapView 0x%p %llu for % s (%s)"), memory, pos, h.GetHint(), LastErrorToText().data); (void)viewA;
 			pos += slice.size;
 		}
 
diff --git a/Engine/Source/Programs/UnrealBuildAccelerator/Common/Public/UbaFileMapping.h b/Engine/Source/Programs/UnrealBuildAccelerator/Common/Public/UbaFileMapping.h
--- a/Engine/Source/Programs/UnrealBuildAccelerator/Common/Public/UbaFileMapping.h
+++ b/Engine/Source/Programs/UnrealBuildAccelerator/Common/Public/UbaFileMapping.h
@@ -20,7 +20,7 @@ namespace uba
 	inline constexpr u32 FILE_MAP_READ = 0x0004;
 	inline constexpr u32 FILE_MAP_ALL_ACCESS = FILE_MAP_WRITE | FILE_MAP_READ;
 	inline constexpr u32 SEC_RESERVE = 0x04000000;
-	inline constexpr u32 PAGE_READWRITE = 0x04;
+	inline constexpr u32 PAGE_READWRITE = PROT_READ | PROT_WRITE;
 
 	#if PLATFORM_MAC
 	inline constexpr u32 SHM_MAX_FILENAME = PSHMNAMLEN;
diff --git a/Engine/Source/Programs/UnrealBuildAccelerator/Core/Private/UbaPlatform.cpp b/Engine/Source/Programs/UnrealBuildAccelerator/Core/Private/UbaPlatform.cpp
--- a/Engine/Source/Programs/UnrealBuildAccelerator/Core/Private/UbaPlatform.cpp
+++ b/Engine/Source/Programs/UnrealBuildAccelerator/Core/Private/UbaPlatform.cpp
@@ -602,6 +602,7 @@ namespace uba
 		//#endif
 		//SIGABRT,
 		//SIGFPE,
+		SIGBUS,
 		SIGILL,
 		//SIGINT,
 		SIGSEGV,
diff --git a/Engine/Source/Programs/UnrealBuildAccelerator/Core/Private/UbaSharedMemoryView.cpp b/Engine/Source/Programs/UnrealBuildAccelerator/Core/Private/UbaSharedMemoryView.cpp
--- a/Engine/Source/Programs/UnrealBuildAccelerator/Core/Private/UbaSharedMemoryView.cpp
+++ b/Engine/Source/Programs/UnrealBuildAccelerator/Core/Private/UbaSharedMemoryView.cpp
@@ -207,7 +207,7 @@ namespace uba
 		void* viewA = mmap(targetAddress, size, protect, MAP_SHARED | MAP_FIXED, handle.shmFd, offset);
 		#endif
 
-		UBA_ASSERTF(viewA, TC("MapViewOfFile3 failed - %s (%s)"), LastErrorToText().data, hint);
+		UBA_ASSERTF(viewA == targetAddress, TC("MapViewOfFile3 failed - %s (%s)"), LastErrorToText().data, hint);
 		if (!viewA)
 			return false;
 
diff --git a/Engine/Source/Programs/UnrealBuildAccelerator/Core/Public/UbaPlatform.h b/Engine/Source/Programs/UnrealBuildAccelerator/Core/Public/UbaPlatform.h
--- a/Engine/Source/Programs/UnrealBuildAccelerator/Core/Public/UbaPlatform.h
+++ b/Engine/Source/Programs/UnrealBuildAccelerator/Core/Public/UbaPlatform.h
@@ -185,7 +185,7 @@ namespace uba
 	#define UBA_NOINLINE __attribute__((noinline))
     #define ERROR_SUCCESS 0
 	#define DUPLICATE_SAME_ACCESS 0
-	inline constexpr u32 PAGE_READONLY = 0x02;
+	inline constexpr u32 PAGE_READONLY = PROT_READ;
 
 	struct FileMappingHandle
 	{
