[UBA]
* Improved assert
* Fixed so detoured process is not trying to send back closed files if exit message has been sent

diff --git a/Engine/Source/Programs/UnrealBuildAccelerator/Detours/Private/Windows/UbaDetoursFunctionsNtDll.inl b/Engine/Source/Programs/UnrealBuildAccelerator/Detours/Private/Windows/UbaDetoursFunctionsNtDll.inl
--- a/Engine/Source/Programs/UnrealBuildAccelerator/Detours/Private/Windows/UbaDetoursFunctionsNtDll.inl
+++ b/Engine/Source/Programs/UnrealBuildAccelerator/Detours/Private/Windows/UbaDetoursFunctionsNtDll.inl
@@ -357,7 +357,7 @@ NTSTATUS NTAPI Detoured_NtQueryDirectoryFile(HANDLE FileHandle, HANDLE Event, PI
 			prevInformation = it;
 			it = (u8*)fileNamePos + info.FileNameLength + 2;
 
-			DEBUG_LOG_DETOURED(L"NtQueryDirectoryFile", L"%llu %ls", u64(FileHandle), fileNamePos);
+			DEBUG_LOG_DETOURED(L"NtQueryDirectoryFile", L"%llu %.*s", u64(FileHandle), int(fileNameBytes/sizeof(tchar)), fileNamePos);
 
 			res = STATUS_SUCCESS;
 
@@ -851,6 +851,10 @@ NTSTATUS NTAPI Shared_NtCreateFile(bool IsCreateFunc, PHANDLE hFileHandle, ACCES
 
 	DevirtualizePath(fileName);
 
+
+	if (fileName.EndsWith(L".pchast"))
+		printf("");
+
 	u32 dirTableOffset = ~u32(0);
 
 	//UBA_ASSERT(CreateDisposition != FILE_SUPERSEDE);
@@ -1655,7 +1659,8 @@ NTSTATUS NTAPI Detoured_NtClose(HANDLE handle)
 
 	if (fo->closeId)
 	{
-		Rpc_UpdateCloseHandle(path, fo->closeId, fo->deleteOnClose, fo->newName.c_str(), memoryHandle, memoryWritten, true);
+		if (!g_exitMessageSent)
+			Rpc_UpdateCloseHandle(path, fo->closeId, fo->deleteOnClose, fo->newName.c_str(), memoryHandle, memoryWritten, true);
 	}
 	else
 	{
