[UBA]
* Removed debug code

diff --git a/Engine/Source/Programs/UnrealBuildAccelerator/Detours/Private/Windows/UbaDetoursFunctionsNtDll.inl b/Engine/Source/Programs/UnrealBuildAccelerator/Detours/Private/Windows/UbaDetoursFunctionsNtDll.inl
--- a/Engine/Source/Programs/UnrealBuildAccelerator/Detours/Private/Windows/UbaDetoursFunctionsNtDll.inl
+++ b/Engine/Source/Programs/UnrealBuildAccelerator/Detours/Private/Windows/UbaDetoursFunctionsNtDll.inl
@@ -851,10 +851,6 @@ NTSTATUS NTAPI Shared_NtCreateFile(bool IsCreateFunc, PHANDLE hFileHandle, ACCES
 
 	DevirtualizePath(fileName);
 
-
-	if (fileName.EndsWith(L".pchast"))
-		printf("");
-
 	u32 dirTableOffset = ~u32(0);
 
 	//UBA_ASSERT(CreateDisposition != FILE_SUPERSEDE);
