[UBA]
* Changed code to use EventSlim instead

diff --git a/Engine/Source/Programs/UnrealBuildAccelerator/Common/Private/UbaNetworkClient.cpp b/Engine/Source/Programs/UnrealBuildAccelerator/Common/Private/UbaNetworkClient.cpp
--- a/Engine/Source/Programs/UnrealBuildAccelerator/Common/Private/UbaNetworkClient.cpp
+++ b/Engine/Source/Programs/UnrealBuildAccelerator/Common/Private/UbaNetworkClient.cpp
@@ -728,7 +728,7 @@ namespace uba
 		BinaryWriter& writer = *message.m_sendWriter;
 
 		u16 messageId = 0;
-		Event gotResponse;
+		EventSlim gotResponse;
 
 
 		if (response)
@@ -786,7 +786,7 @@ namespace uba
 				{
 					UBA_ASSERT(!message.m_doneFunc);
 					message.m_doneUserData = &gotResponse;
-					message.m_doneFunc = [](bool error, void* userData) { ((Event*)userData)->Set(); };
+					message.m_doneFunc = [](bool error, void* userData) { ((EventSlim*)userData)->Set(); };
 				}
 				break;
 			}
diff --git a/Engine/Source/Programs/UnrealBuildAccelerator/Common/Private/UbaNetworkServer.cpp b/Engine/Source/Programs/UnrealBuildAccelerator/Common/Private/UbaNetworkServer.cpp
--- a/Engine/Source/Programs/UnrealBuildAccelerator/Common/Private/UbaNetworkServer.cpp
+++ b/Engine/Source/Programs/UnrealBuildAccelerator/Common/Private/UbaNetworkServer.cpp
@@ -54,7 +54,7 @@ namespace uba
 		}
 
 		NetworkServer& server;
-		Event workAvailable;
+		EventSlim workAvailable;
 		u8* writeMem = nullptr;
 		u32 writeMemSize = 0;
 
@@ -1580,9 +1580,11 @@ namespace uba
 				backend.SetRecvBufferSize(backendConnection, m_recvBufferSize);
 			if (m_sendBufferSize)
 				backend.SetSendBufferSize(backendConnection, m_sendBufferSize);
+		}
+
+		if (IsWindows)
 			if (m_useThreadedSend)
 				backend.SetThreadedSend(backendConnection);
-		}
 
 		m_connections.emplace_back(*this, backend, backendConnection, remoteSocketAddr, requiresCrypto, cryptoKey, m_connectionIdCounter++);
 		m_maxActiveConnections = Max(m_maxActiveConnections, u32(m_connections.size()));
diff --git a/Engine/Source/Programs/UnrealBuildAccelerator/Common/Private/UbaStorageUtils.cpp b/Engine/Source/Programs/UnrealBuildAccelerator/Common/Private/UbaStorageUtils.cpp
--- a/Engine/Source/Programs/UnrealBuildAccelerator/Common/Private/UbaStorageUtils.cpp
+++ b/Engine/Source/Programs/UnrealBuildAccelerator/Common/Private/UbaStorageUtils.cpp
@@ -72,7 +72,7 @@ namespace uba
 				u64 fileSize = 0;
 				bool error = false;
 				Vector<CasKey> keys;
-				Event done;
+				EventSlim done;
 			};
 
 			u32 workCount = u32((fileSize + BufferSlotSize - 1) / BufferSlotSize);
@@ -156,7 +156,7 @@ namespace uba
 			Entry(u8* slot, u32 i, u32 messageMaxSize) : reader(slot + i * messageMaxSize, 0, SendMaxSize), done(EventResetType_Manual) {}
 			NetworkMessage message;
 			BinaryReader reader;
-			Event done;
+			EventSlim done;
 		};
 
 		u64 sendCountCapacity = capacity / messageMaxSize;
@@ -196,7 +196,7 @@ namespace uba
 			entry.message.Init(client, StorageServiceId, StorageMessageType_FetchSegment, writer);
 			writer.WriteU16(fetchId);
 			writer.WriteU32(readIndex + i + 1);
-			if (entry.message.SendAsync(entry.reader, [](bool error, void* userData) { ((Event*)userData)->Set(); }, &entry.done))
+			if (entry.message.SendAsync(entry.reader, [](bool error, void* userData) { ((EventSlim*)userData)->Set(); }, &entry.done))
 				continue;
 			error = entry.message.GetError();
 			entry.~Entry();
@@ -1055,11 +1055,11 @@ namespace uba
 				// For some very unknown reason ASAN on linux triggers on the "new[]" call when doing delete[]
 				// while doing it manually works properly. Will stop investigating this and move on
 				#if PLATFORM_WINDOWS
-				events = new Event[workCount];
+				events = new EventSlim[workCount];
 				#else
-				events = (Event*)aligned_alloc(alignof(Event), sizeof(Event)*workCount);
+				events = (EventSlim*)aligned_alloc(alignof(EventSlim), sizeof(EventSlim)*workCount);
 				for (auto i=0;i!=workCount; ++i)
-					new (events + i) Event();
+					new (events + i) EventSlim();
 				#endif
 			}
 			~WorkRec()
@@ -1068,7 +1068,7 @@ namespace uba
 				delete[] events;
 				#else
 				for (auto i=0;i!=workCount; ++i)
-					events[i].~Event();
+					events[i].~EventSlim();
 				aligned_free(events);
 				#endif
 			}
@@ -1077,7 +1077,7 @@ namespace uba
 			BufferCache* bufferSlots;
 			Atomic<u64> refCount;
 			Atomic<u64> compressCounter = 0;
-			Event* events;
+			EventSlim* events;
 			const u8* uncompressedData = nullptr;
 			u64 written = 0;
 			u64 workCount = 0;
@@ -1339,7 +1339,7 @@ namespace uba
 				u64 decompressedSize = 0;
 				u64 decompressedLeft = 0;
 				u64 written = 0;
-				Event done;
+				EventSlim done;
 				Atomic<bool> error = false;
 			};
 
diff --git a/Engine/Source/Programs/UnrealBuildAccelerator/Common/Private/UbaWorkManager.cpp b/Engine/Source/Programs/UnrealBuildAccelerator/Common/Private/UbaWorkManager.cpp
--- a/Engine/Source/Programs/UnrealBuildAccelerator/Common/Private/UbaWorkManager.cpp
+++ b/Engine/Source/Programs/UnrealBuildAccelerator/Common/Private/UbaWorkManager.cpp
@@ -244,7 +244,7 @@ namespace uba
 		Worker* m_nextWorker = nullptr;
 		Worker* m_prevWorker = nullptr;
 		Atomic<bool> m_loop = true;
-		Event m_workAvailable;
+		EventSlim m_workAvailable;
 		Thread m_thread;
 	};
 
diff --git a/Engine/Source/Programs/UnrealBuildAccelerator/Common/Public/UbaNetworkServer.h b/Engine/Source/Programs/UnrealBuildAccelerator/Common/Public/UbaNetworkServer.h
--- a/Engine/Source/Programs/UnrealBuildAccelerator/Common/Public/UbaNetworkServer.h
+++ b/Engine/Source/Programs/UnrealBuildAccelerator/Common/Public/UbaNetworkServer.h
@@ -167,7 +167,7 @@ namespace uba
 
 		struct PopWorkerRequest
 		{
-			Event ev;
+			EventSlim ev;
 			PopWorkerRequest* next = nullptr;
 			Worker* worker = nullptr;
 		};
diff --git a/Engine/Source/Programs/UnrealBuildAccelerator/Common/Public/UbaWorkManager.h b/Engine/Source/Programs/UnrealBuildAccelerator/Common/Public/UbaWorkManager.h
--- a/Engine/Source/Programs/UnrealBuildAccelerator/Common/Public/UbaWorkManager.h
+++ b/Engine/Source/Programs/UnrealBuildAccelerator/Common/Public/UbaWorkManager.h
@@ -155,7 +155,7 @@ namespace uba
 
 		workCount = Min(workCount, u32(size - 1)/BatchSize);
 
-		Event doneEvent(EventResetType_Manual);
+		EventSlim doneEvent(EventResetType_Manual);
 
 		struct Context
 		{
@@ -165,7 +165,7 @@ namespace uba
 			u32 activeCount;
 			bool isDone;
 			Futex lock;
-			Event* doneEvent;
+			EventSlim* doneEvent;
 		};
 
 		auto context = new Context();
