[UBA]
* Added callback to NetworkBackend for when data is sent to be able to track in speedtest
* Increased recv buffer in tcp backend to 256kb.
* Implemented scattered recv in tcp backend but turned off because it didn't help
* Changed to EventSlim in tcp backend

diff --git a/Engine/Source/Programs/UnrealBuildAccelerator/Common/Private/UbaNetworkBackendTcp.cpp b/Engine/Source/Programs/UnrealBuildAccelerator/Common/Private/UbaNetworkBackendTcp.cpp
--- a/Engine/Source/Programs/UnrealBuildAccelerator/Common/Private/UbaNetworkBackendTcp.cpp
+++ b/Engine/Source/Programs/UnrealBuildAccelerator/Common/Private/UbaNetworkBackendTcp.cpp
@@ -19,6 +19,8 @@
 #pragma comment (lib, "Ws2_32.lib")
 #pragma comment(lib, "IPHLPAPI.lib") // For GetAdaptersInfo
 #pragma comment(lib, "qwave.lib")
+#define WSABUFMAX 256
+#define WSABUFINIT(a, b) { u32(b), (char*)a }
 #else
 #include "UbaLinuxNetworkWrappers.h"
 #include <netdb.h>
@@ -44,6 +46,9 @@
 #define strcpy_s(a, b, c) strcpy(a, c)
 #define WSAPOLLFD pollfd
 #define WSAPoll poll
+#define WSABUF iovec
+#define WSABUFMAX IOV_MAX
+#define WSABUFINIT(a, b) { (void*)a, b }
 #endif
 
 #define UBA_LOG_SOCKET_ERRORS UBA_DEBUG
@@ -91,14 +96,14 @@ namespace uba
 		const u8* data;
 		u32 dataLen;
 		SendBlock* next;
-		Event sent;
+		EventSlim sent;
 	};
 
 	struct ThreadedSendContext
 	{
 		SendBlock* firstBlock = nullptr;
 		SendBlock* lastBlock = nullptr;
-		Event blockQueueEvent;
+		EventSlim blockQueueEvent;
 		Futex blockQueueLock;
 		Thread thread;
 	};
@@ -133,6 +138,9 @@ namespace uba
 		void* dataSentContext = nullptr;
 		DataSentCallback* dataSentCallback = nullptr;
 
+		void* dataRecvContext = nullptr;
+		DataSentCallback* dataRecvCallback = nullptr;
+
 		void* disconnectContext = nullptr;
 		DisconnectCallback* disconnectCallback = nullptr;
 
@@ -186,8 +194,9 @@ namespace uba
 
 	struct NetworkBackendTcp::RecvCache
 	{
-		u8 bytes[128*1024];
-		u32 byteCount = 0;
+		u8 bytes[256*1024];
+		u32 pos = 0;
+		u32 written = 0;
 	};
 
 	////////////////////////////////////////////////////////////////////////////////////////////////////
@@ -369,6 +378,20 @@ namespace uba
 		conn.dataSentContext = context;
 	}
 
+	void NetworkBackendTcp::SetDataReceivedCallback(void* connection, void* context, DataReceivedCallback* callback)
+	{
+		auto& conn = *(Connection*)connection;
+		conn.dataRecvCallback = callback;
+		conn.dataRecvContext = context;
+	}
+
+	void NetworkBackendTcp::TraverseConnections(const Function<void(void*)>& func)
+	{
+		SCOPED_FUTEX(m_connectionsLock, lock);
+		for (auto& conn : m_connections)
+			func(&conn);
+	}
+
 	void NetworkBackendTcp::SetRecvCallbacks(void* connection, void* context, u32 headerSize, RecvHeaderCallback* h, RecvBodyCallback* b, const tchar* recvHint)
 	{
 		UBA_ASSERT(h);
@@ -404,7 +427,7 @@ namespace uba
 		if (conn.threadedSendContext)
 			return;
 		auto context = new ThreadedSendContext();
-		context->blockQueueEvent.Create(EventResetType_Auto);
+		context->blockQueueEvent.Create(EventResetType_Manual);
 		context->thread.Start([this, connPtr = &conn] { ThreadSend(*connPtr); return 0; }, TC("UbaTcpSnd"));
 		conn.threadedSendContext = context;
 	}
@@ -691,8 +714,6 @@ namespace uba
 					break;
 				isFirst = false;
 
-				m_totalRecv += connection.headerSize;
-
 				auto hc = connection.headerCallback;
 				if (!hc)
 				{
@@ -702,13 +723,17 @@ namespace uba
 
 				if (!hc(connection.recvContext, connection.uid, headerData, bodyContext, bodyData, bodySize))
 					break;
+
+				u32 received = connection.headerSize + bodySize;
+				m_totalRecv += received;
+				if (auto c = connection.dataRecvCallback)
+					c(connection.dataRecvContext, received);
+
 				if (!bodySize)
 					continue;
 
 				bool success = RecvSocket(connection, recvCache, bodyData, bodySize, TC("Body"), false, connection.allowLess);
 
-				m_totalRecv += bodySize;
-
 				auto bc = connection.bodyCallback;
 				if (!bc)
 				{
@@ -905,8 +930,10 @@ namespace uba
 		// For some reason performance gets destroyed on wine when changing these
 		if (IsWindows && !IsRunningWine())
 		{
-			SetRecvBuf(logger, socketFd, m_recvBufferSize);
-			SetSendBuf(logger, socketFd, m_sendBufferSize);
+			if (m_recvBufferSize)
+				SetRecvBuf(logger, socketFd, m_recvBufferSize);
+			if (m_sendBufferSize)
+				SetSendBuf(logger, socketFd, m_sendBufferSize);
 		}
 
 		// Socket is good, cancel the socket close scope and break out of the loop.
@@ -1027,6 +1054,8 @@ namespace uba
 			SendBlock* newLast = context.lastBlock;
 			context.firstBlock = nullptr;
 			context.lastBlock = nullptr;
+			if (!first && !newFirst)
+				context.blockQueueEvent.Reset();
 			queueLock.Leave();
 
 			if (!first && !newFirst)
@@ -1054,21 +1083,18 @@ namespace uba
 			if (socket == INVALID_SOCKET)
 				break;
 
-		#if PLATFORM_WINDOWS
-			WSABUF buf[256];
+			WSABUF buf[WSABUFMAX];
 			u32 count = 0;
-			for (SendBlock* it = first; it && count < 255; it = it->next, ++count)
-				buf[count] = { u32(it->dataLen), (char*)it->data };
+			for (SendBlock* it = first; it && count < WSABUFMAX; it = it->next, ++count)
+				buf[count] = WSABUFINIT(it->data, it->dataLen);
+
+			#if PLATFORM_WINDOWS
 			DWORD bytesSent = 0;
 			int res = WSASend(socket, buf, count, &bytesSent, 0, NULL, NULL);
-		#else
-			iovec buf[IOV_MAX];
-			u32 count = 0;
-			for (SendBlock* it = first; it && count < IOV_MAX; it = it->next, ++count)
-				buf[count] = { (void*)it->data, it->dataLen };
+			#else
 			ssize_t res = writev(socket, buf, count);
 			u64 bytesSent = (u64)res;
-		#endif
+			#endif
 
 			if (res != SOCKET_ERROR)
 			{
@@ -1354,68 +1380,107 @@ namespace uba
 
 		SOCKET socket = connection.socket;
 
+		u32 retryCount = 4;
+
 		u8* buffer = (u8*)b;
 		u32 recvLeft = bufferLen;
 		while (recvLeft)
 		{
-			if (recvCache.byteCount)
+			if (recvCache.written)
 			{
-				u32 toCopy = Min(recvCache.byteCount, recvLeft);
-				memcpy(buffer, recvCache.bytes, toCopy);
-				recvCache.byteCount -= toCopy;
-				if (recvCache.byteCount)
-					memmove(recvCache.bytes, recvCache.bytes + toCopy, recvCache.byteCount);
+				u32 toCopy = Min(recvCache.written - recvCache.pos, recvLeft);
+				memcpy(buffer, recvCache.bytes + recvCache.pos, toCopy);
+				recvCache.pos += toCopy;
+				if (recvCache.pos == recvCache.written)
+				{
+					recvCache.pos = 0;
+					recvCache.written = 0;
+				}
 				recvLeft -= toCopy;
 				buffer += toCopy;
 				if (!recvLeft)
 					return true;
 			}
 
-			u8* writePos;
-			u32 writeCapacity;
-
-			if (sizeof(recvCache.bytes) < recvLeft)
+			int read;
+#if PLATFORM_WINDOWS
+			if (false)//!IsRunningWine())
 			{
-				writePos = buffer;
-				writeCapacity = recvLeft;
+				WSABUF bufs[2];
+				bufs[0] = WSABUFINIT(buffer, recvLeft);
+				bufs[1] = WSABUFINIT(recvCache.bytes, sizeof(recvCache.bytes));
+				u32 bufCount = 2;
+
+				DWORD bytesReceived = 0;
+				DWORD flags = 0;
+				read = WSARecv(socket, bufs, bufCount, &bytesReceived, &flags, NULL, NULL);
+				if (read == 0)
+					read = int(bytesReceived);
+				else
+					read = SOCKET_ERROR;
+				if (read > 0)
+				{
+					if (u32(read) >= recvLeft)
+					{
+						recvCache.written = read - recvLeft;
+						return true;
+					}
+					recvLeft -= read;
+					buffer += read;
+
+					if (allowLess)
+					{
+						bufferLen = read;
+						return true;
+					}
+					continue;
+				}
 			}
 			else
+#endif
 			{
-				writePos = recvCache.bytes + recvCache.byteCount;
-				writeCapacity = sizeof(recvCache.bytes) - recvCache.byteCount;
-			}
+				u8* writePos;
+				u32 writeCapacity;
 
-			int read = (int)recv(socket, (char*)writePos, writeCapacity, 0);
-
-			if (read > 0)
-			{
-				if (writePos == buffer)
+				if (sizeof(recvCache.bytes) < recvLeft)
 				{
-					recvLeft -= read;
-					buffer += read;
+					writePos = buffer;
+					writeCapacity = recvLeft;
 				}
 				else
 				{
-					recvCache.byteCount += read;
-					UBA_ASSERT(recvCache.byteCount <= sizeof(recvCache.bytes));
-					u32 toCopy = Min(recvCache.byteCount, recvLeft);
-					memcpy(buffer, recvCache.bytes, toCopy);
-					recvCache.byteCount -= toCopy;
-					if (recvCache.byteCount)
-						memmove(recvCache.bytes, recvCache.bytes + toCopy, recvCache.byteCount);
-					recvLeft -= toCopy;
-					buffer += toCopy;
+					writePos = recvCache.bytes;
+					writeCapacity = sizeof(recvCache.bytes);
 				}
 
-				if (recvLeft == 0)
-					return true;
-			
-				if (allowLess)
+				read = (int)recv(socket, (char*)writePos, writeCapacity, 0);
+				if (read > 0)
 				{
-					bufferLen = read;
-					return true;
+					if (writePos == buffer)
+					{
+						recvLeft -= read;
+						buffer += read;
+					}
+					else
+					{
+						recvCache.written = read;
+						u32 toCopy = Min(recvCache.written, recvLeft);
+						memcpy(buffer, recvCache.bytes, toCopy);
+						recvCache.pos += toCopy;
+						recvLeft -= toCopy;
+						buffer += toCopy;
+					}
+
+					if (recvLeft == 0)
+						return true;
+
+					if (allowLess)
+					{
+						bufferLen = read;
+						return true;
+					}
+					continue;
 				}
-				continue;
 			}
 
 			if (read == 0)
@@ -1430,6 +1495,11 @@ namespace uba
 
 			if (WSAGetLastError() == WSAEWOULDBLOCK)
 			{
+				if (retryCount > 0)
+				{
+					--retryCount;
+					continue;
+				}
 				u32 timeoutMs = connection.recvTimeoutMs;
 				if (!timeoutMs)
 					timeoutMs = DefaultNetworkReceiveTimeoutSeconds * 1000; // 10 minutes time out
diff --git a/Engine/Source/Programs/UnrealBuildAccelerator/Common/Public/UbaNetworkBackend.h b/Engine/Source/Programs/UnrealBuildAccelerator/Common/Public/UbaNetworkBackend.h
--- a/Engine/Source/Programs/UnrealBuildAccelerator/Common/Public/UbaNetworkBackend.h
+++ b/Engine/Source/Programs/UnrealBuildAccelerator/Common/Public/UbaNetworkBackend.h
@@ -37,8 +37,6 @@ namespace uba
 
 		// Send data to connection. If sendContext is set it means that backend expects sendContext to stay around in memory until send is done (or rather, response is received)
 		virtual bool Send(Logger& logger, void* connection, const void* data, u32 dataSize, SendContext& sendContext, const tchar* sendHint) = 0;
-		using DataSentCallback = void(void* context, u32 bytes);
-		virtual void SetDataSentCallback(void* connection, void* context, DataSentCallback* callback) = 0;
 
 		// Recv data from connection. This is callback based so header callback is called first after headerSize bytes have been read
 		// If header callback sets a value in outBodySize, body callback will be called once body size bytes are received.
@@ -49,6 +47,13 @@ namespace uba
 		using RecvTimeoutCallback = bool(void* context, u32 timeoutMs, const tchar* recvHint, const tchar* hint);
 		virtual void SetRecvTimeout(void* connection, u32 timeoutMs, void* context, RecvTimeoutCallback* callback) = 0;
 
+		// Callbacks to be able to track stats
+		using DataSentCallback = void(void* context, u32 bytes);
+		using DataReceivedCallback = void(void* context, u32 bytes);
+		virtual void SetDataSentCallback(void* connection, void* context, DataSentCallback* callback) {}
+		virtual void SetDataReceivedCallback(void* connection, void* context, DataReceivedCallback* callback) {}
+		virtual void TraverseConnections(const Function<void(void*)>& func) {}
+
 		// Change so send is handled in a separate thread instead of directly on calling thread
 		virtual void SetThreadedSend(void* connection) = 0;
 
diff --git a/Engine/Source/Programs/UnrealBuildAccelerator/Common/Public/UbaNetworkBackendTcp.h b/Engine/Source/Programs/UnrealBuildAccelerator/Common/Public/UbaNetworkBackendTcp.h
--- a/Engine/Source/Programs/UnrealBuildAccelerator/Common/Public/UbaNetworkBackendTcp.h
+++ b/Engine/Source/Programs/UnrealBuildAccelerator/Common/Public/UbaNetworkBackendTcp.h
@@ -36,7 +36,6 @@ namespace uba
 		virtual ~NetworkBackendTcp();
 		virtual void Shutdown(void* connection) override;
 		virtual bool Send(Logger& logger, void* connection, const void* data, u32 dataSize, SendContext& sendContext, const tchar* sendHint) override;
-		virtual void SetDataSentCallback(void* connection, void* context, DataSentCallback* callback) override;
 		virtual void SetRecvCallbacks(void* connection, void* context, u32 headerSize, RecvHeaderCallback* h, RecvBodyCallback* b, const tchar* recvHint) override;
 		virtual void SetRecvTimeout(void* connection, u32 timeoutMs, void* context, RecvTimeoutCallback* callback) override;
 		virtual void SetThreadedSend(void* connection) override;
@@ -50,6 +49,10 @@ namespace uba
 		virtual bool Connect(Logger& logger, const sockaddr& remoteSocketAddr, const ConnectedFunc& connectedFunc, bool* timedOut = nullptr, const tchar* nameHint = nullptr) override;
 		virtual void DeleteConnection(void* connection) override;
 
+		virtual void SetDataSentCallback(void* connection, void* context, DataSentCallback* callback) override;
+		virtual void SetDataReceivedCallback(void* connection, void* context, DataReceivedCallback* callback) override;
+		virtual void TraverseConnections(const Function<void(void*)>& func) override;
+
 		virtual void SetSendBufferSize(void* connection, u32 bytes) override;
 		virtual void SetRecvBufferSize(void* connection, u32 bytes) override;
 
