[UBA]
* Improved Speedtest code to be more similar to iperf to be able to compare numbers.

diff --git a/Engine/Source/Programs/UnrealBuildAccelerator/Agent/Private/UbaAgent.cpp b/Engine/Source/Programs/UnrealBuildAccelerator/Agent/Private/UbaAgent.cpp
--- a/Engine/Source/Programs/UnrealBuildAccelerator/Agent/Private/UbaAgent.cpp
+++ b/Engine/Source/Programs/UnrealBuildAccelerator/Agent/Private/UbaAgent.cpp
@@ -1492,11 +1492,19 @@ namespace uba
 						sessionClient->SetMaxIdleTime(maxIdleSeconds);
 					}
 				}
-				// This is an estimation based on tcp limitations (ack and sliding windows).
-				// For every 15ms latency on "best ping") we increase targetConnectionCount up to maxConnectionCount
-				if (!storageClient->IsUsingProxy())
-					if (u64 bestPing = sessionClient->GetBestPing())
-						targetConnectionCount = Max(minConnectionCount, Min(u32(TimeToMs(bestPing) / 15), maxConnectionCount));
+
+				if (u32 desiredConnectionCount = sessionClient->GetDesiredConnectionCount())
+				{
+					targetConnectionCount = desiredConnectionCount;
+				}
+				else
+				{
+					// This is an estimation based on tcp limitations (ack and sliding windows).
+					// For every 15ms latency on "best ping") we increase targetConnectionCount up to maxConnectionCount
+					if (!storageClient->IsUsingProxy())
+						if (u64 bestPing = sessionClient->GetBestPing())
+							targetConnectionCount = Max(minConnectionCount, Min(u32(TimeToMs(bestPing) / 15), maxConnectionCount));
+				}
 
 				if (!isTerminating)
 				{
diff --git a/Engine/Source/Programs/UnrealBuildAccelerator/Cli/Private/UbaCli.cpp b/Engine/Source/Programs/UnrealBuildAccelerator/Cli/Private/UbaCli.cpp
--- a/Engine/Source/Programs/UnrealBuildAccelerator/Cli/Private/UbaCli.cpp
+++ b/Engine/Source/Programs/UnrealBuildAccelerator/Cli/Private/UbaCli.cpp
@@ -102,7 +102,6 @@ namespace uba
 		logger.Info(TC("   -cachecommand=<cmd>     Send command to cache server. Will output result in log"));
 		logger.Info(TC("   -writecachesummary      Write cache summary file about connected cache server"));
 		logger.Info(TC("   -decompress=<file>      Decompress a compressed file in-place"));
-		logger.Info(TC("   -speedtest              Run a network speedtest"));
 		logger.Info(TC(""));
 		logger.Info(TC("  CoordinatorOptions (if coordinator set):"));
 		logger.Info(TC("   -uri=<address>          Uri to coordinator"));
@@ -111,7 +110,8 @@ namespace uba
 		logger.Info(TC("   -maxcores=<number>      Max number of cores that will be asked for from coordinator"));
 		logger.Info(TC(""));
 		logger.Info(TC("  If <executable> is a .yaml-file UbaCli creates a scheduler to execute commands from the yaml file instead"));
-		
+		logger.Info(TC("  If <executable> is \"speedtest\" then a network speed test will be performed"));
+
 		logger.Info(TC(""));
 		return false;
 	}
@@ -185,7 +185,6 @@ namespace uba
 		bool useHackVfs = false;
 		bool useScheduler = false;
 		bool decompressInPlace = false;
-		bool speedTest = false;
 		TString checkFileTable;
 		TString cacheFilterString;
 		TString cacheCommand;
@@ -467,10 +466,6 @@ namespace uba
 				{
 					storeCompressed = false;
 				}
-				else if (name.Equals(TCV("-speedtest")))
-				{
-					speedTest = true;
-				}
 				else if (name.Equals(TCV("-dir")))
 				{
 					if (value.IsEmpty())
@@ -959,13 +954,6 @@ namespace uba
 			return true;
 		}
 
-		if (speedTest)
-		{
-			if (commandType == CommandType_NotSet)
-				commandType = CommandType_Agent;
-			application = TC("speedtest");
-		}
-
 		if (exit)
 			return true;
 
@@ -984,7 +972,7 @@ namespace uba
 			if (application.empty())
 				return PrintHelp(TC("No executable provided"));
 
-			if (!IsAbsolutePath(application.c_str()) && !speedTest)
+			if (!IsAbsolutePath(application.c_str()) && !StringView(application).Equals(TCV("speedTest")))
 			{
 				StringBuffer<> fullApplicationName;
 				if (!SearchPathForFile(logger, fullApplicationName, application.c_str(), currentDir, {}))
diff --git a/Engine/Source/Programs/UnrealBuildAccelerator/Common/Private/UbaNetworkClient.cpp b/Engine/Source/Programs/UnrealBuildAccelerator/Common/Private/UbaNetworkClient.cpp
--- a/Engine/Source/Programs/UnrealBuildAccelerator/Common/Private/UbaNetworkClient.cpp
+++ b/Engine/Source/Programs/UnrealBuildAccelerator/Common/Private/UbaNetworkClient.cpp
@@ -625,6 +625,13 @@ namespace uba
 		return m_connections.front().backend;
 	}
 
+	void NetworkClient::TraverseConnections(const Function<void(NetworkBackend&, void*)>& func)
+	{
+		SCOPED_WRITE_LOCK(m_connectionsLock, lock);
+		for (auto& c : m_connections)
+			func(*c.backend, c.backendConnection);
+	}
+
 	void NetworkClient::OnDisconnected(Connection& connection, u32 reason)
 	{
 		if (connection.connected.exchange(0) == 1)
diff --git a/Engine/Source/Programs/UnrealBuildAccelerator/Common/Private/UbaSessionClient.cpp b/Engine/Source/Programs/UnrealBuildAccelerator/Common/Private/UbaSessionClient.cpp
--- a/Engine/Source/Programs/UnrealBuildAccelerator/Common/Private/UbaSessionClient.cpp
+++ b/Engine/Source/Programs/UnrealBuildAccelerator/Common/Private/UbaSessionClient.cpp
@@ -229,6 +229,11 @@ namespace uba
 		return m_maxProcessCount;
 	}
 
+	u32 SessionClient::GetDesiredConnectionCount()
+	{
+		return m_desiredConnectionCount;
+	}
+
 	bool SessionClient::Exists(const StringView& path, u32& outAttributes)
 	{
 		u32 tableOffset = 0;
@@ -1766,14 +1771,18 @@ namespace uba
 		msg.Send();
 	}
 
-	void SessionClient::SendLogEntry(StringView text, LogEntryType type)
+	void SessionClient::SendLogEntry(StringView text, LogEntryType type, bool ack)
 	{
 		StackBinaryWriter<1024> writer;
 		NetworkMessage msg(m_client, ServiceId, SessionMessageType_LogEntry, writer);
 		writer.WriteByte(type);
 		writer.WriteString(text);
+
 		StackBinaryReader<16> reader;
-		msg.Send(reader);
+		if (ack)
+			msg.Send(reader);
+		else
+			msg.Send();
 	}
 
 	bool SessionClient::SendRootsHandle(RootsHandle rootsHandle)
@@ -2040,7 +2049,7 @@ namespace uba
 				{
 					if (startInfo.applicationStr == TC("speedtest"))
 					{
-						RunSpeedTest();
+						RunSpeedTest(startInfo.argumentsStr);
 						SendReturnProcess(startInfo.processId, TC("speedtest"));
 						continue;
 					}
@@ -2354,9 +2363,59 @@ namespace uba
 		return true;
 	}
 
-	void SessionClient::RunSpeedTest()
+	void SessionClient::RunSpeedTest(StringView args)
 	{
-		u64 dataSize = 300ull * 1024 * 1024;
+		struct Logger2 : Logger
+		{
+			SessionClient& sc;
+			Logger2(SessionClient& c) : sc(c) {}
+			virtual void BeginScope() override {}
+			virtual void EndScope() override {}
+			virtual void Log(LogEntryType type, const tchar* str, u32 strLen)
+			{
+				//sc.m_logger.Log(type, str, strLen);
+				sc.SendLogEntry(StringView(str, strLen));
+			}
+		} logger(*this);
+
+
+		u64 connectionCount = 1;
+		ParseArguments(args.data, args.count, [&](const tchar* arg, u32 argLen)
+			{
+				StringView v(arg, argLen);
+				if (v.StartsWith(TCV("-c=")))
+					Parse(connectionCount, arg+3, argLen-3);
+			});
+
+		logger.Info(TC("Starting speed test.. Waiting for %llu connections"), connectionCount);
+
+		m_desiredConnectionCount = u32(connectionCount);
+		while (m_client.IsConnected() && m_client.GetConnectionCount() != connectionCount)
+			Sleep(1);
+
+		struct Connection
+		{
+			Atomic<u64> send = 0;
+			Atomic<u64> recv = 0;
+
+			u64 oldSend = 0;
+			u64 oldRecv = 0;
+
+			u64 sendDelta = 0;
+			u64 recvDelta = 0;
+		};
+
+		UnorderedMap<void*, Connection> connections;
+
+		m_client.TraverseConnections([&](NetworkBackend& backend, void* conn)
+			{
+				Connection& c = connections[conn];
+				backend.SetDataSentCallback(conn, &c, [](void* context, u32 bytes) { ((Connection*)context)->send += bytes; });
+				backend.SetDataReceivedCallback(conn, &c, [](void* context, u32 bytes) { ((Connection*)context)->recv += bytes; });
+			});
+
+
+		u64 dataSize = 1024ull * 1024 * 1024;
 		u8* dest = new u8[dataSize];
 
 		u64 messageCount = dataSize / (SendMaxSize - 10);
@@ -2365,32 +2424,105 @@ namespace uba
 		{
 			Entry(NetworkClient& c, BinaryWriter& w, u8* d) : msg(c, ServiceId, SessionMessageType_SpeedTest, w), ev(EventResetType_Manual), reader(d) {}
 			NetworkMessage msg;
-			Event ev;
+			EventSlim ev;
 			BinaryReader reader;
 		};
 
-		for (u32 j=0; j!=10; ++j)
-		{
-			StringBuffer<> note;
-			note.Appendf(TC("Receiving 300mb (%i)"), j);
-			SendLogEntry(note);
+		Event logLoop(EventResetType_Manual);
+		Thread logThread([&]()
+			{
+				u64 startTime = GetTime();
+				u64 lastTime = GetTime();
+				while (!logLoop.IsSet(1000))
+				{
+					u64 time = GetTime();
+					if (TimeToMs(time - startTime) > 10 * 1000)
+						break;
+
+					u64 deltaTime = time - lastTime;
+					lastTime = time;
+					u32 index = 0;
+					u64 totalRecv = 0;
+
+					for (auto& kv : connections)
+					{
+						auto& c = kv.second;
+						u64 send = c.send;
+						u64 recv = c.recv;
+						c.sendDelta = send - c.oldSend;
+						c.recvDelta = recv - c.oldRecv;
+						c.oldSend = send;
+						c.oldRecv = recv;
+					}
 
-			List<Entry> entries;
+					StringBuffer<2048> str;
+					for (auto& kv : connections)
+					{
+						auto& c = kv.second;
+
+						float bps = float(c.recvDelta * 8) / TimeToS(deltaTime);
+
+						str.Appendf(TC("\r\n [%3i] %s %sit/sec"), index++, BytesToText(c.recvDelta).str, BytesToText(u32(bps)).str);
+						totalRecv += c.recvDelta;
+					}
 
-			u64 start = GetTime();
-			for (u64 i=0; i!=messageCount;++i)
+					if (connections.size() > 1)
+					{
+						float bps = float(totalRecv * 8) / TimeToS(deltaTime);
+						str.Appendf(TC("\r\n [SUM] %s %sit/sec"), BytesToText(totalRecv).str, BytesToText(u32(bps)).str);
+					}
+					logger.Log(LogEntryType_Info, str.data, str.count);
+				}
+				return 0;
+			}, TC("UbaSpdTst"));
+
+		struct EntryData { u8 data[sizeof(Entry)]; };
+
+		Vector<EntryData> entries;
+		entries.resize(messageCount);
+		u64 entryCounter = 0;
+		while (!logThread.Wait(0))
+		{
+			u64 entryIndex = entryCounter % messageCount;
+			StackBinaryWriter<16> writer;
+			auto& entry = *new (entries[entryIndex].data) Entry(m_client, writer, dest + (SendMaxSize - 10) * entryIndex);
+			writer.WriteByte(0);
+			if (!entry.msg.SendAsync(entry.reader, [](bool err, void* userData) { ((EventSlim*)userData)->Set(); }, & entry.ev))
+				break;
+
+			++entryCounter;
+			if (entryCounter >= messageCount)
 			{
-				StackBinaryWriter<16> writer;
-				auto& entry = entries.emplace_back(m_client, writer, dest + (SendMaxSize - 10) * i);
-				writer.WriteByte(0);
-				entry.msg.SendAsync(entry.reader, [](bool err, void* userData) { ((Event*)userData)->Set(); }, &entry.ev);
+				u64 prevEntryIndex = entryCounter % messageCount;
+				u8* data = entries[prevEntryIndex].data;
+				auto& prevEntry = *(Entry*)data;
+				prevEntry.ev.IsSet();
+				// prevEntry.msg.ProcessAsyncResults(prevEntry.reader); // Not really needed, just decrypts
+				prevEntry.~Entry();
+				*(u64*)data = ~0ull;
 			}
-			for (auto& entry : entries)
-				entry.ev.IsSet();
+		}
+
+		logLoop.Set();
+		logThread.Wait();
 
-			note.Clear().Appendf(TC(".. took (%s)"), TimeToText(GetTime() - start).str);
-			SendLogEntry(note);
+		for (u64 i=0, e=Min(entryCounter, messageCount); i!=e; ++i)
+		{
+			u8* data = entries[i].data;
+			if (*(u64*)data == ~0ull)
+				continue;
+			auto& entry = *(Entry*)data;
+			entry.ev.IsSet();
+			entry.~Entry();
 		}
+
+		logger.Info(TC("Speed test done"));
+
+		m_client.TraverseConnections([&](NetworkBackend& backend, void* conn)
+			{
+				backend.SetDataSentCallback(conn, nullptr, nullptr);
+				backend.SetDataReceivedCallback(conn, nullptr, nullptr);
+			});
 	}
 
 	bool SessionClient::AllocFailed(Process& process, const tchar* allocType, u32 error)
diff --git a/Engine/Source/Programs/UnrealBuildAccelerator/Common/Public/UbaNetworkClient.h b/Engine/Source/Programs/UnrealBuildAccelerator/Common/Public/UbaNetworkClient.h
--- a/Engine/Source/Programs/UnrealBuildAccelerator/Common/Public/UbaNetworkClient.h
+++ b/Engine/Source/Programs/UnrealBuildAccelerator/Common/Public/UbaNetworkClient.h
@@ -76,6 +76,7 @@ namespace uba
 		u64 GetTotalRecvBytes() { return m_recvBytes; }
 
 		NetworkBackend* GetFirstConnectionBackend();
+		void TraverseConnections(const Function<void(NetworkBackend&, void*)>& func);
 
 	private:
 		struct Connection;
diff --git a/Engine/Source/Programs/UnrealBuildAccelerator/Common/Public/UbaSessionClient.h b/Engine/Source/Programs/UnrealBuildAccelerator/Common/Public/UbaSessionClient.h
--- a/Engine/Source/Programs/UnrealBuildAccelerator/Common/Public/UbaSessionClient.h
+++ b/Engine/Source/Programs/UnrealBuildAccelerator/Common/Public/UbaSessionClient.h
@@ -50,6 +50,7 @@ namespace uba
 
 		u64 GetBestPing();
 		u32 GetMaxProcessCount();
+		u32 GetDesiredConnectionCount();
 
 		bool Exists(const StringView& path, u32& outAttributes);
 
@@ -104,14 +105,14 @@ namespace uba
 		bool SendProcessFinished(ProcessImpl& process, u32 exitCode);
 		void SendPing(u64 memAvail, u64 memTotal);
 		void SendNotification(StringView text);
-		void SendLogEntry(StringView text, LogEntryType type = LogEntryType_Info);
+		void SendLogEntry(StringView text, LogEntryType type = LogEntryType_Info, bool ack = false);
 		bool SendRootsHandle(RootsHandle rootsHandle);
 		void SendLogFileToServer(ProcessImpl& pi);
 		void GetLogFileName(StringBufferBase& out, const tchar* logFile, const tchar* arguments, u32 processId);
 		u32 WriteLogLines(BinaryWriter& writer, ProcessImpl& process);
 		bool ParseDirectoryTable();
 		bool EntryExists(const StringView& path, u32& outTableOffset);
-		void RunSpeedTest();
+		void RunSpeedTest(StringView args);
 
 		void ThreadCreateProcessLoop();
 
@@ -142,6 +143,7 @@ namespace uba
 		Atomic<const tchar*> m_terminationReason = nullptr;
 		Atomic<u64> m_terminationTime = 0;
 		Atomic<u32> m_maxProcessCount = 0;
+		Atomic<u32> m_desiredConnectionCount = 0;
 		Atomic<float> m_cpuUsage = 0.0f;
 
 		struct ApplicationEnvironment { Futex lock; TString virtualApplication; TString realApplication; };
diff --git a/Engine/Source/Programs/UnrealBuildAccelerator/Core/Public/UbaEventWaitOnAddress.inl b/Engine/Source/Programs/UnrealBuildAccelerator/Core/Public/UbaEventWaitOnAddress.inl
--- a/Engine/Source/Programs/UnrealBuildAccelerator/Core/Public/UbaEventWaitOnAddress.inl
+++ b/Engine/Source/Programs/UnrealBuildAccelerator/Core/Public/UbaEventWaitOnAddress.inl
@@ -4,7 +4,7 @@
 
 #if PLATFORM_WINDOWS
 
-#include <Windows.h>
+#include "UbaPlatform.h"
 
 namespace uba
 {
