[UBA]
* Fixed so network memory backend can be tested with speedtest using -memsock

diff --git a/Engine/Source/Programs/UnrealBuildAccelerator/Cli/Private/UbaCli.cpp b/Engine/Source/Programs/UnrealBuildAccelerator/Cli/Private/UbaCli.cpp
--- a/Engine/Source/Programs/UnrealBuildAccelerator/Cli/Private/UbaCli.cpp
+++ b/Engine/Source/Programs/UnrealBuildAccelerator/Cli/Private/UbaCli.cpp
@@ -184,6 +184,7 @@ namespace uba
 		bool logToFile = false;
 		bool useHackVfs = false;
 		bool useScheduler = false;
+		bool useMemsock = false;
 		bool decompressInPlace = false;
 		TString checkFileTable;
 		TString cacheFilterString;
@@ -432,6 +433,10 @@ namespace uba
 				{
 					useHackVfs = true;
 				}
+				else if (name.Equals(TCV("-memsock")))
+				{
+					useMemsock = true;
+				}
 				else if (name.Equals(TCV("-cache")))
 				{
 					if (value.IsEmpty())
@@ -1089,7 +1094,10 @@ namespace uba
 		#endif
 
 		NetworkBackendTcpCreateInfo nbtci(logWriter);
-		NetworkBackendTcp networkBackend(nbtci);
+		NetworkBackendTcp networkBackendTcp(nbtci);
+		NetworkBackendMemory networkBackendMemory(logWriter);
+
+		NetworkBackend& networkBackend = useMemsock ? (NetworkBackend&)networkBackendMemory : (NetworkBackend&)networkBackendTcp;
 
 		NetworkServerCreateInfo nsci(logWriter);
 		nsci.useThreadedSend = true;
@@ -1310,14 +1318,12 @@ namespace uba
 
 		auto RunWithClient = [&](const Function<bool()>& func, u32 clientCount)
 			{
-				Vector<Client> clients;
-				clients.resize(clientCount);
-				u32 clientIndex = 0;
-				for (auto& c : clients)
+				List<Client> clients;
+				for (u32 i=0; i!=clientCount; ++i)
 				{
-					u32 maxProcessor = Min(maxProcessCount/u32(clients.size()), 32u);
-					ClientInitInfo cii { logWriter, networkBackend, g_rootDir.data, TC("127.0.0.1"), port, clientZone, maxProcessor, clientIndex++};
-					if (!c.Init(cii))
+					u32 maxProcessor = Min(maxProcessCount/clientCount, 32u);
+					ClientInitInfo cii { logWriter, networkBackend, g_rootDir.data, TC("127.0.0.1"), port, clientZone, maxProcessor, i};
+					if (!clients.emplace_back().Init(cii))
 						return false;
 				}
 				return func();
diff --git a/Engine/Source/Programs/UnrealBuildAccelerator/Cli/Private/UbaClient.h b/Engine/Source/Programs/UnrealBuildAccelerator/Cli/Private/UbaClient.h
--- a/Engine/Source/Programs/UnrealBuildAccelerator/Cli/Private/UbaClient.h
+++ b/Engine/Source/Programs/UnrealBuildAccelerator/Cli/Private/UbaClient.h
@@ -65,7 +65,29 @@ namespace uba
 				sessionClient->Start();
 			}
 			
-			return networkClient->Connect(*networkBackend, info.host, info.port);
+			m_host = info.host;
+			m_port = info.port;
+
+			if (!networkClient->Connect(*networkBackend, m_host.c_str(), m_port))
+				return false;
+
+			m_loopDone.Create(EventResetType_Manual);
+			m_loop.Start([this]()
+				{
+					u32 timeoutMs = 0;
+					while (!m_loopDone.IsSet(timeoutMs))
+					{
+						timeoutMs = 1000;
+						u32 desiredConnections = sessionClient->GetDesiredConnectionCount();
+						if (desiredConnections <= networkClient->GetConnectionCount())
+							continue;
+						networkClient->Connect(*networkBackend, m_host.c_str(), m_port);
+						timeoutMs = 0;
+					}
+					return 0;
+				}, TC("UbaClLoop"));
+
+			return true;
 		}
 
 		bool StartProxy(u16 proxyPort, const Guid& storageServerUid)
@@ -98,6 +120,11 @@ namespace uba
 
 		~Client()
 		{
+			if (m_loopDone.IsCreated())
+			{
+				m_loopDone.Set();
+				m_loop.Wait();
+			}
 			if (proxyNetworkServer)
 				proxyNetworkServer->DisconnectClients();
 			if (storageClient)
@@ -114,6 +141,9 @@ namespace uba
 			delete networkBackendMem;
 		}
 
+		TString m_host;
+		u16 m_port;
+
 		NetworkBackendMemory* networkBackendMem = nullptr;
 		NetworkClient* networkClient = nullptr;
 		StorageClient* storageClient = nullptr;
@@ -123,5 +153,8 @@ namespace uba
 		NetworkServer* proxyNetworkServer = nullptr;
 		StorageProxy* proxyStorage = nullptr;
 		TString serverPrefix;
+
+		Event m_loopDone;
+		Thread m_loop;
 	};
 }
diff --git a/Engine/Source/Programs/UnrealBuildAccelerator/Common/Private/UbaNetworkBackendMemory.cpp b/Engine/Source/Programs/UnrealBuildAccelerator/Common/Private/UbaNetworkBackendMemory.cpp
--- a/Engine/Source/Programs/UnrealBuildAccelerator/Common/Private/UbaNetworkBackendMemory.cpp
+++ b/Engine/Source/Programs/UnrealBuildAccelerator/Common/Private/UbaNetworkBackendMemory.cpp
@@ -27,6 +27,11 @@ namespace uba
 			#if UBA_VALIDATE_MULTIPLE_SEND
 			Atomic<u32> sendCount = 0;
 			#endif
+
+			DataSentCallback* dataSentCallback = nullptr;
+			void* dataSentContext = nullptr;
+			DataReceivedCallback* dataRecvCallback = nullptr;
+			void* dataRecvContext = nullptr;
 		};
 
 		Peer peer[2];
@@ -106,11 +111,29 @@ namespace uba
 			dataPos += bodySize;
 			dataLeft -= bodySize;
 		}
+
+		if (auto c = m_connection->peer[from].dataSentCallback)
+			c(m_connection->peer[from].dataSentContext, dataSize);
+		if (auto c = m_connection->peer[to].dataRecvCallback)
+			c(m_connection->peer[to].dataRecvContext, dataSize);
+
 		return true;
 	}
 
 	void NetworkBackendMemory::SetDataSentCallback(void* connection, void* context, DataSentCallback* callback)
 	{
+		UBA_ASSERT(m_connection);
+		u64 from = ((uintptr_t)connection) - 1;
+		m_connection->peer[from].dataSentCallback = callback;
+		m_connection->peer[from].dataSentContext = context;
+	}
+
+	void NetworkBackendMemory::SetDataReceivedCallback(void* connection, void* context, DataReceivedCallback* callback)
+	{
+		UBA_ASSERT(m_connection);
+		u64 from = ((uintptr_t)connection) - 1;
+		m_connection->peer[from].dataRecvCallback = callback;
+		m_connection->peer[from].dataRecvContext = context;
 	}
 
 	void NetworkBackendMemory::SetRecvCallbacks(void* connection, void* context, u32 headerSize, RecvHeaderCallback* h, RecvBodyCallback* b, const tchar* recvHint)
diff --git a/Engine/Source/Programs/UnrealBuildAccelerator/Common/Private/UbaSessionClient.cpp b/Engine/Source/Programs/UnrealBuildAccelerator/Common/Private/UbaSessionClient.cpp
--- a/Engine/Source/Programs/UnrealBuildAccelerator/Common/Private/UbaSessionClient.cpp
+++ b/Engine/Source/Programs/UnrealBuildAccelerator/Common/Private/UbaSessionClient.cpp
@@ -2462,14 +2462,14 @@ namespace uba
 
 						float bps = float(c.recvDelta * 8) / TimeToS(deltaTime);
 
-						str.Appendf(TC("\r\n [%3i] %s %sit/sec"), index++, BytesToText(c.recvDelta).str, BytesToText(u32(bps)).str);
+						str.Appendf(TC("\r\n [%3i] %s %sit/sec"), index++, BytesToText(c.recvDelta).str, BytesToText(u64(bps)).str);
 						totalRecv += c.recvDelta;
 					}
 
 					if (connections.size() > 1)
 					{
 						float bps = float(totalRecv * 8) / TimeToS(deltaTime);
-						str.Appendf(TC("\r\n [SUM] %s %sit/sec"), BytesToText(totalRecv).str, BytesToText(u32(bps)).str);
+						str.Appendf(TC("\r\n [SUM] %s %sit/sec"), BytesToText(totalRecv).str, BytesToText(u64(bps)).str);
 					}
 					logger.Log(LogEntryType_Info, str.data, str.count);
 				}
diff --git a/Engine/Source/Programs/UnrealBuildAccelerator/Common/Public/UbaNetworkBackendMemory.h b/Engine/Source/Programs/UnrealBuildAccelerator/Common/Public/UbaNetworkBackendMemory.h
--- a/Engine/Source/Programs/UnrealBuildAccelerator/Common/Public/UbaNetworkBackendMemory.h
+++ b/Engine/Source/Programs/UnrealBuildAccelerator/Common/Public/UbaNetworkBackendMemory.h
@@ -16,6 +16,7 @@ namespace uba
 		virtual void Shutdown(void* connection) override;
 		virtual bool Send(Logger& logger, void* connection, const void* data, u32 dataSize, SendContext& sendContext, const tchar* sendHint) override;
 		virtual void SetDataSentCallback(void* connection, void* context, DataSentCallback* callback) override;
+		virtual void SetDataReceivedCallback(void* connection, void* context, DataReceivedCallback* callback) override;
 		virtual void SetRecvCallbacks(void* connection, void* context, u32 headerSize, RecvHeaderCallback* h, RecvBodyCallback* b, const tchar* recvHint) override;
 		virtual void SetRecvTimeout(void* connection, u32 timeoutMs, void* context, RecvTimeoutCallback* callback) override;
 		virtual void SetThreadedSend(void* connection) override;
