[UBA]
* Fixed crash when trying to run speedtest in ubacli with local mode

diff --git a/Engine/Source/Programs/UnrealBuildAccelerator/Cli/Private/UbaCli.cpp b/Engine/Source/Programs/UnrealBuildAccelerator/Cli/Private/UbaCli.cpp
--- a/Engine/Source/Programs/UnrealBuildAccelerator/Cli/Private/UbaCli.cpp
+++ b/Engine/Source/Programs/UnrealBuildAccelerator/Cli/Private/UbaCli.cpp
@@ -969,6 +969,11 @@ namespace uba
 			return PrintHelp(errorMsg);
 		}
 
+		bool isSpeedTest = StringView(application).Equals(TCV("speedTest"));
+		if (isSpeedTest)
+			if (commandType != CommandType_Agent && commandType != CommandType_Remote)
+				return logger.Error(TC("Speedtest does can only run in \"remote\" or \"agent\" mode"));
+
 		StringBuffer<512> currentDir;
 		GetCurrentDirectoryW(currentDir);
 
@@ -977,7 +982,7 @@ namespace uba
 			if (application.empty())
 				return PrintHelp(TC("No executable provided"));
 
-			if (!IsAbsolutePath(application.c_str()) && !StringView(application).Equals(TCV("speedTest")))
+			if (!IsAbsolutePath(application.c_str()) && !isSpeedTest)
 			{
 				StringBuffer<> fullApplicationName;
 				if (!SearchPathForFile(logger, fullApplicationName, application.c_str(), currentDir, {}))
