[UBA]
* Changed so UBA_TRACK_SHARED_MEMORY_ALLOCATIONS also enable tracing if memory allocations

diff --git a/Engine/Source/Programs/UnrealBuildAccelerator/Common/Private/UbaSessionServer.cpp b/Engine/Source/Programs/UnrealBuildAccelerator/Common/Private/UbaSessionServer.cpp
--- a/Engine/Source/Programs/UnrealBuildAccelerator/Common/Private/UbaSessionServer.cpp
+++ b/Engine/Source/Programs/UnrealBuildAccelerator/Common/Private/UbaSessionServer.cpp
@@ -2360,9 +2360,11 @@ namespace uba
 
 		m_trace.SessionUpdate(0, serverConnectionCount, serverSend, serverRecv, 0, memAvail, memTotal, cpuLoad);
 
-#if 0
+		#if UBA_TRACK_SHARED_MEMORY_ALLOCATIONS
 		m_sharedMemory.TraceStats(m_trace, 30);
+		#endif
 
+#if 0
 		m_fileMappingBuffer.TraceStatus(m_trace, 20, TC("Session"), 1);
 
 		m_trace.StatusUpdate(40, 1, TCV("Independents"), LogEntryType_Info);
