[UBA]
* Changed so in-memory cas files created from segments of virtual files are thrown away after first use. It is unlikely same cas is used twice and in that case we pay the cost to re-compress it. Reuse will happen if helper disconnects while downloading the cas requiring a new client to pick up the same work
* Added option to disable un-dropping of dropped cas files. This should be safe to disable but need more testing

diff --git a/Engine/Source/Programs/UnrealBuildAccelerator/Common/Private/UbaSession.cpp b/Engine/Source/Programs/UnrealBuildAccelerator/Common/Private/UbaSession.cpp
--- a/Engine/Source/Programs/UnrealBuildAccelerator/Common/Private/UbaSession.cpp
+++ b/Engine/Source/Programs/UnrealBuildAccelerator/Common/Private/UbaSession.cpp
@@ -1831,6 +1831,7 @@ namespace uba
 				out.lastWriteTime = entry.lastWriteTime;
 				out.storedSize = entry.storedSize;
 				out.createIndependentMapping = entry.createIndependentMapping;
+				out.dropCasAfterUse = true; // TODO: Make this an option. This will throw away cas files after they have been used
 				return true;
 			});
 
diff --git a/Engine/Source/Programs/UnrealBuildAccelerator/Common/Private/UbaStorage.cpp b/Engine/Source/Programs/UnrealBuildAccelerator/Common/Private/UbaStorage.cpp
--- a/Engine/Source/Programs/UnrealBuildAccelerator/Common/Private/UbaStorage.cpp
+++ b/Engine/Source/Programs/UnrealBuildAccelerator/Common/Private/UbaStorage.cpp
@@ -43,6 +43,7 @@ namespace uba
 		table.GetValueAsBool(storeCompressed, TC("StoreCompressed"));
 		table.GetValueAsBool(allowDeferredDeletes, TC("AllowDeferredDeletes"));
 		table.GetValueAsBool(skipCleanupOnShutdown, TC("SkipCleanupOnShutdown"));
+		table.GetValueAsBool(allowUnDropOfCas, TC("AllowUnDropOfCas"));
 		table.GetValueAsU32(fileEntryTableWipeThreshold, TC("FileEntryTableWipeThreshold"));
 		table.GetValueAsU32(maxParallelCopyOrLink, TC("MaxParallelCopyOrLink"));
 		table.GetValueAsU32(deleteCasBatchSize, TC("DeleteCasBatchSize"));
@@ -354,7 +355,7 @@ namespace uba
 		return w.CompressToFile(out.storedSize, from, readHandle, readMem, fileSize, to, header, headerSize, lastWriteTime, m_tempPath.data);
 	}
 
-	bool StorageImpl::WriteCasFileNoCheck(WriteResult& out, const StringKey& fileNameKey, const tchar* fileName, const CasKey& casKey, const tchar* casFile, bool storeCompressed)
+	bool StorageImpl::WriteCasFileNoCheck(WriteCasResult& out, const StringKey& fileNameKey, const tchar* fileName, const CasKey& casKey, const tchar* casFile, bool storeCompressed)
 	{
 		StorageStats& stats = Stats();
 		TimerScope ts(stats.createCas);
@@ -370,11 +371,12 @@ namespace uba
 					return false;
 				out.sourceStoredSize = externalMapping.storedSize;
 				out.isIndependent = externalMapping.createIndependentMapping;
+				out.dropAfterUse = externalMapping.dropCasAfterUse;
 				if (externalMapping.createIndependentMapping)
 				{
 					CompressWriter w(m_logger, m_sharedMemory, m_bufferSlots, m_workManager, Stats(), m_casCompressor, m_casCompressionLevel, m_asyncUnmapViewOfFile);
 					out.memoryOffset = 0;
-					return w.CompressToMapping(out.memoryHandle, out.storedSize, fileMem, externalMapping.size, casFile);
+					return w.CompressToMapping(out.memoryHandle, out.storedSize, fileMem, externalMapping.size, fileName);//casFile);
 				}
 				else
 					return WriteCompressedFile(out, fileName, InvalidFileHandle, fileMem, externalMapping.size, casFile, nullptr, 0);
@@ -462,7 +464,7 @@ namespace uba
 		return true;
 	}
 
-	bool StorageImpl::WriteCasFile(WriteResult& out, const tchar* fileName, const CasKey& casKey)
+	bool StorageImpl::WriteCasFile(WriteCasResult& out, const tchar* fileName, const CasKey& casKey)
 	{
 		UBA_ASSERT(IsCompressed(casKey) == m_storeCompressed);
 		StringBuffer<> casFile;
@@ -567,7 +569,11 @@ namespace uba
 		}
 
 		// File is requested by server itself to provide to someone remote. We have to un-drop the cas in case it was marked as dropped
-		casEntry.dropped = false;
+		if (casEntry.dropped && m_allowUnDropOfCas)
+		{
+			m_logger.Info(TC("Un-dropped cas for %s (%s). Cas will leak in memory"), fileName, BytesToText(fileSize).str);
+			casEntry.dropped = false;
+		}
 
 		if (casEntry.verified)
 			if (casEntry.exists)
@@ -630,7 +636,7 @@ namespace uba
 		casEntry.verified = true;
 		casEntry.exists = false;
 
-		WriteResult res;
+		WriteCasResult res;
 		if (!WriteCasFileNoCheck(res, fileNameKey, fileName, casKey, casFile.data, IsCompressed(casKey)))
 			return false;
 
@@ -908,6 +914,7 @@ namespace uba
 		m_allowDeleteVerified = info.allowDeleteVerified;
 		m_allowDeferredDeletes = info.allowDeferredDeletes;
 		m_skipCleanupOnShutdown = info.skipCleanupOnShutdown;
+		m_allowUnDropOfCas = info.allowUnDropOfCas;
 		m_writeToDisk = info.writeToDisk;
 		if (!m_writeToDisk)
 			m_casCapacityBytes = 0;
@@ -1955,12 +1962,10 @@ namespace uba
 
 		casEntry.verified = true;
 		deferredLock.Leave();
-		WriteResult res;
-		{
-			//TrackHintScope ths(*TrackWorkScope::GetActiveScope(), TCV("WriteCasFile"));
-			if (!WriteCasFile(res, deferredCreation.data, casKey))
-				return m_logger.Error(TC("Failed to write deferred cas %s (%s)"), CasKeyString(casKey).str, deferredCreation.data);
-		}
+		WriteCasResult res;
+		if (!WriteCasFile(res, deferredCreation.data, casKey))
+			return m_logger.Error(TC("Failed to write deferred cas %s (%s)"), CasKeyString(casKey).str, deferredCreation.data);
+		casEntry.dropped = res.dropAfterUse;
 
 		if (res.sourceStoredSize != deferredFileSize) // TODO: Should also check last written time
 			m_logger.Warning(TC("Deferred cas %s was deferred with source file size %llu but was read with source file size %llu (%s)"), CasKeyString(casKey).str, deferredFileSize, res.sourceStoredSize, deferredCreation.data);
@@ -1981,6 +1986,10 @@ namespace uba
 			*out = &casEntry;
 			++casEntry.refCount;
 		}
+		else if (res.dropAfterUse)
+		{
+			m_logger.Warning(TC("Cas is set to drop-after-use but it is not being used in this call so it will leak %s (%s)"), CasKeyString(casKey).str, deferredCreation.data);
+		}
 
 		casEntryLock.Leave();
 
@@ -2027,7 +2036,7 @@ namespace uba
 		casEntry.verified = true;
 		if (!fileName)
 			return false;
-		WriteResult res;
+		WriteCasResult res;
 		if (!WriteCasFile(res, fileName, casKey))
 			return false;
 		casEntry.exists = true;
diff --git a/Engine/Source/Programs/UnrealBuildAccelerator/Common/Private/UbaStorageServer.cpp b/Engine/Source/Programs/UnrealBuildAccelerator/Common/Private/UbaStorageServer.cpp
--- a/Engine/Source/Programs/UnrealBuildAccelerator/Common/Private/UbaStorageServer.cpp
+++ b/Engine/Source/Programs/UnrealBuildAccelerator/Common/Private/UbaStorageServer.cpp
@@ -310,10 +310,7 @@ namespace uba
 		casEntry.exists = false;
 		casEntry.dropped = false;
 
-		m_server.AddWork([this, memoryHandle, size = casEntry.memorySize](const WorkContext& context)
-			{
-				m_casDataBuffer.CloseIndependent(memoryHandle, size, TC("ReleaseCasEntryNoLock"));
-			}, 1, TC("ReleaseCasEntryNoLock"));
+		m_casDataBuffer.CloseIndependent(memoryHandle, casEntry.memorySize, TC("ReleaseCasEntryNoLock"));
 	}
 
 	void StorageServer::ReleaseActiveRefCount(const CasKey& casKey, u32 clientId)
@@ -872,6 +869,10 @@ namespace uba
 				fetch.casEntry = &casEntry;
 				fetch.sendCasTime = GetTime() - start;
 
+				#if UBA_TRACK_FETCH_STORE_HINTS
+				fetch.hint = logHint.ToString();
+				#endif
+
 				return true;
 			}
 			case StorageMessageType_FetchSegment:
@@ -1103,6 +1104,10 @@ namespace uba
 				firstStore->fileAccessor = fileAccessor;
 				firstStore->recvCasTime = GetTime() - start;
 
+				#if UBA_TRACK_FETCH_STORE_HINTS
+				firstStore->hint = hint.ToString();
+				#endif
+
 				if (m_trace)
 					m_trace->FileStoreBegin(clientId, casKey, fileSize, hint, m_traceStore);
 
@@ -1216,6 +1221,11 @@ namespace uba
 					s.casEntry = firstStore->casEntry;
 					s.totalWritten = firstStore->totalWritten.load();
 					s.recvCasTime = firstStore->recvCasTime.load();
+
+					#if UBA_TRACK_FETCH_STORE_HINTS
+					s.hint = firstStore->hint;
+					#endif
+
 				}
 				return true;
 			}
diff --git a/Engine/Source/Programs/UnrealBuildAccelerator/Common/Public/UbaStorage.h b/Engine/Source/Programs/UnrealBuildAccelerator/Common/Public/UbaStorage.h
--- a/Engine/Source/Programs/UnrealBuildAccelerator/Common/Public/UbaStorage.h
+++ b/Engine/Source/Programs/UnrealBuildAccelerator/Common/Public/UbaStorage.h
@@ -101,6 +101,7 @@ namespace uba
 			u64 lastWriteTime = 0;
 			u64 storedSize = InvalidValue;
 			bool createIndependentMapping = false;
+			bool dropCasAfterUse = false;
 		};
 		using ExternalFileMappingsProvider = Function<bool(ExternalFileMapping& out, StringKey fileNameKey, const tchar* fileName)>;
 		virtual void RegisterExternalFileMappingsProvider(ExternalFileMappingsProvider&& provider) {}
@@ -126,6 +127,7 @@ namespace uba
 		bool allowDeleteVerified = false;
 		bool allowDeferredDeletes = true;
 		bool skipCleanupOnShutdown = false;
+		bool allowUnDropOfCas = true;
 		MutexHandle exclusiveMutex = InvalidMutexHandle; // Set this to hand over pre-exclusive access to storage. Storage will release mutex on shutdown
 		u8 casCompressor = 0;
 		u8 casCompressionLevel = 0;
@@ -195,10 +197,15 @@ namespace uba
 		void ReportFileInfoWeak(const StringKey& fileNameKey, u64 verifiedLastWriteTime, u64 verifiedSize);
 		bool HasBeenSeen(const CasKey& casKey);
 
+		struct WriteCasResult : WriteResult
+		{
+			bool dropAfterUse = false;
+		};
+
 		virtual void RegisterExternalFileMappingsProvider(ExternalFileMappingsProvider&& provider) override;
 		virtual bool WriteCompressedFile(WriteResult& out, const tchar* from, FileHandle readHandle, u8* readMem, u64 fileSize, const tchar* toFile, const void* header, u64 headerSize, u64 lastWriteTime = 0) final;
-		bool WriteCasFileNoCheck(WriteResult& out, const StringKey& fileNameKey, const tchar* fileName, const CasKey& casKey, const tchar* casFile, bool storeCompressed);
-		bool WriteCasFile(WriteResult& out, const tchar* fileName, const CasKey& casKey);
+		bool WriteCasFileNoCheck(WriteCasResult& out, const StringKey& fileNameKey, const tchar* fileName, const CasKey& casKey, const tchar* casFile, bool storeCompressed);
+		bool WriteCasFile(WriteCasResult& out, const tchar* fileName, const CasKey& casKey);
 		bool VerifyExisting(bool& outReturnValue, CasEntryWriteLock& casEntryLock, const CasKey& casKey, CasEntry& casEntry, StringBufferBase& casFile);
 		bool AddCasFile(StringKey fileNameKey, const tchar* fileName, const CasKey& casKey, bool deferCreation, u64 fileSize = InvalidValue, u64 lastWriteTime = 0);
 		void CasEntryAccessed(const CasKey& casKey);
@@ -311,6 +318,7 @@ namespace uba
 		bool m_allowDeleteVerified = false;
 		bool m_allowDeferredDeletes = true;
 		bool m_skipCleanupOnShutdown = false;
+		bool m_allowUnDropOfCas = true;
 		bool m_writeToDisk;
 
 		MutexHandle m_exclusiveMutex = InvalidMutexHandle;
diff --git a/Engine/Source/Programs/UnrealBuildAccelerator/Common/Public/UbaStorageServer.h b/Engine/Source/Programs/UnrealBuildAccelerator/Common/Public/UbaStorageServer.h
--- a/Engine/Source/Programs/UnrealBuildAccelerator/Common/Public/UbaStorageServer.h
+++ b/Engine/Source/Programs/UnrealBuildAccelerator/Common/Public/UbaStorageServer.h
@@ -5,6 +5,8 @@
 #include "UbaNetworkServer.h"
 #include "UbaStorage.h"
 
+#define UBA_TRACK_FETCH_STORE_HINTS 0
+
 namespace uba
 {
 	struct StorageServerCreateInfo : StorageCreateInfo
@@ -90,7 +92,12 @@ namespace uba
 			Atomic<u64> recvCasTime = 0;
 			u64 fileSize = 0;
 			u64 actualSize = 0;
+
+			#if UBA_TRACK_FETCH_STORE_HINTS
+			TString hint;
+			#endif
 		};
+
 		ReaderWriterLock m_activeStoresLock;
 		UnorderedMap<u16, ActiveStore> m_activeStores;
 
@@ -109,6 +116,10 @@ namespace uba
 			const u8* memoryPos = nullptr;
 			bool ownsMapping = false;
 
+			#if UBA_TRACK_FETCH_STORE_HINTS
+			TString hint;
+			#endif
+
 			void Release(StorageServer& server, const tchar* reason);
 		};
 		ReaderWriterLock m_activeFetchesLock;
