[UBA]
* Changed so visualizer show exit codes in hex if over 0xc0000000

diff --git a/Engine/Source/Programs/UnrealBuildAccelerator/Visualizer/Private/UbaVisualizer.cpp b/Engine/Source/Programs/UnrealBuildAccelerator/Visualizer/Private/UbaVisualizer.cpp
--- a/Engine/Source/Programs/UnrealBuildAccelerator/Visualizer/Private/UbaVisualizer.cpp
+++ b/Engine/Source/Programs/UnrealBuildAccelerator/Visualizer/Private/UbaVisualizer.cpp
@@ -2198,6 +2198,8 @@ namespace uba
 			{
 				if (process.exitCode == ProcessCancelExitCode)
 					logger.Info(L"  ExitCode: Cancelled");
+				else if (process.exitCode >= 0xC0000000u)
+					logger.Info(L"  ExitCode: 0x%x", process.exitCode);
 				else
 					logger.Info(L"  ExitCode: %7u", process.exitCode);
 			}
