[UBA]
* Fixed so received files are properly handle when process is reused.
* Fixed code path for using transient files and then manually write them to disk or fetching memory

diff --git a/Engine/Source/Programs/UnrealBuildAccelerator/Common/Private/UbaSessionServer.cpp b/Engine/Source/Programs/UnrealBuildAccelerator/Common/Private/UbaSessionServer.cpp
--- a/Engine/Source/Programs/UnrealBuildAccelerator/Common/Private/UbaSessionServer.cpp
+++ b/Engine/Source/Programs/UnrealBuildAccelerator/Common/Private/UbaSessionServer.cpp
@@ -1269,6 +1269,21 @@ namespace uba
 		return true;
 	}
 
+	bool SessionServer::WriteReceivedFiles(RemoteProcess& process, Timer& writeFilesTimer, const ConnectionInfo& connectionInfo)
+	{
+		TimerScope ts(writeFilesTimer);
+		Atomic<bool> writeSuccess = true;
+		m_server.ParallelFor(u32(process.m_receivedFiles.size()), process.m_receivedFiles, [&](const WorkContext&, auto& it)
+			{
+				if (!WriteReceivedFile(*it, &process, connectionInfo))
+					writeSuccess = false;
+				++writeFilesTimer.count;
+			}, AsView(TC("WriteReceivedFiles")));
+		--writeFilesTimer.count;
+		process.m_receivedFiles.clear();
+		return writeSuccess;
+	}
+
 	bool SessionServer::HandleSendFileToServer(const ConnectionInfo& connectionInfo, const WorkContext& workContext, BinaryReader& reader, BinaryWriter& writer)
 	{
 		u32 clientId = connectionInfo.GetId();
@@ -1635,19 +1650,8 @@ namespace uba
 
 		Timer writeFilesTimer;
 		if (exitCode == 0)
-		{
-			TimerScope ts(writeFilesTimer);
-			Atomic<bool> writeSuccess = true;
-			m_server.ParallelFor(u32(process.m_receivedFiles.size()), process.m_receivedFiles, [&](const WorkContext&, auto& it)
-				{
-					if (!WriteReceivedFile(*it, &process, connectionInfo))
-						writeSuccess = false;
-					++writeFilesTimer.count;
-				}, AsView(TC("WriteReceivedFiles")));
-			if (!writeSuccess)
+			if (!WriteReceivedFiles(process, writeFilesTimer, connectionInfo))
 				exitCode = UBA_EXIT_CODE(20);
-			--writeFilesTimer.count;
-		}
 
 		u32 id = process.m_processId;
 		Vector<ProcessLogLine> emptyLines;
@@ -1810,6 +1814,13 @@ namespace uba
 		}
 
 		auto& remoteProcess = *(RemoteProcess*)h.m_process;
+
+		Timer writeFilesTimer; // TODO: This should be added somewhere
+		if (prevExitCode == 0)
+			if (!WriteReceivedFiles(remoteProcess, writeFilesTimer, connectionInfo))
+				prevExitCode = UBA_EXIT_CODE(20);
+		//processStats.writeFiles = writeFilesTimer
+
 		SCOPED_FUTEX(remoteProcess.m_exitedLock, exitedLock);
 		NextProcessInfo nextProcess;
 		bool newProcess;
@@ -2212,7 +2223,7 @@ namespace uba
 		m_receivedFiles.erase(findIt);
 		lock2.Leave();
 
-		m_storage.DropCasFile(casKey, true, filePath.data);
+		m_storage.DropCasFile(casKey, true, filePath.data, true);
 		return true;
 	}
 
diff --git a/Engine/Source/Programs/UnrealBuildAccelerator/Common/Private/UbaStorage.cpp b/Engine/Source/Programs/UnrealBuildAccelerator/Common/Private/UbaStorage.cpp
--- a/Engine/Source/Programs/UnrealBuildAccelerator/Common/Private/UbaStorage.cpp
+++ b/Engine/Source/Programs/UnrealBuildAccelerator/Common/Private/UbaStorage.cpp
@@ -2126,9 +2126,14 @@ namespace uba
 			if (!StorageImpl::GetCasFileName(casFile, casKey))
 				return false;
 			if (reduceRefCount)
+			{
+				UBA_ASSERTF(casEntry.refCount > 0, TC("Reducing refcount of cas %s which already has 0 refcount (%s)"), casFile.data, hint);
 				--casEntry.refCount;
+			}
+			if (casEntry.refCount != 0)
+				return m_logger.Info(TC("Skipped dropping cas %s because refcount was not 0 (%s)"), casFile.data, hint).ToFalse();
 
-			UBA_ASSERTF(casEntry.refCount == 0, TC("Force deleting cas %s with refcount %u (%s)"), casFile.data, casEntry.refCount, hint);
+			//UBA_ASSERTF(casEntry.refCount == 0, TC("Force deleting cas %s with refcount %u (%s)"), casFile.data, casEntry.refCount, hint);
 			u64 sizeDeleted = 0;
 
 			if (casEntry.memoryHandle.IsValid())
diff --git a/Engine/Source/Programs/UnrealBuildAccelerator/Common/Public/UbaSessionServer.h b/Engine/Source/Programs/UnrealBuildAccelerator/Common/Public/UbaSessionServer.h
--- a/Engine/Source/Programs/UnrealBuildAccelerator/Common/Public/UbaSessionServer.h
+++ b/Engine/Source/Programs/UnrealBuildAccelerator/Common/Public/UbaSessionServer.h
@@ -125,6 +125,7 @@ namespace uba
 		
 		struct ReceivedFile;
 		bool WriteReceivedFile(ReceivedFile& file, RemoteProcess* process, const ConnectionInfo& connectionInfo);
+		bool WriteReceivedFiles(RemoteProcess& process, Timer& writeFilesTimer, const ConnectionInfo& connectionInfo);
 
 		void TraceSessionUpdate();
 
