[UBA]
* Fixed so only virtual files sets dropCasAfterUse. We don't want to drop cas files for all memory mapped files
* Added parameter to CreateVirtualFile to prevent it being added to directory table (this is WIP for coming change)

diff --git a/Engine/Source/Programs/UnrealBuildAccelerator/Common/Private/UbaSession.cpp b/Engine/Source/Programs/UnrealBuildAccelerator/Common/Private/UbaSession.cpp
--- a/Engine/Source/Programs/UnrealBuildAccelerator/Common/Private/UbaSession.cpp
+++ b/Engine/Source/Programs/UnrealBuildAccelerator/Common/Private/UbaSession.cpp
@@ -1831,7 +1831,7 @@ namespace uba
 				out.lastWriteTime = entry.lastWriteTime;
 				out.storedSize = entry.storedSize;
 				out.createIndependentMapping = entry.createIndependentMapping;
-				out.dropCasAfterUse = true; // TODO: Make this an option. This will throw away cas files after they have been used
+				out.dropCasAfterUse = entry.dropCasAfterUse; // TODO: Make this an option. This will throw away cas files after they have been used
 				return true;
 			});
 
@@ -2111,13 +2111,15 @@ namespace uba
 		return RegisterCreateFileForWrite(fileNameKey, fixedFilePath, false, sourceSize, 0, false);
 	}
 
-	bool Session::CreateVirtualFile(const tchar* filePath, const void* memory, u64 memorySize, bool transient)
+	bool Session::CreateVirtualFile(const tchar* filePath, const void* memory, u64 memorySize, bool transient, bool registerToDir)
 	{
 		UBA_ASSERT(!m_runningRemote);
 		StringBuffer<> fixedFilePath;
 		auto fileNameKey = GetKeyAndFixedName(fixedFilePath, filePath);
 		if (!CreateVirtualFileInternal(fileNameKey, fixedFilePath, memory, memorySize, transient))
 			return false;
+		if (!registerToDir)
+			return true;
 		return RegisterCreateFileForWrite(fileNameKey, fixedFilePath, false, memorySize, 0, false);
 	}
 
@@ -3692,6 +3694,7 @@ namespace uba
 		entry.lastWriteTime = 0; // TODO: Take lastwritetime of source file?
 		entry.createIndependentMapping = true; // We want independent mapping in casdb so they can be deleted
 		entry.isInvisible = transient;
+		entry.dropCasAfterUse = true;
 
 		StringBuffer<> mappingName;
 		GetMappingString(mappingName, memoryHandle, mappingOffset, entry.canBeFreed);
diff --git a/Engine/Source/Programs/UnrealBuildAccelerator/Common/Public/UbaSession.h b/Engine/Source/Programs/UnrealBuildAccelerator/Common/Public/UbaSession.h
--- a/Engine/Source/Programs/UnrealBuildAccelerator/Common/Public/UbaSession.h
+++ b/Engine/Source/Programs/UnrealBuildAccelerator/Common/Public/UbaSession.h
@@ -68,7 +68,7 @@ namespace uba
 		bool RefreshDirectory(const tchar* dirName, bool forceRegister = false); // Tell uba a directory on disk has been changed by some other system while session is running
 		bool RegisterNewFile(const tchar* filePath); // Tell uba a new file on disk has been added by some other system while session is running
 		bool RegisterVirtualFile(const tchar* filePath, const tchar* sourceFile, u64 sourceOffset, u64 sourceSize);
-		bool CreateVirtualFile(const tchar* filePath, const void* memory, u64 memorySize, bool transient);
+		bool CreateVirtualFile(const tchar* filePath, const void* memory, u64 memorySize, bool transient, bool registerToDir = true);
 		bool DeleteVirtualFile(const tchar* filePath);
 		void RegisterDeleteFile(const tchar* filePath); // Tell uba a file on disk has been deleted by some other system while session is running
 		bool RegisterNewDirectory(const tchar* directoryPath); // Tell uba a directory on disk has been added by some other system while session is running
@@ -298,6 +298,7 @@ namespace uba
 			u8 canBeFreed:1 = false;
 			u8 createIndependentMapping:1 = false;
 			u8 isInvisible:1 = false;
+			u8 dropCasAfterUse:1 = false;
 
 			#if UBA_DEBUG_TRACK_MAPPING
 			TString name;
