[UBA]
* Changed so knowninputs provide a filenamekey. If virtual files are registered with dropAfterUse, then we can end up in a scenario where another helper gets a "knowninput" where the cas content has been dropped.. and in order to be able to recreate the cas content we need to have a way to find the virtual file which the filenamekey is

diff --git a/Engine/Source/Programs/UnrealBuildAccelerator/Agent/Private/UbaAgent.cpp b/Engine/Source/Programs/UnrealBuildAccelerator/Agent/Private/UbaAgent.cpp
--- a/Engine/Source/Programs/UnrealBuildAccelerator/Agent/Private/UbaAgent.cpp
+++ b/Engine/Source/Programs/UnrealBuildAccelerator/Agent/Private/UbaAgent.cpp
@@ -184,7 +184,7 @@ namespace uba
 		for (auto file : g_ubaAgentBinaries)
 		{
 			Storage::RetrieveResult result;
-			if (!storageClient.RetrieveCasFile(result, keys[index++], file))
+			if (!storageClient.RetrieveCasFile(result, keys[index++], {}, file))
 				return false;
 			StringBuffer<256> fullFile(binDir);
 			fullFile.Append(file);
diff --git a/Engine/Source/Programs/UnrealBuildAccelerator/Common/Private/UbaSessionClient.cpp b/Engine/Source/Programs/UnrealBuildAccelerator/Common/Private/UbaSessionClient.cpp
--- a/Engine/Source/Programs/UnrealBuildAccelerator/Common/Private/UbaSessionClient.cpp
+++ b/Engine/Source/Programs/UnrealBuildAccelerator/Common/Private/UbaSessionClient.cpp
@@ -243,7 +243,7 @@ namespace uba
 		return outAttributes != 0;
 	}
 
-	bool SessionClient::RetrieveCasFile(CasKey& outNewKey, u64& outSize, const CasKey& casKey, const tchar* hint, bool storeUncompressed, bool allowProxy)
+	bool SessionClient::RetrieveCasFile(CasKey& outNewKey, u64& outSize, const CasKey& casKey, StringKey keyHint, const tchar* hint, bool storeUncompressed, bool allowProxy)
 	{
 		//StringBuffer<> workName;
 		//u32 len = TStrlen(hint);
@@ -259,7 +259,7 @@ namespace uba
 		//	storeUncompressed = !IsCompressed(casKey);
 
 		Storage::RetrieveResult result;
-		bool res = m_storage.RetrieveCasFile(result, tempKey, hint, nullptr, 1, allowProxy);
+		bool res = m_storage.RetrieveCasFile(result, tempKey, keyHint, hint, nullptr, 1, allowProxy);
 		outNewKey = result.casKey;
 		outSize = result.size;
 		return res;
@@ -347,7 +347,7 @@ namespace uba
 		bool storeUncompressed = true;
 		CasKey newKey;
 		u64 fileSize;
-		if (!RetrieveCasFile(newKey, fileSize, casKey, outVirtual.data, storeUncompressed))
+		if (!RetrieveCasFile(newKey, fileSize, casKey, {}, outVirtual.data, storeUncompressed))
 			UBA_ASSERTF(false, TC("Casfile not found for %s using %s"), outVirtual.data, CasKeyString(casKey).str);
 		StringBuffer<> destFile;
 		if (isAbsolute || fileName.Contains(TC(".."))) // This is not beautiful, but we need to keep some dlls in the sub folder (for cl.exe etc)
@@ -421,7 +421,7 @@ namespace uba
 				bool storeUncompressed = true;
 				u64 fileSize;
 				const tchar* moduleName = m.name.c_str();
-				if (!RetrieveCasFile(newCasKey, fileSize, m.casKey, moduleName, storeUncompressed))
+				if (!RetrieveCasFile(newCasKey, fileSize, m.casKey, {}, moduleName, storeUncompressed))
 				{
 					m_logger.Error(TC("Casfile not found for %s (%s)"), moduleName, CasKeyString(m.casKey).str);
 					success = false;
@@ -629,7 +629,7 @@ namespace uba
 				{
 					bool storeUncompressed = memoryMapAlignment == 0;
 					bool allowProxy = rules.AllowStorageProxy(fileName);
-					if (!RetrieveCasFile(newCasKey, fileSize, casKey, fileName.data, storeUncompressed, allowProxy))
+					if (!RetrieveCasFile(newCasKey, fileSize, casKey, fileNameKey, fileName.data, storeUncompressed, allowProxy))
 						return m_logger.Error(TC("Error retrieving cas entry %s (%s) (1)"), CasKeyString(casKey).str, fileName.data);
 
 					if (!m_storage.GetCasFileName(newName, newCasKey))
@@ -660,7 +660,7 @@ namespace uba
 						entry.handled = true;
 						Storage::RetrieveResult result;
 						bool allowProxy = rules.AllowStorageProxy(fileName);
-						if (!m_storage.RetrieveCasFile(result, casKey, fileName.data, &m_fileMappingBuffer, memoryMapAlignment, allowProxy))
+						if (!m_storage.RetrieveCasFile(result, casKey, fileNameKey, fileName.data, &m_fileMappingBuffer, memoryMapAlignment, allowProxy))
 							return m_logger.Error(TC("Error retrieving cas entry %s (%s) (2)"), CasKeyString(casKey).str, fileName.data);
 						entry.success = true;
 						entry.contentSize = result.size;
@@ -1322,9 +1322,9 @@ namespace uba
 			GetDirectoryOfCurrentModule(m_logger, dir);
 			u64 dirCount = dir.count;
 			dir.Append(PathSeparator).Append(UBA_AGENT_EXECUTABLE);
-			m_storage.CalculateCasKey(keys[0], dir.data);
+			m_storage.CalculateCasKey(keys[0], {}, dir.data);
 			dir.Resize(dirCount).Append(PathSeparator).Append(UBA_DETOURS_LIBRARY);
-			m_storage.CalculateCasKey(keys[1], dir.data);
+			m_storage.CalculateCasKey(keys[1], {}, dir.data);
 		}
 
 		writer.WriteCasKey(keys[0]);
@@ -1390,7 +1390,7 @@ namespace uba
 				{
 					TimerScope s(m_stats.storageRetrieve);
 					Storage::RetrieveResult result;
-					if (!m_storage.RetrieveCasFile(result, AsCompressed(detoursBinaryKey[i], false), UBA_DETOURS_LIBRARY))
+					if (!m_storage.RetrieveCasFile(result, AsCompressed(detoursBinaryKey[i], false), {}, UBA_DETOURS_LIBRARY))
 						return;
 				}
 				StringKey key;
@@ -1535,6 +1535,7 @@ namespace uba
 		{
 			while (knownInputsCount--)
 			{
+				StringKey knownFile = reader.ReadStringKey();
 				CasKey knownInputKey = reader.ReadCasKey();
 				u32 mappingAlignment = reader.ReadU32();
 				bool allowProxy = reader.ReadBool();
@@ -1542,11 +1543,11 @@ namespace uba
 				if (storeUncompressed)
 					knownInputKey = AsCompressed(knownInputKey, false);
 
-				m_client.AddWork([knownInputKey, allowProxy, this](const WorkContext&)
+				m_client.AddWork([knownInputKey, knownFile, allowProxy, this](const WorkContext&)
 					{
 						TrackWorkScope tws(m_client, StringBuffer<>().Appendf(TC("KnownInput")), ColorWork);
 						Storage::RetrieveResult result;
-						bool res = m_storage.RetrieveCasFile(result, knownInputKey, TC("KnownInput"), nullptr, 1, allowProxy);
+						bool res = m_storage.RetrieveCasFile(result, knownInputKey, knownFile, TC("KnownInput"), nullptr, 1, allowProxy);
 						(void)res;
 					}, 1, TC("KnownInput"));
 			}
@@ -2685,7 +2686,7 @@ namespace uba
 			#endif
 			Storage::RetrieveResult result;
 			StringBuffer<> throwaway;
-			if (m_storage.RetrieveCasFile(result, AsCompressed(detoursSymbolsKey, false), symbolsFile.data))
+			if (m_storage.RetrieveCasFile(result, AsCompressed(detoursSymbolsKey, false), {}, symbolsFile.data))
 				WriteBinFile(throwaway, symbolsFile, detoursSymbolsKey, KeyToString(StringKeyZero), DefaultAttributes());
 		}
 		return true;
diff --git a/Engine/Source/Programs/UnrealBuildAccelerator/Common/Private/UbaSessionServer.cpp b/Engine/Source/Programs/UnrealBuildAccelerator/Common/Private/UbaSessionServer.cpp
--- a/Engine/Source/Programs/UnrealBuildAccelerator/Common/Private/UbaSessionServer.cpp
+++ b/Engine/Source/Programs/UnrealBuildAccelerator/Common/Private/UbaSessionServer.cpp
@@ -113,7 +113,7 @@ namespace uba
 		u32 m_sessionId = 0;
 		TString m_executingHost;
 
-		struct KnownInput { CasKey key; u32 mappingAlignment = 0; bool allowProxy = true; };
+		struct KnownInput { StringKey file; CasKey key; u32 mappingAlignment = 0; bool allowProxy = true; };
 		KnownInput* m_knownInputs = nullptr;
 		u32 m_knownInputsCount = 0;
 		Event m_knownInputsDone;
@@ -397,6 +397,7 @@ namespace uba
 
 							UBA_ASSERT(keysIndex < knownInputsCount);
 							auto& ki = keys[keysIndex++];
+							ki.file = fileNameKey;
 							ki.key = casKey;
 							ki.mappingAlignment = GetMemoryMapAlignment(fileName, true);
 							ki.allowProxy = rules->AllowStorageProxy(fileName);
@@ -1296,7 +1297,7 @@ namespace uba
 		CasKey casKey = reader.ReadCasKey();
 
 		Storage::RetrieveResult res;
-		bool success = m_storage.RetrieveCasFile(res, casKey, file.destination.c_str(), nullptr, 1, true, clientId);
+		bool success = m_storage.RetrieveCasFile(res, casKey, file.destinationKey, file.destination.c_str(), nullptr, 1, true, clientId);
 		auto writeResponse = MakeGuard([&]() { writer.WriteBool(success); });
 
 		if (!success)
@@ -1575,6 +1576,7 @@ namespace uba
 		{
 			if (!toSendCount--)
 				break;
+			writer.WriteStringKey(kv->file);
 			writer.WriteCasKey(kv->key);
 			writer.WriteU32(kv->mappingAlignment);
 			writer.WriteBool(kv->allowProxy);
diff --git a/Engine/Source/Programs/UnrealBuildAccelerator/Common/Private/UbaStorage.cpp b/Engine/Source/Programs/UnrealBuildAccelerator/Common/Private/UbaStorage.cpp
--- a/Engine/Source/Programs/UnrealBuildAccelerator/Common/Private/UbaStorage.cpp
+++ b/Engine/Source/Programs/UnrealBuildAccelerator/Common/Private/UbaStorage.cpp
@@ -379,7 +379,10 @@ namespace uba
 					return w.CompressToMapping(out.memoryHandle, out.storedSize, fileMem, externalMapping.size, fileName);//casFile);
 				}
 				else
+				{
+					UBA_ASSERTF(!out.dropAfterUse, TC("External file mapping %s has dropAfterUse set but is not an independent mapping"), fileName);
 					return WriteCompressedFile(out, fileName, InvalidFileHandle, fileMem, externalMapping.size, casFile, nullptr, 0);
+				}
 			}
 		}
 
@@ -639,10 +642,17 @@ namespace uba
 		WriteCasResult res;
 		if (!WriteCasFileNoCheck(res, fileNameKey, fileName, casKey, casFile.data, IsCompressed(casKey)))
 			return false;
-
+		casEntry.memoryHandle = res.memoryHandle;
+		casEntry.memoryOffset = res.memoryOffset;
+		casEntry.memorySize = res.storedSize;
+		casEntry.dropped = res.dropAfterUse;
 		casEntry.exists = true;
+
 		casEntryLock.Leave();
-		CasEntryWritten(casEntry, res.storedSize);
+
+		if (!res.memoryHandle.IsValid())
+			CasEntryWritten(casEntry, res.storedSize);
+
 		return true;
 	}
 
@@ -1428,7 +1438,7 @@ namespace uba
 						}
 						else
 						{
-							if (!CalculateCasKey(checkedKey, filePath.c_str()))
+							if (!CalculateCasKey(checkedKey, {}, filePath.c_str()))
 								return;
 							checkedKey = AsCompressed(checkedKey, false);
  						}
@@ -1515,7 +1525,7 @@ namespace uba
 									casKey = ((CompressedFileHeader*)file.GetData())->casKey;
 							}
 							if (casKey == CasKeyZero)
-								if (!CalculateCasKey(casKey, p.data()))
+								if (!CalculateCasKey(casKey, {}, p.data()))
 									m_logger.Warning(TC("Failed to calculate cas key for %s"), p.data());
 
 							if (casKey != fe.casKey)
@@ -1596,7 +1606,7 @@ namespace uba
 		return res;
 	}
 
-	bool StorageImpl::RetrieveCasFile(RetrieveResult& out, const CasKey& casKey, const tchar* hint, FileMappingBuffer* mappingBuffer, u64 memoryMapAlignment, bool allowProxy, u32 clientId)
+	bool StorageImpl::RetrieveCasFile(RetrieveResult& out, const CasKey& casKey, StringKey keyHint, const tchar* hint, FileMappingBuffer* mappingBuffer, u64 memoryMapAlignment, bool allowProxy, u32 clientId)
 	{
 		UBA_ASSERT(false);
 		return false;
@@ -2181,8 +2191,19 @@ namespace uba
 		return true;
 	}
 
-	bool StorageImpl::CalculateCasKey(CasKey& out, const tchar* fileName)
+	bool StorageImpl::CalculateCasKey(CasKey& out, StringKey fileNameKey, const tchar* fileName)
 	{
+		ExternalFileMapping externalMapping;
+		if (m_externalFileMappingsProvider && m_externalFileMappingsProvider(externalMapping, fileNameKey, fileName))
+		{
+			MemoryMappingScope scope(m_logger, m_sharedMemory);
+			u8* fileMem = scope.MapView(TC("CalculateCasKey"), externalMapping.handle, externalMapping.offset, externalMapping.size, fileName);
+			if (!fileMem)
+				return false;
+			out = CalculateCasKey(fileMem, externalMapping.size, true);
+			return out != CasKeyZero;
+		}
+
 		FileHandle fileHandle;
 		if (!OpenFileSequentialRead(m_logger, fileName, fileHandle))
 			return m_logger.Error(TC("[CalculateCasKey] OpenFileSequentialRead failed for %s (%s)"), fileName, LastErrorToText().data);
diff --git a/Engine/Source/Programs/UnrealBuildAccelerator/Common/Private/UbaStorageClient.cpp b/Engine/Source/Programs/UnrealBuildAccelerator/Common/Private/UbaStorageClient.cpp
--- a/Engine/Source/Programs/UnrealBuildAccelerator/Common/Private/UbaStorageClient.cpp
+++ b/Engine/Source/Programs/UnrealBuildAccelerator/Common/Private/UbaStorageClient.cpp
@@ -211,7 +211,7 @@ namespace uba
 		return true;
 	}
 
-	bool StorageClient::RetrieveCasFile(RetrieveResult& out, const CasKey& casKeyTmp, const tchar* hint, FileMappingBuffer* mappingBuffer, u64 memoryMapAlignment, bool allowProxy, u32 clientId)
+	bool StorageClient::RetrieveCasFile(RetrieveResult& out, const CasKey& casKeyTmp, StringKey keyHint, const tchar* hint, FileMappingBuffer* mappingBuffer, u64 memoryMapAlignment, bool allowProxy, u32 clientId)
 	{
 		CasKey casKey = casKeyTmp;
 
@@ -370,6 +370,7 @@ namespace uba
 				NetworkMessage msg(*client, ServiceId, StorageMessageType_FetchBegin, writer);
 				writer.WriteByte(wantsProxy ? 1 : 0);
 				writer.WriteCasKey(casKey);
+				writer.WriteStringKey(keyHint);
 				writer.WriteString(hint);
 				BinaryReader reader(slot, 0, SendMaxSize);
 				if (!msg.Send(reader))
@@ -1111,7 +1112,7 @@ namespace uba
 							return;
 
 						CasKey casKey;
-						if (!CalculateCasKey(casKey, filePath.c_str()))
+						if (!CalculateCasKey(casKey, {}, filePath.c_str()))
 						{
 							m_logger.Error(TC("Failed to calculate cas key for %s"), filePath.c_str());
 							return;
diff --git a/Engine/Source/Programs/UnrealBuildAccelerator/Common/Private/UbaStorageServer.cpp b/Engine/Source/Programs/UnrealBuildAccelerator/Common/Private/UbaStorageServer.cpp
--- a/Engine/Source/Programs/UnrealBuildAccelerator/Common/Private/UbaStorageServer.cpp
+++ b/Engine/Source/Programs/UnrealBuildAccelerator/Common/Private/UbaStorageServer.cpp
@@ -96,7 +96,7 @@ namespace uba
 		return true;
 	}
 
-	bool StorageServer::RetrieveCasFile(RetrieveResult& out, const CasKey& casKey, const tchar* hint, FileMappingBuffer* mappingBuffer, u64 memoryMapAlignment, bool allowProxy, u32 clientId)
+	bool StorageServer::RetrieveCasFile(RetrieveResult& out, const CasKey& casKey, StringKey keyHint, const tchar* hint, FileMappingBuffer* mappingBuffer, u64 memoryMapAlignment, bool allowProxy, u32 clientId)
 	{
 		UBA_ASSERT(!mappingBuffer);
 		UBA_ASSERT(casKey != CasKeyZero);
@@ -580,6 +580,7 @@ namespace uba
 
 				u64 start = GetTime();
 				casKey = reader.ReadCasKey();
+				StringKey keyHint = reader.ReadStringKey();
 				reader.ReadString(hint);
 
 				StringBuffer<16> temp;
@@ -614,7 +615,10 @@ namespace uba
 						// Last resort.. use hint to load file into cas (hint should be renamed since it is now a critical parameter)
 						// We better check the caskey first to make sure it is matching on the server
 
-						StringKey fileNameKey = CaseInsensitiveFs ? ToStringKeyLower(hint) : ToStringKey(hint);
+						StringKey fileNameKey = keyHint;
+						if (keyHint == StringKeyZero)
+							fileNameKey = CaseInsensitiveFs ? ToStringKeyLower(hint) : ToStringKey(hint);
+
 						CasKey checkedCasKey;
 						{
 							SCOPED_READ_LOCK(m_fileTableLookupLock, lookupLock);
@@ -631,7 +635,7 @@ namespace uba
 						if (checkedCasKey == CasKeyZero)
 						{
 							m_logger.Info(TC("Server did not find cas %s for %s in file table lookup. Recalculating cas key"), CasKeyString(casKey).str, hint.data);
-							if (!CalculateCasKey(checkedCasKey, hint.data))
+							if (!CalculateCasKey(checkedCasKey, fileNameKey, hint.data))
 							{
 								m_logger.Error(TC("FetchBegin failed for cas file %s (%s) requested by client %u. Can't calculate cas key for file"), CasKeyString(casKey).str, hint.data, connectionInfo.GetId());
 								writer.WriteU16(0);
diff --git a/Engine/Source/Programs/UnrealBuildAccelerator/Common/Public/UbaSessionClient.h b/Engine/Source/Programs/UnrealBuildAccelerator/Common/Public/UbaSessionClient.h
--- a/Engine/Source/Programs/UnrealBuildAccelerator/Common/Public/UbaSessionClient.h
+++ b/Engine/Source/Programs/UnrealBuildAccelerator/Common/Public/UbaSessionClient.h
@@ -55,7 +55,7 @@ namespace uba
 		bool Exists(const StringView& path, u32& outAttributes);
 
 	private:
-		bool RetrieveCasFile(CasKey& outNewKey, u64& outSize, const CasKey& casKey, const tchar* hint, bool storeUncompressed, bool allowProxy = true);
+		bool RetrieveCasFile(CasKey& outNewKey, u64& outSize, const CasKey& casKey, StringKey keyHint, const tchar* hint, bool storeUncompressed, bool allowProxy = true);
 
 		virtual bool PrepareProcess(ProcessImpl& process, bool isChild, StringBufferBase& outRealApplication, const tchar*& outRealWorkingDir) override;
 		virtual void* GetProcessEnvironmentVariables() override;
diff --git a/Engine/Source/Programs/UnrealBuildAccelerator/Common/Public/UbaStorage.h b/Engine/Source/Programs/UnrealBuildAccelerator/Common/Public/UbaStorage.h
--- a/Engine/Source/Programs/UnrealBuildAccelerator/Common/Public/UbaStorage.h
+++ b/Engine/Source/Programs/UnrealBuildAccelerator/Common/Public/UbaStorage.h
@@ -52,7 +52,7 @@ namespace uba
 		virtual bool GetFileData(void* outData, const CasKey& casKey, const tchar* hint) = 0;
 
 		struct RetrieveResult { CasKey casKey; u64 size = 0; MappedView2 view; };
-		virtual bool RetrieveCasFile(RetrieveResult& out, const CasKey& casKey, const tchar* hint, FileMappingBuffer* mappingBuffer = nullptr, u64 memoryMapAlignment = 1, bool allowProxy = true, u32 clientId = 0) = 0;
+		virtual bool RetrieveCasFile(RetrieveResult& out, const CasKey& casKey, StringKey keyHint, const tchar* hint, FileMappingBuffer* mappingBuffer = nullptr, u64 memoryMapAlignment = 1, bool allowProxy = true, u32 clientId = 0) = 0;
 
 		struct CachedFileInfo { CasKey casKey; };
 		virtual bool VerifyAndGetCachedFileInfo(CachedFileInfo& out, StringKey fileNameKey, u64 verifiedLastWriteTime, u64 verifiedSize) = 0;
@@ -62,7 +62,7 @@ namespace uba
 		virtual bool StoreCasFileClient(CasKey& out, StringKey fileNameKey, const tchar* fileName, SharedMemoryHandle memoryHandle, u64 memoryOffset, u64 fileSize, const tchar* hint, bool keepMappingInMemory = false, bool storeCompressed = true) = 0;
 		virtual bool DropCasFile(const CasKey& casKey, bool forceDelete, const tchar* hint, bool reduceRefCount = false) = 0;
 		virtual bool ReportBadCasFile(const CasKey& casKey) = 0;
-		virtual bool CalculateCasKey(CasKey& out, const tchar* fileName) = 0;
+		virtual bool CalculateCasKey(CasKey& out, StringKey fileNameKey, const tchar* fileName) = 0;
 
 		using FormattingFunc = Function<bool(MemoryBlock& destData, const void* sourceData, u64 sourceSize, const tchar* hint)>;
 		virtual bool CopyOrLink(const CasKey& casKey, const tchar* destination, u32 fileAttributes, bool writeCompressed = false, u32 clientId = 0, const FormattingFunc& formattingFunc = {}, bool isTemp = false, bool allowHardLink = true) = 0;
@@ -163,14 +163,14 @@ namespace uba
 		virtual bool DeleteCasForFile(const tchar* file) override;
 		virtual bool GetFileSize(u64& outSize, const CasKey& casKey, const tchar* hint) override;
 		virtual bool GetFileData(void* outData, const CasKey& casKey, const tchar* hint) override;
-		virtual bool RetrieveCasFile(RetrieveResult& out, const CasKey& casKey, const tchar* hint, FileMappingBuffer* mappingBuffer = nullptr, u64 memoryMapAlignment = 1, bool allowProxy = true, u32 clientId = 0) override;
+		virtual bool RetrieveCasFile(RetrieveResult& out, const CasKey& casKey, StringKey keyHint, const tchar* hint, FileMappingBuffer* mappingBuffer = nullptr, u64 memoryMapAlignment = 1, bool allowProxy = true, u32 clientId = 0) override;
 		virtual bool VerifyAndGetCachedFileInfo(CachedFileInfo& out, StringKey fileNameKey, u64 verifiedLastWriteTime, u64 verifiedSize) override;
 		virtual bool InvalidateCachedFileInfo(StringKey fileNameKey) override;
 		virtual bool StoreCasFile(CasKey& out, const tchar* fileName, const CasKey& casKeyOverride, bool deferCreation) override;
 		virtual bool StoreCasFileClient(CasKey& out, StringKey fileNameKey, const tchar* fileName, SharedMemoryHandle memoryHandle, u64 memoryOffset, u64 fileSize, const tchar* hint, bool keepMappingInMemory = false, bool storeCompressed = true) override;
 		virtual bool DropCasFile(const CasKey& casKey, bool forceDelete, const tchar* hint, bool reduceRefCount = false) override;
 		virtual bool ReportBadCasFile(const CasKey& casKey) override;
-		virtual bool CalculateCasKey(CasKey& out, const tchar* fileName) override;
+		virtual bool CalculateCasKey(CasKey& out, StringKey fileNameKey, const tchar* fileName) override;
 		virtual bool CopyOrLink(const CasKey& casKey, const tchar* destination, u32 fileAttributes, bool writeCompressed = false, u32 clientId = 0, const FormattingFunc& formattingFunc = {}, bool isTemp = false, bool allowHardLink = true) override;
 		virtual bool FakeCopy(const CasKey& casKey, const tchar* destination, u64 size, u64 lastWritten, u32 fileAttributes, bool deleteExisting = true) override;
 		virtual bool GetCasFileName(StringBufferBase& out, const CasKey& casKey) override;
diff --git a/Engine/Source/Programs/UnrealBuildAccelerator/Common/Public/UbaStorageClient.h b/Engine/Source/Programs/UnrealBuildAccelerator/Common/Public/UbaStorageClient.h
--- a/Engine/Source/Programs/UnrealBuildAccelerator/Common/Public/UbaStorageClient.h
+++ b/Engine/Source/Programs/UnrealBuildAccelerator/Common/Public/UbaStorageClient.h
@@ -55,7 +55,7 @@ namespace uba
 		virtual MappedView2 MapView(const CasKey& casKey, const tchar* hint) override;
 
 		virtual bool GetZone(StringBufferBase& out) override;
-		virtual bool RetrieveCasFile(RetrieveResult& out, const CasKey& casKey, const tchar* hint, FileMappingBuffer* mappingBuffer = nullptr, u64 memoryMapAlignment = 1, bool allowProxy = true, u32 clientId = 0) override;
+		virtual bool RetrieveCasFile(RetrieveResult& out, const CasKey& casKey, StringKey keyHint, const tchar* hint, FileMappingBuffer* mappingBuffer = nullptr, u64 memoryMapAlignment = 1, bool allowProxy = true, u32 clientId = 0) override;
 		virtual bool StoreCasFile(CasKey& out, const tchar* fileName, const CasKey& casKeyOverride, bool deferCreation) override;
 		virtual bool StoreCasFileClient(CasKey& out, StringKey fileNameKey, const tchar* fileName, SharedMemoryHandle mappingFile, u64 mappingOffset, u64 fileSize, const tchar* hint, bool keepMappingInMemory = false, bool storeCompressed = true) override;
 		virtual bool HasCasFile(const CasKey& casKey, CasEntry** out = nullptr) override;
diff --git a/Engine/Source/Programs/UnrealBuildAccelerator/Common/Public/UbaStorageServer.h b/Engine/Source/Programs/UnrealBuildAccelerator/Common/Public/UbaStorageServer.h
--- a/Engine/Source/Programs/UnrealBuildAccelerator/Common/Public/UbaStorageServer.h
+++ b/Engine/Source/Programs/UnrealBuildAccelerator/Common/Public/UbaStorageServer.h
@@ -43,7 +43,7 @@ namespace uba
 	protected:
 
 		virtual bool GetZone(StringBufferBase& out) override;
-		virtual bool RetrieveCasFile(RetrieveResult& out, const CasKey& casKey, const tchar* hint, FileMappingBuffer* mappingBuffer = nullptr, u64 memoryMapAlignment = 1, bool allowProxy = true, u32 clientId = 0) override;
+		virtual bool RetrieveCasFile(RetrieveResult& out, const CasKey& casKey, StringKey keyHint, const tchar* hint, FileMappingBuffer* mappingBuffer = nullptr, u64 memoryMapAlignment = 1, bool allowProxy = true, u32 clientId = 0) override;
 		virtual bool IsDisallowedPath(const tchar* fileName) override;
 		virtual void SetTrace(Trace* trace, bool detailed) override;
 		virtual bool HasProxy(u32 clientId) override;
diff --git a/Engine/Source/Programs/UnrealBuildAccelerator/Test/Private/UbaTestStorage.cpp b/Engine/Source/Programs/UnrealBuildAccelerator/Test/Private/UbaTestStorage.cpp
--- a/Engine/Source/Programs/UnrealBuildAccelerator/Test/Private/UbaTestStorage.cpp
+++ b/Engine/Source/Programs/UnrealBuildAccelerator/Test/Private/UbaTestStorage.cpp
@@ -216,14 +216,51 @@ namespace uba
 #endif
 
 		CasKey casKey;
-		if (!storageServer.CalculateCasKey(casKey, fileName.data))
+		if (!storageServer.CalculateCasKey(casKey, {}, fileName.data))
 			return false;
 
 		casKey = AsCompressed(casKey, false);
 
 		Storage::RetrieveResult result;
-		if (!storageClient.RetrieveCasFile(result, casKey, fileName.data))
-			return logger.Error(TC("Failed to store file %s"), fileName.data);
+		if (!storageClient.RetrieveCasFile(result, casKey, {}, fileName.data))
+			return logger.Error(TC("Failed to retrieve file %s"), fileName.data);
+
+		storageClient.DropCasFile(casKey, true, fileName.data, false);
+		storageServer.DropCasFile(casKey, true, fileName.data, false);
+		if (!storageClient.RetrieveCasFile(result, casKey, {}, fileName.data))
+			return logger.Error(TC("Failed to retrieve file %s"), fileName.data);
+
+
+
+		StringKey testFileNameKey = ToStringKeyLower(fileName);
+
+		auto& sharedMem = storageServer.GetSharedMemory();
+		auto memHandle = sharedMem.CreateHandle(logger, TC(""));
+		sharedMem.ExtendMemory(memHandle, 1);
+		u8* mem = sharedMem.MapView(memHandle, TC(""), SharedMemoryMapType_ReadWrite);
+		*mem = 'H';
+		sharedMem.UnmapView(memHandle, mem);
+
+		storageServer.RegisterExternalFileMappingsProvider([&](Storage::ExternalFileMapping& out, StringKey fileNameKey, const tchar* fileName)
+			{
+				if (fileNameKey != testFileNameKey)
+					return false;
+				out.handle = memHandle;
+				out.size = 1;
+				out.dropCasAfterUse = true;
+				out.createIndependentMapping = true;
+				return true;
+			});
+
+		storageClient.DropCasFile(casKey, true, fileName.data, false);
+		storageServer.DropCasFile(casKey, true, fileName.data, false);
+		if (!storageClient.RetrieveCasFile(result, casKey, testFileNameKey, TC("KnownInput")))
+			return logger.Error(TC("Failed to retrieve file %s"), fileName.data);
+
+		storageClient.DropCasFile(casKey, true, fileName.data, false);
+		storageServer.DropCasFile(casKey, true, fileName.data, false);
+		if (!storageClient.RetrieveCasFile(result, casKey, testFileNameKey, TC("KnownInput")))
+			return logger.Error(TC("Failed to retrieve file %s"), fileName.data);
 
 		return true;
 	}
