[UBA]
* Changed proxy code fetch begin message to be compatible with previous change

diff --git a/Engine/Source/Programs/UnrealBuildAccelerator/Common/Private/UbaStorageProxy.cpp b/Engine/Source/Programs/UnrealBuildAccelerator/Common/Private/UbaStorageProxy.cpp
--- a/Engine/Source/Programs/UnrealBuildAccelerator/Common/Private/UbaStorageProxy.cpp
+++ b/Engine/Source/Programs/UnrealBuildAccelerator/Common/Private/UbaStorageProxy.cpp
@@ -181,6 +181,7 @@ namespace uba
 		UBA_ASSERT((recvFlags & 2) == 0);(void)recvFlags;
 
 		CasKey casKey = reader.ReadCasKey();
+		StringKey keyHint = reader.ReadStringKey();
 		StringBuffer<> hint;
 		reader.ReadString(hint);
 
@@ -220,6 +221,7 @@ namespace uba
 			auto& writer2 = file.beginMessage->writer;
 			writer2.WriteByte(2); // Does not want proxy but informs it is proxy
 			writer2.WriteCasKey(casKey);
+			writer2.WriteStringKey(keyHint);
 			writer2.WriteString(hint);
 			writer2.WriteBytes(reader.GetPositionData(), reader.GetLeft());
 
