[UBA]
* Fixed ubacli compile error

diff --git a/Engine/Source/Programs/UnrealBuildAccelerator/Cli/Private/UbaCli.cpp b/Engine/Source/Programs/UnrealBuildAccelerator/Cli/Private/UbaCli.cpp
--- a/Engine/Source/Programs/UnrealBuildAccelerator/Cli/Private/UbaCli.cpp
+++ b/Engine/Source/Programs/UnrealBuildAccelerator/Cli/Private/UbaCli.cpp
@@ -704,9 +704,9 @@ namespace uba
 						{
 							Storage::RetrieveResult res;
 							storageServer.EnsureCasFile(casKey, TC("Dummy"));
-							if (!client.storageClient->RetrieveCasFile(res, AsCompressed(casKey, false), TC("")))
+							if (!client.storageClient->RetrieveCasFile(res, AsCompressed(casKey, false), {}, TC("")))
 								success = false;
-							if (!client.storageClient->RetrieveCasFile(res, casKey, TC("")))
+							if (!client.storageClient->RetrieveCasFile(res, casKey, {}, TC("")))
 								success = false;
 
 							#if 0
