[UBA]
* Implemented file parallel send so agents can use multiple sockets to send files. Disabled by default
* Added unit test to test storing large file

diff --git a/Engine/Source/Programs/UnrealBuildAccelerator/Common/Private/UbaStorageClient.cpp b/Engine/Source/Programs/UnrealBuildAccelerator/Common/Private/UbaStorageClient.cpp
--- a/Engine/Source/Programs/UnrealBuildAccelerator/Common/Private/UbaStorageClient.cpp
+++ b/Engine/Source/Programs/UnrealBuildAccelerator/Common/Private/UbaStorageClient.cpp
@@ -36,6 +36,7 @@ namespace uba
 		table.GetValueAsBool(sendCompressed, TC("SendCompressed"));
 		table.GetValueAsBool(allowProxy, TC("AllowProxy"));
 		table.GetValueAsBool(sendOneBigFileAtTheTime, TC("SendOneBigAtTheTime"));
+		table.GetValueAsBool(fileParallelSend, TC("FileParallelSend"));
 		table.GetValueAsBool(checkExistsOnServer, TC("CheckExistsOnServer"));
 		table.GetValueAsBool(resendCas, TC("ResendCas"));
 		table.GetValueAsU32(proxyConnectionCount, TC("ProxyConnectionCount"));
@@ -47,6 +48,7 @@ namespace uba
 	,	m_sendCompressed(info.sendCompressed)
 	,	m_allowProxy(info.allowProxy)
 	,	m_sendOneBigFileAtTheTime(info.sendOneBigFileAtTheTime)
+	,	m_fileParallelSend(info.fileParallelSend)
 	,	m_checkExistsOnServer(info.checkExistsOnServer)
 	,	m_resendCas(info.resendCas)
 	,	m_zone(info.zone)
@@ -999,6 +1001,7 @@ namespace uba
 			{
 				FileSender sender { m_logger, m_client, m_bufferSlots, Stats(), m_sendOneAtTheTimeLock, m_casCompressor, m_casCompressionLevel };
 				sender.m_sendOneBigFileAtTheTime = m_sendOneBigFileAtTheTime;
+				sender.m_parallelSend = m_fileParallelSend;
 				if (!sender.SendFileCompressed(casKey, fileName, fileMem, fileSize, hint))
 					return false;
 			}
diff --git a/Engine/Source/Programs/UnrealBuildAccelerator/Common/Private/UbaStorageUtils.cpp b/Engine/Source/Programs/UnrealBuildAccelerator/Common/Private/UbaStorageUtils.cpp
--- a/Engine/Source/Programs/UnrealBuildAccelerator/Common/Private/UbaStorageUtils.cpp
+++ b/Engine/Source/Programs/UnrealBuildAccelerator/Common/Private/UbaStorageUtils.cpp
@@ -34,6 +34,25 @@ namespace uba
 
 	////////////////////////////////////////////////////////////////////////////////////////////////////
 
+	struct IndexContainer
+	{
+		struct iterator
+		{
+			iterator() : index(0) {}
+			iterator(u32 i) : index(i) {}
+			iterator operator++(int) { return iterator(index++); }
+			bool operator==(const iterator& o) const { return index == o.index; }
+			u32 index;
+		};
+		IndexContainer(u64 c) : count(u32(c)) {}
+		iterator begin() const { return iterator(0); }
+		iterator end() const { return iterator(count); }
+		u32 size() const { return count; }
+		u32 count;
+	};
+
+	////////////////////////////////////////////////////////////////////////////////////////////////////
+
 	u8 GetCompressor(const tchar* str)
 	{
 		#define OODLE_COMPRESSOR(x) if (Equals(str, TC(#x))) return u8(OodleLZ_Compressor_##x);
@@ -341,22 +360,7 @@ namespace uba
 		u32 connectionCount = client.GetConnectionCount();
 		if (workManager && connectionCount > 1) // Use parallel send.. to make sure all sockets send buffers are fully utilized
 		{
-			struct Container
-			{
-				struct iterator
-				{
-					iterator() : index(0) {}
-					iterator(u32 i) : index(i) {}
-					iterator operator++(int) { return iterator(index++); }
-					bool operator==(const iterator& o) const { return index == o.index; }
-					u32 index;
-				};
-				Container(u64 c) : count(u32(c)) {}
-				iterator begin() const { return iterator(0); }
-				iterator end() const { return iterator(count); }
-				u32 size() { return count; }
-				u32 count;
-			} container(sendCount);
+			IndexContainer container(sendCount);
 
 			Atomic<bool> error = false;
 
@@ -462,7 +466,6 @@ namespace uba
 		u64 fileSize = memoryBlock.writtenSize - messageHeaderMaxSize;
 
 		u16 storeId = 0;
-		bool isFirst = true;
 		bool sendEnd = false;
 		u64 sendLeft = fileSize;
 		u64 sendPos = 0;
@@ -480,71 +483,89 @@ namespace uba
 		bool hasSendOneAtTheTimeLock = false;
 		auto lockGuard = MakeGuard([&]() { if (hasSendOneAtTheTimeLock) m_sendOneAtTheTimeLock.Leave(); });
 
-		while (sendLeft)
 		{
-			u64 writerStartOffset = messageHeader + (isFirst ? firstMessageOverHead : (sizeof(u16) + sizeof(u64)));
-			BinaryWriter writer(readData + sendPos - writerStartOffset, 0, client.GetMessageMaxSize());
-			NetworkMessage msg(client, StorageServiceId, isFirst ? StorageMessageType_StoreBegin : StorageMessageType_StoreSegment, writer);
-			if (isFirst)
-			{
-				writer.WriteCasKey(casKey);
-				writer.WriteU64(fileSize);
-				writer.WriteU64(sourceSize);
-				writer.WriteString(hint);
-			}
-			else
-			{
-				UBA_ASSERT(storeId != 0);
-				writer.WriteU16(storeId);
-				writer.WriteU64(sendPos);
-			}
-
+			u64 writerStartOffset = messageHeader + firstMessageOverHead;
+			BinaryWriter writer(readData - writerStartOffset, 0, client.GetMessageMaxSize());
+			NetworkMessage msg(client, StorageServiceId, StorageMessageType_StoreBegin, writer);
+			writer.WriteCasKey(casKey);
+			writer.WriteU64(fileSize);
+			writer.WriteU64(sourceSize);
+			writer.WriteString(hint);
 			u64 capacityLeft = writer.GetCapacityLeft();
 			u64 toWrite = Min(sendLeft, capacityLeft);
 			writer.AllocWrite(toWrite);
-
 			sendLeft -= toWrite;
 			sendPos += toWrite;
 
-			bool isDone = sendLeft == 0;
-
-			if (isFirst && !isDone && m_sendOneBigFileAtTheTime)
+			if (sendLeft != 0 && m_sendOneBigFileAtTheTime)
 			{
 				m_sendOneAtTheTimeLock.Enter();
 				hasSendOneAtTheTimeLock = true;
 			}
+			StackBinaryReader<128> reader;
+			if (!msg.Send(reader))
+				return false;
+			storeId = reader.ReadU16();
+			sendEnd = reader.ReadBool();
+		}
 
-			if (isFirst) // First message must always be acknowledged (provide a reader) to make sure there is an entry on server that can be waited on.
-			{
-				StackBinaryReader<128> reader;
-				if (!msg.Send(reader))
-					return false;
-				storeId = reader.ReadU16();
-				sendEnd = reader.ReadBool();
-				if (isDone)
-					break;
+		if (sendLeft != 0)
+		{
+			if (!storeId) // Zero means error
+				return m_logger.Error(TC("Server failed to start storing file %s (%s)"), CasKeyString(casKey).str, hint);
+			if (storeId == u16(~0)) // File already exists on server
+				return sendEndMessage(); //m_logger.Info(TC("Server already has file %s (%s)"), CasKeyString(casKey).str, hint);
 
-				if (!storeId) // Zero means error
-					return m_logger.Error(TC("Server failed to start storing file %s (%s)"), CasKeyString(casKey).str, hint);
+			u64 writerHeaderSize = messageHeader + (sizeof(u16) + sizeof(u64));
+			u64 writerCapacity = client.GetMessageMaxSize() - writerHeaderSize;
+			u64 sendCount = (sendLeft + writerCapacity - 1) / writerCapacity;
 
-				if (storeId == u16(~0)) // File already exists on server
+			auto SendSegment = [&](u64 index)
 				{
-					//m_logger.Info(TC("Server already has file %s (%s)"), CasKeyString(casKey).str, hint);
-					return sendEndMessage();
-				}
+					u64 segmentPos = sendPos + writerCapacity*index;
+					u8* writerStart = readData + segmentPos - writerHeaderSize;
+					u8 prevData[32];
+					memcpy(prevData, writerStart, writerHeaderSize);
+
+					BinaryWriter writer(writerStart, 0, client.GetMessageMaxSize());
+					NetworkMessage msg(client, StorageServiceId, StorageMessageType_StoreSegment, writer);
+					writer.WriteU16(storeId);
+					writer.WriteU64(segmentPos);
+					u64 toWrite = Min(fileSize - segmentPos, writerCapacity);
+					writer.AllocWrite(toWrite);
+					bool res = msg.Send();
+					memcpy(writerStart, prevData, writerHeaderSize);
+					return res;
+				};
 
-				isFirst = false;
+			// If we have more than one connection we can use parallel for to send to multiple connections at the same time.
+			u32 connectionCount = m_parallelSend && sendCount > 1 ? client.GetConnectionCount() : 1;
+			if (connectionCount > 1)
+			{
+				// We have to leapfrog every other message since the SendSegment modifies the memory just before the sent segment
+				for (u32 i=0;i!=2; ++i)
+				{
+					Atomic<bool> error;
+					IndexContainer container((sendCount+1)/2);
+					m_client.ParallelFor(connectionCount - 1, container, [&](const WorkContext&, auto& it)
+					{
+						u32 index = it.index*2+i;
+						if (index < sendCount)
+							if (!SendSegment(index))
+								error = true;
+					}, TCV("SendCompressedFile"));
+					if (error)
+						return false;
+				}
 			}
 			else
 			{
-				if (!msg.Send())
-					return false;
-				if (isDone)
-					break;
+				for (u64 i=0; i!=sendCount; ++i)
+					if (!SendSegment(i))
+						return false;
 			}
 		}
 
-
 		m_stats.sendCasBytesRaw += sourceSize;
 		m_stats.sendCasBytesComp += fileSize;
 		m_bytesSent = fileSize;
diff --git a/Engine/Source/Programs/UnrealBuildAccelerator/Common/Public/UbaStorageClient.h b/Engine/Source/Programs/UnrealBuildAccelerator/Common/Public/UbaStorageClient.h
--- a/Engine/Source/Programs/UnrealBuildAccelerator/Common/Public/UbaStorageClient.h
+++ b/Engine/Source/Programs/UnrealBuildAccelerator/Common/Public/UbaStorageClient.h
@@ -25,6 +25,7 @@ namespace uba
 		bool sendCompressed = true;
 		bool allowProxy = true;
 		bool sendOneBigFileAtTheTime = true;
+		bool fileParallelSend = false;
 		bool checkExistsOnServer = false; // There is no point having this enabled if server is not writing received cas files to disk (which is default)
 		bool resendCas = false; // Will try to send created file everytime even if the cas has been seen before
 		
@@ -69,6 +70,7 @@ namespace uba
 		bool m_sendCompressed;
 		bool m_allowProxy;
 		bool m_sendOneBigFileAtTheTime;
+		bool m_fileParallelSend;
 		bool m_checkExistsOnServer;
 		bool m_resendCas;
 
diff --git a/Engine/Source/Programs/UnrealBuildAccelerator/Common/Public/UbaStorageUtils.h b/Engine/Source/Programs/UnrealBuildAccelerator/Common/Public/UbaStorageUtils.h
--- a/Engine/Source/Programs/UnrealBuildAccelerator/Common/Public/UbaStorageUtils.h
+++ b/Engine/Source/Programs/UnrealBuildAccelerator/Common/Public/UbaStorageUtils.h
@@ -42,6 +42,7 @@ namespace uba
 		u8 m_casCompressor = DefaultCompressor;
 		u8 m_casCompressionLevel = DefaultCompressionLevel;
 		bool m_sendOneBigFileAtTheTime = true;
+		bool m_parallelSend = false;
 		u64 m_bytesSent = 0;
 	};
 
diff --git a/Engine/Source/Programs/UnrealBuildAccelerator/Test/Private/UbaTest.cpp b/Engine/Source/Programs/UnrealBuildAccelerator/Test/Private/UbaTest.cpp
--- a/Engine/Source/Programs/UnrealBuildAccelerator/Test/Private/UbaTest.cpp
+++ b/Engine/Source/Programs/UnrealBuildAccelerator/Test/Private/UbaTest.cpp
@@ -44,6 +44,19 @@ namespace uba
 		return FileExists(logger, testFileName.data);
 	}
 
+	bool FilesEqual(Logger& logger, const tchar* file1, const tchar* file2)
+	{
+		FileAccessor f1(logger, file1);
+		if (!f1.OpenMemoryRead())
+			return false;
+		FileAccessor f2(logger, file2);
+		if (!f2.OpenMemoryRead())
+			return false;
+		if (f1.GetSize() != f2.GetSize())
+			return false;
+		return memcmp(f1.GetData(), f2.GetData(), f1.GetSize()) == 0;
+	}
+
 	bool WrappedMain(int argc, tchar* argv[])
 	{
 		AddExceptionHandler();
diff --git a/Engine/Source/Programs/UnrealBuildAccelerator/Test/Private/UbaTest.h b/Engine/Source/Programs/UnrealBuildAccelerator/Test/Private/UbaTest.h
--- a/Engine/Source/Programs/UnrealBuildAccelerator/Test/Private/UbaTest.h
+++ b/Engine/Source/Programs/UnrealBuildAccelerator/Test/Private/UbaTest.h
@@ -22,6 +22,7 @@ namespace uba
 	bool CreateTestFile(StringBufferBase& outFile, Logger& logger, StringView testRootDir, StringView fileName, StringView content, u32 attributes = DefaultAttributes());
 	bool DeleteTestFile(Logger& logger, StringView testRootDir, StringView fileName);
 	bool FileExists(Logger& logger, StringView dir, StringView fileName);
+	bool FilesEqual(Logger& logger, const tchar* file1, const tchar* file2);
 
 	using RunProcessFunction = Function<ProcessHandle(const ProcessStartInfo&)>;
 	using TestSessionFunction = Function<bool(LoggerWithWriter& logger, SessionServer& session, const tchar* workingDir, const RunProcessFunction& runProcess, void* extraData)>;
diff --git a/Engine/Source/Programs/UnrealBuildAccelerator/Test/Private/UbaTestStorage.cpp b/Engine/Source/Programs/UnrealBuildAccelerator/Test/Private/UbaTestStorage.cpp
--- a/Engine/Source/Programs/UnrealBuildAccelerator/Test/Private/UbaTestStorage.cpp
+++ b/Engine/Source/Programs/UnrealBuildAccelerator/Test/Private/UbaTestStorage.cpp
@@ -4,6 +4,7 @@
 #include "UbaNetworkBackendTcp.h"
 #include "UbaStorageClient.h"
 #include "UbaStorageServer.h"
+#include "UbaTest.h"
 
 namespace uba
 {
@@ -308,10 +309,10 @@ namespace uba
 			return logger.Error(TC("Failed to connect"));
 		auto cg = MakeGuard([&]() { client.Disconnect(); });
 
-		StringBuffer<> fileName;
+		StringBuffer<> sourceFileName;
 		{
-			fileName.Append(rootDir).Append(TCV("UbaTestFile"));
-			FileAccessor fileHandle(logger, fileName.data);
+			sourceFileName.Append(rootDir).Append(TCV("UbaTestFile"));
+			FileAccessor fileHandle(logger, sourceFileName.data);
 			if (!fileHandle.CreateWrite())
 				return logger.Error(TC("Failed to create file for write"));
 			u8 byte = 'H';
@@ -321,20 +322,23 @@ namespace uba
 				return false;
 		}
 
+		StringBuffer<> destFileName;
 		CasKey key;
 		auto store = [&]()
 			{
 				bool storeCompressed = true;
-				if (!storageClient.StoreCasFileClient(key, ToStringKeyLower(fileName), fileName.data, SharedMemoryHandle{}, 0, 0, TC("UbaTestFile"), false, storeCompressed))
-					return logger.Error(TC("Failed to store file %s"), fileName.data);
+				if (!storageClient.StoreCasFileClient(key, ToStringKeyLower(sourceFileName), sourceFileName.data, SharedMemoryHandle{}, 0, 0, TC("UbaTestFile"), false, storeCompressed))
+					return logger.Error(TC("Failed to store file %s"), sourceFileName.data);
 				return true;
 			};
 
 		auto copy = [&]()
 			{
-				fileName.Clear().Append(testRootDir).Append(TCV("Uba")).Append(PathSeparator).Append(TCV("UbaTestFile"));
-				if (!storageServer.CopyOrLink(key, fileName.data, DefaultAttributes(), false, 1, {}, false, true))
-					return logger.Error(TC("Failed to copy cas to file %s"), fileName.data);
+				destFileName.Clear().Append(testRootDir).Append(TCV("Uba")).Append(PathSeparator).Append(TCV("UbaTestFile"));
+				if (!storageServer.CopyOrLink(key, destFileName.data, DefaultAttributes(), false, 1, {}, false, true))
+					return logger.Error(TC("Failed to copy cas to file %s"), destFileName.data);
+				if (!FilesEqual(logger, sourceFileName.data, destFileName.data))
+					return false;
 				return true;
 			};
 
@@ -356,8 +360,7 @@ namespace uba
 		if (!copy())
 			return false;
 
-		fileName.Clear().Append(testRootDir).Append(TCV("Uba")).Append(PathSeparator).Append(TCV("UbaTestFile"));
-		if (!storageServer.StoreCasFile(key, fileName.data, {}, true))
+		if (!storageServer.StoreCasFile(key, sourceFileName.data, {}, true))
 			return false;
 
 		//client.Fetch
@@ -366,6 +369,31 @@ namespace uba
 		//if (!copy())
 		//	return false;
 
+		// Large file
+		{
+			sourceFileName.Clear().Append(rootDir).Append(TCV("UbaTestFile2"));
+			FileAccessor fileHandle(logger, sourceFileName.data);
+			if (!fileHandle.CreateWrite())
+				return logger.Error(TC("Failed to create file for write"));
+			
+			// 256*64*1024 = 16mb file
+			u32 data[16*1024];
+			for (u32 i=0;i!=256; ++i)
+			{
+				for (u32 j=0;j!=sizeof_array(data); ++j)
+					data[j] = u32(rand());
+				if (!fileHandle.Write(data, sizeof(data)))
+					return false;
+			}
+			if (!fileHandle.Close())
+				return false;
+		}
+
+		if (!store())
+			return false;
+		if (!copy())
+			return false;
+
 		return true;
 	}
 
