[UBA]
* Removed assert around size alignment in SharedMemoryView... it is allowed to have any size as long as offset is page aligned.

diff --git a/Engine/Source/Programs/UnrealBuildAccelerator/Core/Private/UbaSharedMemoryView.cpp b/Engine/Source/Programs/UnrealBuildAccelerator/Core/Private/UbaSharedMemoryView.cpp
--- a/Engine/Source/Programs/UnrealBuildAccelerator/Core/Private/UbaSharedMemoryView.cpp
+++ b/Engine/Source/Programs/UnrealBuildAccelerator/Core/Private/UbaSharedMemoryView.cpp
@@ -168,7 +168,6 @@ namespace uba
 		UBA_ASSERTF(m_memory, TC("Init failed?"));
 		UBA_ASSERTF(m_capacity != ~0ull, TC("Can't extend memory on a view directly mapped on to allocator mapping (%s)"), hint);
 		UBA_ASSERTF(AlignUp(offset, PageSize) == offset, TC("ExtendMemory got offset %llu parameter which is not aligned (%s)"), offset, hint);
-		UBA_ASSERTF(AlignUp(size, PageSize) == size, TC("ExtendMemory got size %llu which is not aligned (%s)"), size, hint);
 
 		u8* targetAddress = m_memory + m_mappedSize;
 
