[UBA]
* Added simple support for "overlay files" which can be used to make virtual files only visible to individual processes (they never end up in the directory table)

diff --git a/Engine/Source/Programs/UnrealBuildAccelerator/Common/Private/UbaProcess.cpp b/Engine/Source/Programs/UnrealBuildAccelerator/Common/Private/UbaProcess.cpp
--- a/Engine/Source/Programs/UnrealBuildAccelerator/Common/Private/UbaProcess.cpp
+++ b/Engine/Source/Programs/UnrealBuildAccelerator/Common/Private/UbaProcess.cpp
@@ -252,6 +252,29 @@ namespace uba
 
 		if (m_parentProcess)
 			m_waitForParent.Create(EventResetType_Manual);
+		
+		if (!m_startInfo.overlayFilesStr.empty())
+		{
+			UBA_ASSERT(!m_parentProcess);
+
+			auto& sharedMemory = m_session.m_sharedMemory;
+			SharedMemoryHandle memHandle = sharedMemory.CreateHandle(m_session.m_logger, TC("OverlayFiles"));
+			sharedMemory.ExtendMemory(memHandle, 64*1024);
+			u8* mem = sharedMemory.MapView(memHandle, TC("OverlayFiles"), SharedMemoryMapType_ReadWrite);
+			
+			u32 overlaySize = 0;
+			bool res = m_session.InitOverlayFiles(mem, 64*1024, overlaySize, m_startInfo.overlayFilesStr);
+			sharedMemory.UnmapView(memHandle, mem);
+
+			if (!res)
+			{
+				sharedMemory.CloseHandle(m_session.m_logger, memHandle, TC("OverlayFiles"));
+				return false;
+			}
+
+			m_shared.overlayFilesHandle  = memHandle;
+			m_shared.overlayFilesSize = overlaySize;
+		}
 
 		UBA_ASSERT(m_startInfo.rules);
 
@@ -725,6 +748,9 @@ namespace uba
 			for (auto& memoryHandle : m_shared.sharedMemoryHandles)
 				m_session.m_sharedMemory.CloseHandle(m_session.m_logger, memoryHandle, TC(""));
 			m_shared.sharedMemoryHandles.clear();
+			if (m_shared.overlayFilesHandle.IsValid())
+				m_session.m_sharedMemory.CloseHandle(m_session.m_logger, m_shared.overlayFilesHandle, TC(""));
+			m_shared.overlayFilesHandle = {};
 		}
 
 		UBA_ASSERT(m_refCount);
@@ -915,6 +941,8 @@ namespace uba
 		writer.WriteU64(response.mappedFileTableHandle);
 		writer.WriteU32(response.mappedFileTableSize);
 		writer.WriteU32(response.mappedFileTableCount);
+		writer.WriteU64(m_shared.overlayFilesHandle.ToU64());
+		writer.WriteU32(m_shared.overlayFilesSize);
 
 		if (!m_startInfo.rootsHandle)
 		{
@@ -1500,10 +1528,23 @@ namespace uba
 
 		m_startTime = GetTime();
 
+		// Rewrite overlay.. and let reused process recreate dir table
+		if (!nextProcess.overlayFiles.empty())
+		{
+			UBA_ASSERT(m_shared.overlayFilesHandle.IsValid());
+			auto& sharedMemory = m_session.m_sharedMemory;
+			u8* mem = sharedMemory.MapView(m_shared.overlayFilesHandle, TC("OverlayFiles"), SharedMemoryMapType_ReadWrite);
+			memcpy(mem, nextProcess.overlayFiles.data(), nextProcess.overlayFiles.size());
+			sharedMemory.UnmapView(m_shared.overlayFilesHandle, mem);
+			m_shared.overlayFilesSize = u32(nextProcess.overlayFiles.size());
+		}
+
 		writer.WriteString(nextProcess.arguments);
 		writer.WriteString(nextProcess.workingDir);
 		writer.WriteString(nextProcess.description);
 		writer.WriteString(nextProcess.logFile);
+		writer.WriteU32(m_session.GetDirectoryTableSize());
+		writer.WriteU32(m_shared.overlayFilesSize);
 		return true;
 	}
 
diff --git a/Engine/Source/Programs/UnrealBuildAccelerator/Common/Private/UbaProcessStartInfoHolder.cpp b/Engine/Source/Programs/UnrealBuildAccelerator/Common/Private/UbaProcessStartInfoHolder.cpp
--- a/Engine/Source/Programs/UnrealBuildAccelerator/Common/Private/UbaProcessStartInfoHolder.cpp
+++ b/Engine/Source/Programs/UnrealBuildAccelerator/Common/Private/UbaProcessStartInfoHolder.cpp
@@ -156,6 +156,8 @@ namespace uba
 			logFile = logFileStr.c_str();
 		if (table.GetValueAsString(breadcrumbsStr, TC("Breadcrumbs")))
 			breadcrumbs = breadcrumbsStr.c_str();
+		if (table.GetValueAsString(overlayFilesStr, TC("OverlayFiles")))
+			overlayFiles = overlayFilesStr.c_str();
 
 		table.GetValueAsU32(priorityClass, TC("PriorityClass"));
 		table.GetValueAsBool(trackInputs, TC("TrackInputs"));
diff --git a/Engine/Source/Programs/UnrealBuildAccelerator/Common/Private/UbaScheduler.cpp b/Engine/Source/Programs/UnrealBuildAccelerator/Common/Private/UbaScheduler.cpp
--- a/Engine/Source/Programs/UnrealBuildAccelerator/Common/Private/UbaScheduler.cpp
+++ b/Engine/Source/Programs/UnrealBuildAccelerator/Common/Private/UbaScheduler.cpp
@@ -1023,6 +1023,14 @@ namespace uba
 			outNextProcess.logFile = si.logFile;
 			outNextProcess.breadcrumbs = si.breadcrumbs;
 
+			if (*si.overlayFiles)
+			{
+				outNextProcess.overlayFiles.resize(16*1024);
+				u32 overlayFilesSize = 0;
+				m_session.InitOverlayFiles(outNextProcess.overlayFiles.data(), u32(outNextProcess.overlayFiles.size()), overlayFilesSize, si.overlayFilesStr);
+				outNextProcess.overlayFiles.resize(overlayFilesSize);
+			}
+
 			#if UBA_DEBUG
 			auto PrepPath = [this](StringBufferBase& out, const ProcessStartInfo& psi)
 				{
diff --git a/Engine/Source/Programs/UnrealBuildAccelerator/Common/Private/UbaSession.cpp b/Engine/Source/Programs/UnrealBuildAccelerator/Common/Private/UbaSession.cpp
--- a/Engine/Source/Programs/UnrealBuildAccelerator/Common/Private/UbaSession.cpp
+++ b/Engine/Source/Programs/UnrealBuildAccelerator/Common/Private/UbaSession.cpp
@@ -509,6 +509,8 @@
 
 	bool Session::WriteDirectoryEntriesInternal(DirectoryTable::Directory& dir, const StringKey& dirKey, StringView dirPath, bool isRefresh, u32& outTableOffset)
 	{
+		UBA_ASSERT(dir.parseOffset != 3);
+
 		if (dir.tableOffset != InvalidTableOffset && !isRefresh)
 		{
 			isRefresh = true;
@@ -1171,6 +1173,59 @@
 		return true;
 	}
 
+	StringKey GetKeyAndFixedName(StringBuffer<>& outFixedFilePath, StringKeyHasher& outHasher, const tchar* filePath, StringKey* outDirKey = nullptr, const tchar** outBaseFileName = nullptr)
+	{
+		StringBuffer<> workingDir;
+		if (!IsAbsolutePath(filePath))
+		{
+			GetCurrentDirectoryW(workingDir);
+			workingDir.EnsureEndsWithSlash();
+		}
+		FixPath(filePath, workingDir.data, workingDir.count, outFixedFilePath);
+
+		StringKey dirKey;
+		StringBuffer<> dirNameForHash;
+		const tchar* baseFileName = nullptr;
+		GetDirKey(dirKey, dirNameForHash, baseFileName, outFixedFilePath);
+		if (outDirKey)
+			*outDirKey = dirKey;
+		if (outBaseFileName)
+			*outBaseFileName = baseFileName;
+
+		if (CaseInsensitiveFs)
+			dirNameForHash.MakeLower();
+
+		outHasher.Update(dirNameForHash);
+
+		if (baseFileName)
+		{
+			StringBuffer<256> baseFileNameForHash;
+			baseFileNameForHash.Append(baseFileName);
+			if (CaseInsensitiveFs)
+				baseFileNameForHash.MakeLower();
+
+			outHasher.Update(baseFileNameForHash);
+		}
+
+		StringKey result = ToStringKey(outHasher);
+
+		#if UBA_DEBUG
+		StringBuffer<> testPath(outFixedFilePath);
+		if (CaseInsensitiveFs)
+			testPath.MakeLower();
+		StringKey testKey = ToStringKey(testPath);
+		UBA_ASSERTF(testKey == result, TC("Key mismatch for %s"), outFixedFilePath.data);
+		#endif
+
+		return result;
+	}
+
+	StringKey GetKeyAndFixedName(StringBuffer<>& outFixedFilePath, const tchar* filePath, StringKey* outDirKey = nullptr, const tchar** outBaseFileName = nullptr)
+	{
+		StringKeyHasher hasher;
+		return GetKeyAndFixedName(outFixedFilePath, hasher, filePath, outDirKey, outBaseFileName);
+	}
+
 	bool Session::RegisterCreateFileForWrite(StringKey fileNameKey, const StringView& fileName, bool registerRealFile, u64 fileSize, u64 lastWriteTime, bool invalidateStorage)
 	{
 		// Remote is not updating its own directory table
@@ -1235,7 +1290,7 @@
 			UBA_ASSERT(res); (void)res;
 			dirLock.Enter();
 		}
-		UBA_ASSERTF(dir.parseOffset == 1 || !m_shouldWriteToDisk, TC("Registering create file for write %s with unexpect dir.parseOffset %u "), fileName.data, dir.parseOffset);
+		UBA_ASSERTF(dir.parseOffset == 1 || (!m_shouldWriteToDisk && dir.parseOffset != 3), TC("Registering create file for write %s with unexpect dir.parseOffset %u "), fileName.data, dir.parseOffset);
 
 		if (fileNameKey == StringKeyZero)
 		{
@@ -1283,9 +1338,9 @@
 			BinaryReader reader(m_directoryTableMem + insres.first->second);
 			u32 oldAttr = reader.ReadFileAttributes();
 
-			if (isDirectory) // Ignore updating directories.. they should always be the same regardless
+			if (IsDirectory(oldAttr))  // Ignore updating directories.. they should always be the same regardless
 			{
-				UBA_ASSERT(IsDirectory(oldAttr));
+				UBA_ASSERT(isDirectory || !shouldWriteToDisk);
 				return true;
 			}
 			reader.ReadVolumeSerial();
@@ -1430,6 +1485,106 @@
 		return dirTable.m_memorySize;
 	}
 
+	bool Session::InitOverlayFiles(u8* outMem, u32 capacity, u32& outSize, StringView overlayFiles)
+	{
+		auto& dirTable = m_directoryTable;
+
+		BinaryWriter memWriter(outMem, 0, capacity);
+
+		UnorderedMap<StringKey, u32> dirOffsets;
+
+		tchar* it = const_cast<tchar*>(overlayFiles.data);
+		tchar* end = it + overlayFiles.count;
+		while (*it)
+		{
+			tchar* fileEnd = TStrchr(it, ';');
+			if (fileEnd)
+				*fileEnd = 0;
+			else
+				fileEnd = end;
+
+			StringKey dirKey;
+			const tchar* lastSlash = nullptr;
+			StringBuffer<> fixedFileName;
+			StringKey fileNameKey = GetKeyAndFixedName(fixedFileName, it, &dirKey, &lastSlash);
+
+			u64 fileSize = 0;
+			u64 lastWriteTime = 0;
+			{
+				SCOPED_FUTEX(m_fileMappingTableLookupLock, mappingLookupLock);
+				auto findIt = m_fileMappingTableLookup.find(fileNameKey);
+				if (findIt == m_fileMappingTableLookup.end())
+					return false;
+				FileMappingEntry& entry = findIt->second;
+				mappingLookupLock.Leave();
+				fileSize = entry.contentSize;
+				lastWriteTime = entry.lastWriteTime;
+			}
+
+			auto dirOffsetIns = dirOffsets.try_emplace(dirKey);
+			if (dirOffsetIns.second)
+			{
+				StringView dirName(fixedFileName.data, u32(lastSlash - fixedFileName.data));
+
+				SCOPED_READ_LOCK(dirTable.m_lookupLock, lookupLock);
+				auto findIt = dirTable.m_lookup.find(dirKey);
+				if (findIt == dirTable.m_lookup.end())
+				{
+					lookupLock.Leave();
+					u32 res = WriteDirectoryEntries(dirKey, dirName);
+					UBA_ASSERTF(res, TC("Failed to write directory entries for %.*s"), dirName.count, dirName.data); (void)res;
+					lookupLock.Enter();
+					findIt = dirTable.m_lookup.find(dirKey);
+					if (findIt == dirTable.m_lookup.end())
+						return m_logger.Error(TC("OverlayFiles can only exist in existing directories (%.*s)"), fixedFileName.count, fixedFileName.data);
+				}
+				DirectoryTable::Directory& dir = findIt->second;
+				lookupLock.Leave();
+
+				SCOPED_WRITE_LOCK(dir.lock, dirLock);
+				if (dir.parseOffset != 3)
+				{
+					while (dir.parseOffset == 0)
+					{
+						dirLock.Leave();
+						Sleep(1);
+						dirLock.Enter();
+					}
+					dir.parseOffset = 3; // 3 means it is locked, it is now a directory owned by overlays
+
+					#if UBA_DEBUG_TRACK_DIR
+					m_debugLogger->Info(TC("TRACKOVERL %.*s"), dirName.count, dirName.data);
+					#endif
+				}
+				dirOffsetIns.first->second = dir.tableOffset;
+			}
+
+
+			u8 temp[1024];
+			BinaryWriter writer(temp, 0, sizeof(temp));
+			writer.WriteStringKey(dirKey);
+			writer.Write7BitEncoded(dirOffsetIns.first->second); // Previous entry for same directory
+			writer.Write7BitEncoded(1); // Count
+			writer.WriteString(lastSlash + 1);
+			writer.WriteFileAttributes(DefaultAttributes());
+			writer.WriteVolumeSerial(0);
+			writer.WriteFileIndex(++m_fileIndexCounter);
+			writer.WriteFileTime(lastWriteTime);
+			writer.WriteFileSize(fileSize);
+			u64 written = writer.GetPosition();
+			memWriter.Write7BitEncoded(written); // Storage size
+			dirOffsetIns.first->second = u32(memWriter.GetPosition()) + sizeof(StringKey) + OverlayTableOffset;
+			memWriter.WriteBytes(writer.GetData(), written);
+			if (fileEnd == end)
+				break;
+			*fileEnd = ';';
+			it = ++fileEnd;
+		}
+
+		outSize = u32(memWriter.GetPosition());
+		return true;
+	}
+
 	RootsHandle Session::RegisterRoots(const void* rootsData, uba::u64 rootsDataSize)
 	{
 		CasKey key = ToCasKey(CasKeyHasher().Update(rootsData, rootsDataSize), false);
@@ -2006,55 +2161,6 @@
 		return h;
 	}
 
-	StringKey GetKeyAndFixedName(StringBuffer<>& outFixedFilePath, StringKeyHasher& outHasher, const tchar* filePath)
-	{
-		StringBuffer<> workingDir;
-		if (!IsAbsolutePath(filePath))
-		{
-			GetCurrentDirectoryW(workingDir);
-			workingDir.EnsureEndsWithSlash();
-		}
-		FixPath(filePath, workingDir.data, workingDir.count, outFixedFilePath);
-
-		StringKey dirKey;
-		StringBuffer<> dirNameForHash;
-		const tchar* baseFileName = nullptr;
-		GetDirKey(dirKey, dirNameForHash, baseFileName, outFixedFilePath);
-
-		if (CaseInsensitiveFs)
-			dirNameForHash.MakeLower();
-
-		outHasher.Update(dirNameForHash);
-
-		if (baseFileName)
-		{
-			StringBuffer<256> baseFileNameForHash;
-			baseFileNameForHash.Append(baseFileName);
-			if (CaseInsensitiveFs)
-				baseFileNameForHash.MakeLower();
-
-			outHasher.Update(baseFileNameForHash);
-		}
-
-		StringKey result = ToStringKey(outHasher);
-
-#if UBA_DEBUG
-		StringBuffer<> testPath(outFixedFilePath);
-		if (CaseInsensitiveFs)
-			testPath.MakeLower();
-		StringKey testKey = ToStringKey(testPath);
-		UBA_ASSERTF(testKey == result, TC("Key mismatch for %s"), outFixedFilePath.data);
-#endif
-
-		return result;
-	}
-
-	StringKey GetKeyAndFixedName(StringBuffer<>& outFixedFilePath, const tchar* filePath)
-	{
-		StringKeyHasher hasher;
-		return GetKeyAndFixedName(outFixedFilePath, hasher, filePath);
-	}
-
 	bool Session::RefreshDirectory(const tchar* dirName, bool forceRegister)
 	{
 		UBA_ASSERT(!m_runningRemote);
@@ -2123,12 +2229,13 @@
 		return RegisterCreateFileForWrite(fileNameKey, fixedFilePath, false, memorySize, 0, false);
 	}
 
-	bool Session::DeleteVirtualFile(const tchar* filePath)
+	bool Session::DeleteVirtualFile(const tchar* filePath, bool registerToDir)
 	{
 		UBA_ASSERT(!m_runningRemote);
 		StringBuffer<> fixedFilePath;
 		auto key = GetKeyAndFixedName(fixedFilePath, filePath);
-		RegisterDeleteFile(key, fixedFilePath);
+		if (registerToDir)
+			RegisterDeleteFile(key, fixedFilePath);
 
 		SCOPED_FUTEX(m_virtualSourceFilesLock, virtualSourceFilesLock);
 		auto findIt = m_virtualSourceFiles.find(key); 
@@ -2137,6 +2244,16 @@
 		VirtualSourceFile vsf = findIt->second;
 		m_virtualSourceFiles.erase(findIt);
 		virtualSourceFilesLock.Leave();
+
+		SCOPED_FUTEX(m_fileMappingTableLookupLock, lookupLock);
+		auto findIt2 = m_fileMappingTableLookup.find(key);
+		if (findIt2 != m_fileMappingTableLookup.end())
+		{
+			UBA_ASSERTF(findIt2->second.refCount == 0, TC("Deleting virtual file %s which is in use"), filePath);
+			m_fileMappingTableLookup.erase(findIt2);
+		}
+		lookupLock.Leave();
+
 		m_sharedMemory.CloseHandle(m_logger, vsf.memoryHandle, TC("VirtualFile"));
 		return true;
 	}
diff --git a/Engine/Source/Programs/UnrealBuildAccelerator/Common/Private/UbaSessionClient.cpp b/Engine/Source/Programs/UnrealBuildAccelerator/Common/Private/UbaSessionClient.cpp
--- a/Engine/Source/Programs/UnrealBuildAccelerator/Common/Private/UbaSessionClient.cpp
+++ b/Engine/Source/Programs/UnrealBuildAccelerator/Common/Private/UbaSessionClient.cpp
@@ -1485,6 +1485,7 @@ namespace uba
 	struct SessionClient::InternalProcessStartInfo : ProcessStartInfoHolder
 	{
 		u32 processId = 0;
+		Vector<u8> overlayFilesData;
 	};
 
 
@@ -1526,6 +1527,9 @@ namespace uba
 			InternalProcessStartInfo& info = out.back();
 			info.processId = processId;
 			info.Read(reader);
+
+			info.overlayFilesData.resize(reader.Read7BitEncoded());
+			reader.ReadBytes(info.overlayFilesData.data(), info.overlayFilesData.size());
 		}
 
 		u32 neededDirectoryTableSize = reader.ReadU32();
@@ -2185,6 +2189,18 @@ namespace uba
 							session.m_processFinished(&process);
 					};
 
+					if (!startInfo.overlayFilesData.empty())
+					{
+						UBA_ASSERT(startInfo.overlayFilesStr.empty());
+						SharedMemoryHandle memHandle = m_sharedMemory.CreateHandle(m_logger, TC("OverlayFiles"));
+						m_sharedMemory.ExtendMemory(memHandle, startInfo.overlayFilesData.size());
+						u8* mem = m_sharedMemory.MapView(memHandle, TC("OverlayFiles"), SharedMemoryMapType_ReadWrite);
+						memcpy(mem, startInfo.overlayFilesData.data(), startInfo.overlayFilesData.size());
+						m_sharedMemory.UnmapView(memHandle, mem);
+						process->m_shared.overlayFilesHandle = memHandle;
+						process->m_shared.overlayFilesSize = u32(startInfo.overlayFilesData.size());
+					}
+
 					if (!process->Start(startInfo, true, env, true)) // If false, exitFunction is not called
 					{
 						SendReturnProcess(rec->handle.GetId(), TC("Failed to find executable"));
@@ -2304,7 +2320,7 @@ namespace uba
 		lock2.Leave();
 		if (newMemPosition == m_dirtableParsedPosition)
 			return false;
-		m_directoryTable.ParseDirectoryTableNoLock(m_dirtableParsedPosition, newMemPosition);
+		m_directoryTable.ParseDirectoryTableNoLock(m_directoryTable.m_memory, m_dirtableParsedPosition, newMemPosition);
 		m_dirtableParsedPosition = newMemPosition;
 		return true;
 	}
@@ -2581,6 +2597,8 @@ namespace uba
 			outNextProcess.workingDir = reader.ReadString();
 			outNextProcess.description = reader.ReadString();
 			outNextProcess.logFile = reader.ReadString();
+			outNextProcess.overlayFiles.resize(reader.Read7BitEncoded());
+			reader.ReadBytes(outNextProcess.overlayFiles.data(), outNextProcess.overlayFiles.size());
 
 			if (m_logToFile || (!outNextProcess.logFile.empty() && m_shouldSendLogToServer))
 			{
diff --git a/Engine/Source/Programs/UnrealBuildAccelerator/Common/Private/UbaSessionServer.cpp b/Engine/Source/Programs/UnrealBuildAccelerator/Common/Private/UbaSessionServer.cpp
--- a/Engine/Source/Programs/UnrealBuildAccelerator/Common/Private/UbaSessionServer.cpp
+++ b/Engine/Source/Programs/UnrealBuildAccelerator/Common/Private/UbaSessionServer.cpp
@@ -1503,6 +1503,18 @@ namespace uba
 			writer.WriteU32(process->m_processId);
 			startInfo.Write(writer, applicationOverride);
 
+			if (!startInfo.overlayFilesStr.empty())
+			{
+				Vector<u8> data;
+				data.resize(16*1024);
+				u32 overlayFilesSize = 0;
+				InitOverlayFiles(data.data(), u32(data.size()), overlayFilesSize, startInfo.overlayFilesStr);
+				writer.Write7BitEncoded(overlayFilesSize);
+				writer.WriteBytes(data.data(), overlayFilesSize);
+			}
+			else
+				writer.Write7BitEncoded(0);
+
 			if (process->m_knownInputsDone.IsCreated())
 				processesWithKnownInputsToSend.push_back(process);
 
@@ -1841,6 +1853,8 @@ namespace uba
 			writer.WriteString(nextProcess.workingDir);
 			writer.WriteString(nextProcess.description);
 			writer.WriteString(nextProcess.logFile);
+			writer.Write7BitEncoded(nextProcess.overlayFiles.size());
+			writer.WriteBytes(nextProcess.overlayFiles.data(), nextProcess.overlayFiles.size());
 		}
 		return true;
 	}
@@ -2377,7 +2391,13 @@ namespace uba
 		m_sharedMemory.TraceStats(m_trace, 30);
 		#endif
 
-#if 0
+		#if 0
+		StringBuffer<> str;
+		str.Appendf(TC("DirTable: %s, MapTable: %s"), BytesToText(m_directoryTable.m_memorySize).str, BytesToText(m_fileMappingTableSize).str);
+		m_trace.StatusUpdate(40, 7, str, LogEntryType_Info);
+		#endif
+
+		#if 0
 		m_fileMappingBuffer.TraceStatus(m_trace, 20, TC("Session"), 1);
 
 		m_trace.StatusUpdate(40, 1, TCV("Independents"), LogEntryType_Info);
@@ -2388,7 +2408,7 @@ namespace uba
 		m_trace.StatusUpdate(40, 7, str, LogEntryType_Info);
 
 		m_storage.TraceUpdate(m_trace, 300);
-#endif
+		#endif
 	}
 
 	void SessionServer::WriteRemoteEnvironmentVariables(BinaryWriter& writer)
diff --git a/Engine/Source/Programs/UnrealBuildAccelerator/Common/Public/UbaProcess.h b/Engine/Source/Programs/UnrealBuildAccelerator/Common/Public/UbaProcess.h
--- a/Engine/Source/Programs/UnrealBuildAccelerator/Common/Public/UbaProcess.h
+++ b/Engine/Source/Programs/UnrealBuildAccelerator/Common/Public/UbaProcess.h
@@ -149,6 +149,8 @@ namespace uba
 			UnorderedMap<StringKey, WrittenFile> writtenFiles;
 			OverflowedWrittenFilesMessage* firstOverflowMessage = nullptr;
 			UnorderedSet<SharedMemoryHandle> sharedMemoryHandles; // Track created memory handles. entries can be transferred to WrittenFiles
+			SharedMemoryHandle overlayFilesHandle;
+			u32 overlayFilesSize = 0;
 		};
 
 		u32 m_nativeProcessId = 0;
diff --git a/Engine/Source/Programs/UnrealBuildAccelerator/Common/Public/UbaProcessStartInfo.h b/Engine/Source/Programs/UnrealBuildAccelerator/Common/Public/UbaProcessStartInfo.h
--- a/Engine/Source/Programs/UnrealBuildAccelerator/Common/Public/UbaProcessStartInfo.h
+++ b/Engine/Source/Programs/UnrealBuildAccelerator/Common/Public/UbaProcessStartInfo.h
@@ -45,6 +45,8 @@ namespace uba
 
 		const ApplicationRules* rules = nullptr;// Internal use (for now)
 		int uiLanguage = 1033;					// Internal use
+
+		const tchar* overlayFiles = TC("");   // Semi colon separated string containing paths to overlay files (files that are only visible to this process)
 	};
 
 	enum ProcessExitedResponse : u8
diff --git a/Engine/Source/Programs/UnrealBuildAccelerator/Common/Public/UbaProcessStartInfoHolder.h b/Engine/Source/Programs/UnrealBuildAccelerator/Common/Public/UbaProcessStartInfoHolder.h
--- a/Engine/Source/Programs/UnrealBuildAccelerator/Common/Public/UbaProcessStartInfoHolder.h
+++ b/Engine/Source/Programs/UnrealBuildAccelerator/Common/Public/UbaProcessStartInfoHolder.h
@@ -51,6 +51,9 @@ namespace uba
 
 			breadcrumbsStr = si.breadcrumbs;
 			breadcrumbs = breadcrumbsStr.c_str();
+
+			overlayFilesStr = si.overlayFiles;
+			overlayFiles = overlayFilesStr.c_str();
 		}
 
 		void Write(BinaryWriter& writer, const StringView& applicationOverride)
@@ -116,6 +119,8 @@ namespace uba
 			logFile = logFileStr.c_str();
 			breadcrumbsStr = o.breadcrumbsStr;
 			breadcrumbs = breadcrumbsStr.c_str();
+			overlayFilesStr = o.overlayFiles;
+			overlayFiles = overlayFilesStr.c_str();
 
 			weight = o.weight;
 		}
@@ -130,6 +135,7 @@ namespace uba
 		TString workingDirStr;
 		TString logFileStr;
 		TString breadcrumbsStr;
+		TString overlayFilesStr;
 		float weight = 1.0f;
 
 		bool InternalExpand(const tchar* firstArgBegin, const tchar* firstArgEnd);
diff --git a/Engine/Source/Programs/UnrealBuildAccelerator/Common/Public/UbaSession.h b/Engine/Source/Programs/UnrealBuildAccelerator/Common/Public/UbaSession.h
--- a/Engine/Source/Programs/UnrealBuildAccelerator/Common/Public/UbaSession.h
+++ b/Engine/Source/Programs/UnrealBuildAccelerator/Common/Public/UbaSession.h
@@ -69,7 +69,7 @@ namespace uba
 		bool RegisterNewFile(const tchar* filePath); // Tell uba a new file on disk has been added by some other system while session is running
 		bool RegisterVirtualFile(const tchar* filePath, const tchar* sourceFile, u64 sourceOffset, u64 sourceSize);
 		bool CreateVirtualFile(const tchar* filePath, const void* memory, u64 memorySize, bool transient, bool registerToDir = true);
-		bool DeleteVirtualFile(const tchar* filePath);
+		bool DeleteVirtualFile(const tchar* filePath, bool registerToDir);
 		void RegisterDeleteFile(const tchar* filePath); // Tell uba a file on disk has been deleted by some other system while session is running
 		bool RegisterNewDirectory(const tchar* directoryPath); // Tell uba a directory on disk has been added by some other system while session is running
 
@@ -111,6 +111,8 @@ namespace uba
 		bool DevirtualizeString(TString& inOut, RootsHandle rootsHandle, bool allowPathsWithoutRoot, const tchar* hint = TC(""));
 		bool PopulateLocalToIndexRoots(RootPaths& out, RootsHandle rootsHandle);
 
+		bool InitOverlayFiles(u8* outMem, u32 capacity, u32& outSize, StringView overlayFiles);
+
 		// For testing
 		const u8* GetDirectoryTableMemory() { return m_directoryTableMem; }
 		u32 GetDirectoryTableSize();
@@ -615,6 +617,7 @@ namespace uba
 		TString description;
 		TString logFile;
 		TString breadcrumbs;
+		Vector<u8> overlayFiles;
 	};
 }
 
diff --git a/Engine/Source/Programs/UnrealBuildAccelerator/Core/Private/UbaApplicationRules.cpp b/Engine/Source/Programs/UnrealBuildAccelerator/Core/Private/UbaApplicationRules.cpp
--- a/Engine/Source/Programs/UnrealBuildAccelerator/Core/Private/UbaApplicationRules.cpp
+++ b/Engine/Source/Programs/UnrealBuildAccelerator/Core/Private/UbaApplicationRules.cpp
@@ -697,7 +697,7 @@ namespace uba
 
 		virtual bool IsInvisible(const StringView& fileName) const override
 		{
-			return fileName.EndsWith(TCV(".out"));
+			return fileName.EndsWith(TCV(".tmp")) || fileName.EndsWith(TCV(".out"));
 		}
 	};
 
diff --git a/Engine/Source/Programs/UnrealBuildAccelerator/Core/Public/UbaDirectoryTable.h b/Engine/Source/Programs/UnrealBuildAccelerator/Core/Public/UbaDirectoryTable.h
--- a/Engine/Source/Programs/UnrealBuildAccelerator/Core/Public/UbaDirectoryTable.h
+++ b/Engine/Source/Programs/UnrealBuildAccelerator/Core/Public/UbaDirectoryTable.h
@@ -8,6 +8,7 @@
 namespace uba
 {
 	enum : u64 { InvalidTableOffset = 0 };
+	enum : u64 { OverlayTableOffset = 256*1024*1024 };
 
 	class DirectoryTable
 	{
@@ -23,6 +24,55 @@ namespace uba
 			ParseDirectoryTable(tableSize);
 		}
 
+		void InitOverlay(const u8* mem, u32 tableSize)
+		{
+			m_overlayMemory = mem;
+			ParseDirectoryTableNoLock(mem, 0, tableSize, OverlayTableOffset);
+		}
+
+		void DeinitOverlay(const u8* oldMem, u32 size)
+		{
+			if (!size)
+				return;
+
+			m_overlayMemory = oldMem;
+			for (auto i=m_lookup.begin(), e=m_lookup.end(); i!=e;)
+			{
+				auto& dir = i->second;
+				if (dir.tableOffset < OverlayTableOffset)
+				{
+					++i;
+					continue;
+				}
+
+				while (true)
+				{
+					BinaryReader reader(GetReader(dir.tableOffset));
+					dir.tableOffset = u32(reader.Read7BitEncoded());
+					if (dir.tableOffset < OverlayTableOffset)
+						break;
+				}
+				if (dir.tableOffset == InvalidTableOffset)
+				{
+					i = m_lookup.erase(i);
+					continue;
+				}
+
+				dir.parseOffset = dir.tableOffset;
+
+				for (auto j=dir.files.begin(); j!=dir.files.end();)
+				{
+					if (j->second < OverlayTableOffset)
+						++j;
+					else
+						j = dir.files.erase(j);
+				}
+
+				++i;
+			}
+			m_overlayMemory = nullptr;
+		}
+
 		void ParseDirectoryTable(u32 size)
 		{
 			SCOPED_WRITE_LOCK(m_lookupLock, lock);
@@ -33,13 +83,13 @@ namespace uba
 		{
 			if (size <= m_memorySize)
 				return;
-			ParseDirectoryTableNoLock(m_memorySize, size);
+			ParseDirectoryTableNoLock(m_memory, m_memorySize, size);
 			m_memorySize = size;
 		}
 
-		void ParseDirectoryTableNoLock(u32 from, u32 to)
+		void ParseDirectoryTableNoLock(const u8* memory, u32 from, u32 to, u32 additionalOffset = 0)
 		{
-			BinaryReader reader(m_memory, from, to);
+			BinaryReader reader(memory, from, to);
 			while (true)
 			{
 				u64 pos = reader.GetPosition();
@@ -49,7 +99,9 @@ namespace uba
 				u64 storageSize = reader.Read7BitEncoded();
 				StringKey dirKey = reader.ReadStringKey();
 				auto insres = m_lookup.try_emplace(dirKey, m_memoryBlock); // Note that this is allowed to overwrite
-				insres.first->second.tableOffset = u32(reader.GetPosition());
+				u32& tableOffset = insres.first->second.tableOffset;
+				UBA_ASSERTF(tableOffset < OverlayTableOffset || additionalOffset == OverlayTableOffset, TC("Can't add non-overlay directory items after overlay items"));
+				tableOffset = u32(reader.GetPosition()) + additionalOffset;
 				reader.Skip(storageSize - sizeof(dirKey));
 			}
 		}
@@ -68,9 +120,18 @@ namespace uba
 			dir.parseOffset = dir.tableOffset;
 		}
 
+		BinaryReader GetReader(u32 tableOffset)
+		{
+			if (tableOffset < OverlayTableOffset)
+				return BinaryReader(m_memory, tableOffset);
+			else
+				return { m_overlayMemory - OverlayTableOffset, tableOffset};
+		}
+
 		void PopulateDirectoryRecursive(const StringKeyHasher& hasher, u32 tableOffset, u32 parseOffset, EntryLookup& files)
 		{
-			BinaryReader reader(m_memory, tableOffset);
+			BinaryReader reader(GetReader(tableOffset));
+
 			u32 prevTableOffset = u32(reader.Read7BitEncoded());
 
 			u32 buffer[48*1024];
@@ -85,7 +146,7 @@ namespace uba
 					firstIsRoot = prevTableOffset == InvalidTableOffset;
 					break;
 				}
-				reader.SetPosition(prevTableOffset);
+				reader = GetReader(prevTableOffset);
 				prevTableOffset = u32(reader.Read7BitEncoded());
 				readerOffsets[count++] = u32(reader.GetPosition());
 				if (count == sizeof_array(buffer))
@@ -97,7 +158,7 @@ namespace uba
 
 			for (u32 i=count; i>0; --i)
 			{
-				reader.SetPosition(readerOffsets[i-1]);
+				reader = GetReader(readerOffsets[i-1]);
 				if (firstIsRoot)
 				{
 					firstIsRoot = false;
@@ -275,7 +336,7 @@ namespace uba
 			if (tableOffset & 0x80000000)
 			{
 				tableOffset = tableOffset & ~0x80000000;
-				BinaryReader reader(m_memory, tableOffset);
+				BinaryReader reader(GetReader(tableOffset));
 				u64 prevTableOffset = reader.Read7BitEncoded();
 				while (prevTableOffset != InvalidTableOffset)
 				{
@@ -294,7 +355,7 @@ namespace uba
 				return ~u32(0);
 			}
 
-			BinaryReader reader(m_memory, tableOffset);
+			BinaryReader reader(GetReader(tableOffset));
 			if (outFileName)
 				reader.ReadString(outFileName, fileNameCapacity);
 			else
@@ -358,7 +419,7 @@ namespace uba
 					if (fileIt != directory->files.end())
 					{
 						UBA_ASSERT(fileIt->second != ~0u);
-						BinaryReader reader(m_memory, fileIt->second);
+						BinaryReader reader(GetReader(fileIt->second));
 						StringBuffer<> fileName;
 						reader.ReadString(fileName);
 						out.Append(PathSeparator).Append(fileName);
@@ -433,6 +494,8 @@ namespace uba
 		ReaderWriterLock m_memoryLock;
 		const u8* m_memory = nullptr;
 		u32 m_memorySize = 0;
+
+		const u8* m_overlayMemory = nullptr;
 	};
 
 }
diff --git a/Engine/Source/Programs/UnrealBuildAccelerator/Detours/Private/Windows/UbaDetoursFunctionsWin.cpp b/Engine/Source/Programs/UnrealBuildAccelerator/Detours/Private/Windows/UbaDetoursFunctionsWin.cpp
--- a/Engine/Source/Programs/UnrealBuildAccelerator/Detours/Private/Windows/UbaDetoursFunctionsWin.cpp
+++ b/Engine/Source/Programs/UnrealBuildAccelerator/Detours/Private/Windows/UbaDetoursFunctionsWin.cpp
@@ -65,6 +65,9 @@ namespace uba
 	SharedMemoryHandle g_permanentFilesHandle;
 	u8* g_permanentFilesMemory;
 
+	u32 g_overlayFilesSize;
+
+
 bool g_useMiMalloc;
 constexpr u64  g_pageSize = 64*1024;
 
@@ -984,6 +987,8 @@ void Init(const DetoursPayload& payload, u64 startTime)
 	HANDLE mappedFileTableHandle;
 	u32 mappedFileTableSize;
 	u32 mappedFileTableCount;
+	SharedMemoryHandle overlayFilesHandle;
+	u32 overlayFilesSize;
 
 	{
 		RPC_MESSAGE(Init, init)
@@ -1005,6 +1010,9 @@ void Init(const DetoursPayload& payload, u64 startTime)
 		mappedFileTableSize = reader.ReadU32();
 		mappedFileTableCount = reader.ReadU32();
 
+		overlayFilesHandle = SharedMemoryHandle::FromU64(reader.ReadU64());
+		overlayFilesSize = reader.ReadU32();
+
 		if (u16 vfsSize = reader.ReadU16())
 		{
 			BinaryReader vfsReader(reader.GetPositionData(), 0, vfsSize);
@@ -1068,6 +1076,14 @@ void Init(const DetoursPayload& payload, u64 startTime)
 	{
 		TimerScope ts2(g_stats.dirTable);
 		g_directoryTable.Init(directoryTableMem, directoryTableCount, directoryTableSize);
+
+		if (overlayFilesSize)
+		{
+			SharedMemoryView overlayFilesView;
+			Rpc_GetSharedMemory(overlayFilesView, g_sharedMemoryAllocator, nullptr, overlayFilesHandle, 0, 0, TC("PermanentFiles"));
+			g_directoryTable.InitOverlay(overlayFilesView.DetachMemory(), overlayFilesSize);
+			g_overlayFilesSize = overlayFilesSize;
+		}
 	}
 
 	if (permanentFilesHandle.IsValid())
@@ -1179,8 +1195,18 @@ extern "C"
 		FlushDebugLog();
 		#endif
 
+		Vector<u8> oldOverlayFiles;
+		if (g_overlayFilesSize)
+		{
+			oldOverlayFiles.resize(g_overlayFilesSize);
+			memcpy(oldOverlayFiles.data(), g_directoryTable.m_overlayMemory, g_overlayFilesSize);
+
+		}
+
 		*outArguments = 0;
 		bool newProcess;
+		u32 overlayFilesSize = 0;
+		u32 directoryTableSize = 0;
 		{
 			RPC_MESSAGE(GetNextProcess, log)
 			writer.WriteU32(prevExitCode);
@@ -1196,6 +1222,8 @@ extern "C"
 				reader.SkipString(); // workingDir
 				reader.SkipString(); // description
 				reader.ReadString(g_logName.Clear());
+				directoryTableSize = reader.ReadU32();
+				overlayFilesSize = reader.ReadU32();
 			}
 		}
 
@@ -1213,6 +1241,15 @@ extern "C"
 			g_debugFile = (FileHandle)(u64)debugFile;
 			LogHeader(ToView(outArguments));
 			#endif
+
+			if (overlayFilesSize)
+			{
+				const u8* overlayMem = g_directoryTable.m_overlayMemory;
+				g_directoryTable.DeinitOverlay(oldOverlayFiles.data(), g_overlayFilesSize);
+				g_directoryTable.ParseDirectoryTableNoLock(directoryTableSize);
+				g_directoryTable.InitOverlay(overlayMem, overlayFilesSize);
+				g_overlayFilesSize = overlayFilesSize;
+			}
 		}
 
 		Rpc_UpdateTables();
diff --git a/Engine/Source/Programs/UnrealBuildAccelerator/Host/Private/UbaExports.cpp b/Engine/Source/Programs/UnrealBuildAccelerator/Host/Private/UbaExports.cpp
--- a/Engine/Source/Programs/UnrealBuildAccelerator/Host/Private/UbaExports.cpp
+++ b/Engine/Source/Programs/UnrealBuildAccelerator/Host/Private/UbaExports.cpp
@@ -548,13 +548,13 @@ extern "C"
 	{
 		return server->RegisterVirtualFile(filename, sourceFile, sourceOffset, sourceSize);
 	}
-	bool SessionServer_CreateVirtualFile(uba::SessionServer* server, const uba::tchar* filename, const void* memory, uba::u64 memorySize, bool transient)
+	bool SessionServer_CreateVirtualFile(uba::SessionServer* server, const uba::tchar* filename, const void* memory, uba::u64 memorySize, bool transient, bool registerToDir)
 	{
-		return server->CreateVirtualFile(filename, memory, memorySize, transient);
+		return server->CreateVirtualFile(filename, memory, memorySize, transient, registerToDir);
 	}
-	bool SessionServer_DeleteVirtualFile(uba::SessionServer* server, const uba::tchar* filename)
+	bool SessionServer_DeleteVirtualFile(uba::SessionServer* server, const uba::tchar* filename, bool registerToDir)
 	{
-		return server->DeleteVirtualFile(filename);
+		return server->DeleteVirtualFile(filename, registerToDir);
 	}
 	bool SessionServer_GetOutputFileSize(uba::SessionServer* server, uba::u64& outSize, const uba::tchar* filename)
 	{
diff --git a/Engine/Source/Programs/UnrealBuildAccelerator/Host/Public/UbaExports.h b/Engine/Source/Programs/UnrealBuildAccelerator/Host/Public/UbaExports.h
--- a/Engine/Source/Programs/UnrealBuildAccelerator/Host/Public/UbaExports.h
+++ b/Engine/Source/Programs/UnrealBuildAccelerator/Host/Public/UbaExports.h
@@ -124,8 +124,8 @@ extern "C"
 	UBA_API void SessionServer_RegisterDeleteFile(uba::SessionServer* server, const uba::tchar* filePath);
 	UBA_API bool SessionServer_RegisterNewDirectory(uba::SessionServer* server, const uba::tchar* directoryPath);
 	UBA_API bool SessionServer_RegisterVirtualFile(uba::SessionServer* server, const uba::tchar* filename, const uba::tchar* sourceFile, uba::u64 sourceOffset, uba::u64 sourceSize);
-	UBA_API bool SessionServer_CreateVirtualFile(uba::SessionServer* server, const uba::tchar* filename, const void* memory, uba::u64 memorySize, bool transient);
-	UBA_API bool SessionServer_DeleteVirtualFile(uba::SessionServer* server, const uba::tchar* filename);
+	UBA_API bool SessionServer_CreateVirtualFile(uba::SessionServer* server, const uba::tchar* filename, const void* memory, uba::u64 memorySize, bool transient, bool registerToDir);
+	UBA_API bool SessionServer_DeleteVirtualFile(uba::SessionServer* server, const uba::tchar* filename, bool registerToDir);
 
 	UBA_API bool SessionServer_GetOutputFileSize(uba::SessionServer* server, uba::u64& outSize, const uba::tchar* filename);
 	UBA_API bool SessionServer_GetOutputFileData(uba::SessionServer* server, void* outData, const uba::tchar* filename, bool deleteInternalMapping);
diff --git a/Engine/Source/Programs/UnrealBuildAccelerator/Test/Private/UbaTest.h b/Engine/Source/Programs/UnrealBuildAccelerator/Test/Private/UbaTest.h
--- a/Engine/Source/Programs/UnrealBuildAccelerator/Test/Private/UbaTest.h
+++ b/Engine/Source/Programs/UnrealBuildAccelerator/Test/Private/UbaTest.h
@@ -24,7 +24,7 @@ namespace uba
 	bool FileExists(Logger& logger, StringView dir, StringView fileName);
 	bool FilesEqual(Logger& logger, const tchar* file1, const tchar* file2);
 
-	using RunProcessFunction = Function<ProcessHandle(const ProcessStartInfo&)>;
+	using RunProcessFunction = Function<ProcessHandle(ProcessStartInfo&)>;
 	using TestSessionFunction = Function<bool(LoggerWithWriter& logger, SessionServer& session, const tchar* workingDir, const RunProcessFunction& runProcess, void* extraData)>;
 	bool RunLocal(LoggerWithWriter& logger, const StringBufferBase& testRootDir, const TestSessionFunction& testFunc, void* extraData = nullptr, bool enableDetour = true);
 	bool RunRemote(LoggerWithWriter& logger, const StringBufferBase& testRootDir, const TestSessionFunction& testFunc, bool deleteAll = true, bool serverShouldListen = true);
diff --git a/Engine/Source/Programs/UnrealBuildAccelerator/Test/Private/UbaTestCache.cpp b/Engine/Source/Programs/UnrealBuildAccelerator/Test/Private/UbaTestCache.cpp
--- a/Engine/Source/Programs/UnrealBuildAccelerator/Test/Private/UbaTestCache.cpp
+++ b/Engine/Source/Programs/UnrealBuildAccelerator/Test/Private/UbaTestCache.cpp
@@ -289,7 +289,7 @@ namespace uba
 	}
 
 	void GetTestAppPath(LoggerWithWriter& logger, StringBufferBase& out);
-	StringKey GetKeyAndFixedName(StringBuffer<>& fixedFilePath, const tchar* filePath);
+	StringKey GetKeyAndFixedName(StringBuffer<>& fixedFilePath, const tchar* filePath, StringKey* outDirKey = nullptr, const tchar** outBaseFileName = nullptr);
 	void InvalidateCachedInfo(StorageImpl& storage, StringBufferBase& fileName)
 	{
 		StringBuffer<> fixedFilePath;
diff --git a/Engine/Source/Programs/UnrealBuildAccelerator/Test/Private/UbaTestSession.cpp b/Engine/Source/Programs/UnrealBuildAccelerator/Test/Private/UbaTestSession.cpp
--- a/Engine/Source/Programs/UnrealBuildAccelerator/Test/Private/UbaTestSession.cpp
+++ b/Engine/Source/Programs/UnrealBuildAccelerator/Test/Private/UbaTestSession.cpp
@@ -59,7 +59,7 @@ namespace uba
 		if (!DeleteAllFiles(logger, workingDir.data, false))
 			return false;
 		workingDir.EnsureEndsWithSlash();
-		return testFunc(logger, session, workingDir.data, [&](const ProcessStartInfo& pi) { return session.RunProcess(pi, true, enableDetour); }, extraData);
+		return testFunc(logger, session, workingDir.data, [&](ProcessStartInfo& pi) { return session.RunProcess(pi, true, enableDetour); }, extraData);
 	}
 
 	bool RunLocal(LoggerWithWriter& logger, const StringBufferBase& testRootDir, const TestSessionFunction& testFunc, void* extraData, bool enableDetour)
@@ -249,7 +249,7 @@ namespace uba
 		return SetupServerClientSession(logger, testRootDir, deleteAll, serverShouldListen, serverConfig,
 			[&](LoggerWithWriter& logger, const StringView& workingDir, SessionServer& sessionServer, SessionClient& sessionClient)
 			{
-				return testFunc(logger, sessionServer, workingDir.data, [&](const ProcessStartInfo& pi) { return sessionServer.RunProcessRemote(pi); }, extraData);
+				return testFunc(logger, sessionServer, workingDir.data, [&](ProcessStartInfo& pi) { return sessionServer.RunProcessRemote(pi); }, extraData);
 			});
 	}
 
@@ -771,16 +771,36 @@ namespace uba
 #endif
 	}
 
+	struct RunCreateVirtualFileTestData
+	{
+		bool remote = false;
+		bool transient = false;
+		bool registerToDir = true;
+	};
+
 	bool RunCreateVirtualFileTest(LoggerWithWriter& logger, SessionServer& session, const tchar* workingDir, const RunProcessFunction& runProcess, void* extraData)
 	{
-		bool transient = *(bool*)extraData;
+		auto testData = *(RunCreateVirtualFileTestData*)extraData;
 
 		StringBuffer<> inFile(workingDir);
 		inFile.Append(TCV("VirtualFile.in"));
-		if (!session.CreateVirtualFile(inFile.data, "FOO", 3, transient))
+		if (!session.CreateVirtualFile(inFile.data, "FOO", 3, testData.transient, testData.registerToDir))
 			return false;
 		ProcessHandle ph;
-		if (!RunTestApp(logger, workingDir, runProcess, TC("-virtualFile"), &ph))
+
+
+		auto runProcess2 = [&](ProcessStartInfo& pi)
+			{
+				StringBuffer<> overlayFiles;
+				if (!testData.registerToDir)
+				{
+					overlayFiles.Append(inFile.data).Append(';').Append(inFile.data);
+					pi.overlayFiles = overlayFiles.data;
+				}
+				return testData.remote ? session.RunProcessRemote(pi) : session.RunProcess(pi, true, true);
+			};
+
+		if (!RunTestApp(logger, workingDir, runProcess2, TC("-virtualFile"), &ph))
 			return false;
 
 		StringBuffer<> outFile(workingDir);
@@ -791,7 +811,7 @@ namespace uba
 		if (!success)
 			return false;
 
-		if (!session.DeleteVirtualFile(inFile.data))
+		if (!session.DeleteVirtualFile(inFile.data, testData.registerToDir))
 			return false;
 
 		if (FileExists(logger, outFile.data))
@@ -828,11 +848,17 @@ namespace uba
 	{
 		Config config;
 		config.AddTable(TC("Session")).AddValue(TC("ShouldWriteToDisk"), false);
-		bool transient = true;
-		if (!RunLocal(logger, testRootDir, config, RunCreateVirtualFileTest, &transient))
+
+		RunCreateVirtualFileTestData testData { false };
+		if (!RunLocal(logger, testRootDir, config, RunCreateVirtualFileTest, &testData))
+			return false;
+
+		testData.transient = true;
+		if (!RunLocal(logger, testRootDir, config, RunCreateVirtualFileTest, &testData))
 			return false;
-		transient = false;
-		if (!RunLocal(logger, testRootDir, config, RunCreateVirtualFileTest, &transient))
+
+		testData.registerToDir = false;
+		if (!RunLocal(logger, testRootDir, config, RunCreateVirtualFileTest, &testData))
 			return false;
 		return true;
 	}
@@ -842,11 +868,19 @@ namespace uba
 		Config config;
 		config.AddTable(TC("Session")).AddValue(TC("ShouldWriteToDisk"), false);
 		config.AddTable(TC("Storage")).AddValue(TC("CreateIndependentMappings"), true);
-		bool transient = true;
-		if (!RunRemote(logger, testRootDir, config, RunCreateVirtualFileTest, &transient))
+
+		RunCreateVirtualFileTestData testData { true };
+
+		testData.transient = true;
+		if (!RunRemote(logger, testRootDir, config, RunCreateVirtualFileTest, &testData))
 			return false;
-		transient = false;
-		if (!RunRemote(logger, testRootDir, config, RunCreateVirtualFileTest, &transient))
+			
+		testData.transient = false;
+		if (!RunRemote(logger, testRootDir, config, RunCreateVirtualFileTest, &testData))
+			return false;
+
+		testData.registerToDir = false;
+		if (!RunRemote(logger, testRootDir, config, RunCreateVirtualFileTest, &testData))
 			return false;
 		return true;
 	}
@@ -909,7 +943,7 @@ namespace uba
 		return SetupServerSession(logger, testRootDir, true, true, config,
 			[&](LoggerWithWriter& logger, const StringView& workingDir, SessionServer& sessionServer)
 			{
-				auto runProcess = [&](const ProcessStartInfo& pi) { return sessionServer.RunProcessRemote(pi); };
+				auto runProcess = [&](ProcessStartInfo& pi) { return sessionServer.RunProcessRemote(pi); };
 				if (!CreateTestFile(logger, workingDir, TCV("SpecialFile1"), TCV("0")))
 					return false;
 
