[UBA]
* Implemented overlay files in posix detouring code

diff --git a/Engine/Source/Programs/UnrealBuildAccelerator/Detours/Private/PosixOS/UbaDetoursFunctionsPosix.cpp b/Engine/Source/Programs/UnrealBuildAccelerator/Detours/Private/PosixOS/UbaDetoursFunctionsPosix.cpp
--- a/Engine/Source/Programs/UnrealBuildAccelerator/Detours/Private/PosixOS/UbaDetoursFunctionsPosix.cpp
+++ b/Engine/Source/Programs/UnrealBuildAccelerator/Detours/Private/PosixOS/UbaDetoursFunctionsPosix.cpp
@@ -40,6 +40,10 @@ namespace uba
 	u32 g_processId;
 	pid_t g_pid;
 	void Deinit();
+
+	SharedMemoryAllocatorHandle g_sharedMemoryAllocator;
+	u8* g_sharedMemory;
+	u32 g_overlayFilesSize;
 }
 
 using namespace uba;
@@ -2848,6 +2852,8 @@ namespace uba
 		u64 mappedFileTableHandle;
 		u32 mappedFileTableSize;
 		u32 mappedFileTableCount;
+		SharedMemoryHandle overlayFilesHandle;
+		u32 overlayFilesSize;
 
 		{
 			TimerScope ts(g_stats.init);
@@ -2863,7 +2869,7 @@ namespace uba
 			reader.ReadString(g_virtualApplication);
 			reader.ReadString(g_virtualWorkingDir);
 
-			reader.ReadU64(); // shared memory allocator handle
+			g_sharedMemoryAllocator = SharedMemoryAllocatorHandle::FromU64(reader.ReadU64());
 			reader.ReadU64(); // permanent files memory handle
 			directoryTableHandle = reader.ReadU64();
 			directoryTableSize = reader.ReadU32();
@@ -2872,6 +2878,9 @@ namespace uba
 			mappedFileTableSize = reader.ReadU32();
 			mappedFileTableCount = reader.ReadU32();
 
+			overlayFilesHandle = SharedMemoryHandle::FromU64(reader.ReadU64());
+			overlayFilesSize = reader.ReadU32();
+
 			if (u16 vfsSize = reader.ReadU16())
 			{
 				BinaryReader vfsReader(reader.GetPositionData(), 0, vfsSize);
@@ -2912,6 +2921,14 @@ namespace uba
 		UBA_ASSERTF(dirTableMem != MAP_FAILED, "mmap for dirtable mem failed (%s)", strerror(errno));
 		g_directoryTable.Init(dirTableMem, directoryTableCount, directoryTableSize);
 
+		if (overlayFilesSize)
+		{
+			SharedMemoryView overlayFilesView;
+			Rpc_GetSharedMemory(overlayFilesView, g_sharedMemoryAllocator, nullptr, overlayFilesHandle, 0, 0, TC("PermanentFiles"));
+			g_directoryTable.InitOverlay(overlayFilesView.DetachMemory(), overlayFilesSize);
+			g_overlayFilesSize = overlayFilesSize;
+		}
+
 		if (g_isChild)
 			Rpc_GetWrittenFiles();
 
@@ -3055,8 +3072,18 @@ extern "C"
 		FlushDebugLog();
 		#endif
 
+		Vector<u8> oldOverlayFiles;
+		if (g_overlayFilesSize)
+		{
+			oldOverlayFiles.resize(g_overlayFilesSize);
+			memcpy(oldOverlayFiles.data(), g_directoryTable.m_overlayMemory, g_overlayFilesSize);
+
+		}
+
 		*outArguments = 0;
 		bool newProcess;
+		u32 overlayFilesSize = 0;
+		u32 directoryTableSize = 0;
 		{
 			SCOPED_WRITE_LOCK(g_communicationLock, pcs);
 			BinaryWriter writer;
@@ -3074,6 +3101,8 @@ extern "C"
 				reader.SkipString(); // workingDir
 				reader.SkipString(); // description
 				reader.ReadString(g_logName.Clear());
+				directoryTableSize = reader.ReadU32();
+				overlayFilesSize = reader.ReadU32();
 			}
 		}
 
@@ -3091,6 +3120,15 @@ extern "C"
 			g_debugFile = (FileHandle)debugFile;
 			LogHeader(ToView(outArguments));
 			#endif
+
+			if (overlayFilesSize)
+			{
+				const u8* overlayMem = g_directoryTable.m_overlayMemory;
+				g_directoryTable.DeinitOverlay(oldOverlayFiles.data(), g_overlayFilesSize);
+				g_directoryTable.ParseDirectoryTableNoLock(directoryTableSize);
+				g_directoryTable.InitOverlay(overlayMem, overlayFilesSize);
+				g_overlayFilesSize = overlayFilesSize;
+			}
 		}
 
 		Rpc_UpdateTables();
