[UBA]
* Fixed storage network compatibility mismatch
* Bumped storage network version and added backwards compatibility (needed for cache server)

diff --git a/Engine/Source/Programs/UnrealBuildAccelerator/Common/Private/UbaNetworkServer.cpp b/Engine/Source/Programs/UnrealBuildAccelerator/Common/Private/UbaNetworkServer.cpp
--- a/Engine/Source/Programs/UnrealBuildAccelerator/Common/Private/UbaNetworkServer.cpp
+++ b/Engine/Source/Programs/UnrealBuildAccelerator/Common/Private/UbaNetworkServer.cpp
@@ -34,6 +34,7 @@ namespace uba
 		Guid uid;
 		u32 id;
 		u32 queuedConnectionCount = 0;
+		u32 storageVersion = 0;
 		Atomic<u32> refCount = 0;
 		Atomic<u32> connectionCount = 0;
 		Atomic<u64> sendBytes = 0;
@@ -607,6 +608,18 @@ namespace uba
 		return conn.m_shouldDisconnect;
 	}
 
+	void ConnectionInfo::SetStorageVersion(u32 version) const
+	{
+		auto& conn = *(NetworkServer::Connection*)internalData;
+		conn.m_client->storageVersion = version;
+	}
+
+	u32 ConnectionInfo::GetStorageVersion() const
+	{
+		auto& conn = *(NetworkServer::Connection*)internalData;
+		return conn.m_client->storageVersion;
+	}
+
 	void NetworkServer::Worker::Update(WorkerContext& context)
 	{
 		auto& server = context.server;
diff --git a/Engine/Source/Programs/UnrealBuildAccelerator/Common/Private/UbaStorageProxy.cpp b/Engine/Source/Programs/UnrealBuildAccelerator/Common/Private/UbaStorageProxy.cpp
--- a/Engine/Source/Programs/UnrealBuildAccelerator/Common/Private/UbaStorageProxy.cpp
+++ b/Engine/Source/Programs/UnrealBuildAccelerator/Common/Private/UbaStorageProxy.cpp
@@ -150,9 +150,11 @@ namespace uba
 		u32 clientVersion = reader.ReadU32();
 		if (clientVersion != StorageNetworkVersion)
 		{
-			m_logger.Error(TC("Different network versions. Client: %u, Server: %u. Disconnecting"), clientVersion, StorageNetworkVersion);
+			m_logger.Error(TC("Unsupported network version. Client: %u, Server: %u. Disconnecting"), clientVersion, StorageNetworkVersion);
 			return false;
 		}
+		connectionInfo.SetStorageVersion(clientVersion);
+
 		bool isInProcessClient = reader.ReadBool();
 		if (isInProcessClient)
 			m_inProcessClientId = connectionInfo.GetId();
diff --git a/Engine/Source/Programs/UnrealBuildAccelerator/Common/Private/UbaStorageServer.cpp b/Engine/Source/Programs/UnrealBuildAccelerator/Common/Private/UbaStorageServer.cpp
--- a/Engine/Source/Programs/UnrealBuildAccelerator/Common/Private/UbaStorageServer.cpp
+++ b/Engine/Source/Programs/UnrealBuildAccelerator/Common/Private/UbaStorageServer.cpp
@@ -516,8 +516,9 @@ namespace uba
 				StringBuffer<> clientName;
 				reader.ReadString(clientName);
 				u32 clientVersion = reader.ReadU32();
-				if (clientVersion != StorageNetworkVersion)
-					return m_logger.Error(TC("Different network versions. Client: %u, Server: %u. Disconnecting"), clientVersion, StorageNetworkVersion);
+				if (clientVersion < 4 || clientVersion > 5)
+					return m_logger.Error(TC("Unsupported network version. Client: %u, Server: %u. Disconnecting"), clientVersion, StorageNetworkVersion);
+				connectionInfo.SetStorageVersion(clientVersion);
 
 				if (reader.ReadBool()) // is proxy
 					return m_logger.Error(TC("Proxy is sending connect message. This path is not implemented"));
@@ -580,7 +581,10 @@ namespace uba
 
 				u64 start = GetTime();
 				casKey = reader.ReadCasKey();
-				StringKey keyHint = reader.ReadStringKey();
+				StringKey keyHint = {};
+				if (connectionInfo.GetStorageVersion() >= 5)
+					keyHint = reader.ReadStringKey();
+
 				reader.ReadString(hint);
 
 				StringBuffer<16> temp;
diff --git a/Engine/Source/Programs/UnrealBuildAccelerator/Common/Private/UbaStorageUtils.cpp b/Engine/Source/Programs/UnrealBuildAccelerator/Common/Private/UbaStorageUtils.cpp
--- a/Engine/Source/Programs/UnrealBuildAccelerator/Common/Private/UbaStorageUtils.cpp
+++ b/Engine/Source/Programs/UnrealBuildAccelerator/Common/Private/UbaStorageUtils.cpp
@@ -599,6 +599,7 @@ namespace uba
 			NetworkMessage msg(client, StorageServiceId, StorageMessageType_FetchBegin, writer);
 			writer.WriteBool(false);
 			writer.WriteCasKey(casKey);
+			writer.WriteStringKey({});
 			writer.WriteString(destination);
 			BinaryReader reader(slot + (writeCompressed ? sizeof(CompressedFileHeader) : 0), 0, SendMaxSize); // Create some space in front of reader to write obj file header there if destination is compressed
 			if (!msg.Send(reader))
diff --git a/Engine/Source/Programs/UnrealBuildAccelerator/Common/Public/UbaNetwork.h b/Engine/Source/Programs/UnrealBuildAccelerator/Common/Public/UbaNetwork.h
--- a/Engine/Source/Programs/UnrealBuildAccelerator/Common/Public/UbaNetwork.h
+++ b/Engine/Source/Programs/UnrealBuildAccelerator/Common/Public/UbaNetwork.h
@@ -15,7 +15,7 @@ namespace uba
 	static constexpr u8 CacheServiceId = 3;
 	
 	static constexpr u32 SystemNetworkVersion = 1339;
-	static constexpr u32 StorageNetworkVersion = 4;
+	static constexpr u32 StorageNetworkVersion = 5;
 	static constexpr u32 SessionNetworkVersion = 47;
 	static constexpr u32 CacheNetworkVersion = 5;
 	static constexpr u32 CacheBucketVersion = 3;
diff --git a/Engine/Source/Programs/UnrealBuildAccelerator/Common/Public/UbaNetworkServer.h b/Engine/Source/Programs/UnrealBuildAccelerator/Common/Public/UbaNetworkServer.h
--- a/Engine/Source/Programs/UnrealBuildAccelerator/Common/Public/UbaNetworkServer.h
+++ b/Engine/Source/Programs/UnrealBuildAccelerator/Common/Public/UbaNetworkServer.h
@@ -26,6 +26,10 @@ namespace uba
 		bool GetName(StringBufferBase& out) const;
 		u64 GetCryptoUserData() const;
 		bool ShouldDisconnect() const;
+
+		void SetStorageVersion(u32 version) const;
+		u32 GetStorageVersion() const;
+
 		void* internalData = nullptr;
 	};
 
