[UBA]
* Fixed race condition in unit test where server did not wait for cas fully before using it

diff --git a/Engine/Source/Programs/UnrealBuildAccelerator/Common/Private/UbaStorageServer.cpp b/Engine/Source/Programs/UnrealBuildAccelerator/Common/Private/UbaStorageServer.cpp
--- a/Engine/Source/Programs/UnrealBuildAccelerator/Common/Private/UbaStorageServer.cpp
+++ b/Engine/Source/Programs/UnrealBuildAccelerator/Common/Private/UbaStorageServer.cpp
@@ -88,21 +88,8 @@ namespace uba
 
 	}
 
-	bool StorageServer::GetZone(StringBufferBase& out)
+	bool StorageServer::WaitForCasFile(const CasKey& casKey, const tchar* hint, u32 clientId)
 	{
-		if (m_zone.empty())
-			return false;
-		out.Append(m_zone);
-		return true;
-	}
-
-	bool StorageServer::RetrieveCasFile(RetrieveResult& out, const CasKey& casKey, StringKey keyHint, const tchar* hint, FileMappingBuffer* mappingBuffer, u64 memoryMapAlignment, bool allowProxy, u32 clientId)
-	{
-		UBA_ASSERT(!mappingBuffer);
-		UBA_ASSERT(casKey != CasKeyZero);
-		out.casKey = casKey;
-		out.size = InvalidValue;
-
 		CasKey actualKey = casKey;
 
 		// This must be true because code handling storing files from remotes does not have decompression support
@@ -143,6 +130,23 @@ namespace uba
 		return waitEntry.success;
 	}
 
+	bool StorageServer::GetZone(StringBufferBase& out)
+	{
+		if (m_zone.empty())
+			return false;
+		out.Append(m_zone);
+		return true;
+	}
+
+	bool StorageServer::RetrieveCasFile(RetrieveResult& out, const CasKey& casKey, StringKey keyHint, const tchar* hint, FileMappingBuffer* mappingBuffer, u64 memoryMapAlignment, bool allowProxy, u32 clientId)
+	{
+		UBA_ASSERT(!mappingBuffer);
+		UBA_ASSERT(casKey != CasKeyZero);
+		out.casKey = casKey;
+		out.size = InvalidValue;
+		return WaitForCasFile(casKey, hint, clientId);
+	}
+
 	bool StorageServer::CopyOrLink(const CasKey& casKey, const tchar* destination, u32 fileAttributes, bool writeCompressed, u32 clientId, const FormattingFunc& formattingFunc, bool isTemp, bool allowHardLink)
 	{
 		bool res = StorageImpl::CopyOrLink(casKey, destination, fileAttributes, writeCompressed, clientId, formattingFunc, isTemp, allowHardLink);
diff --git a/Engine/Source/Programs/UnrealBuildAccelerator/Common/Public/UbaStorageServer.h b/Engine/Source/Programs/UnrealBuildAccelerator/Common/Public/UbaStorageServer.h
--- a/Engine/Source/Programs/UnrealBuildAccelerator/Common/Public/UbaStorageServer.h
+++ b/Engine/Source/Programs/UnrealBuildAccelerator/Common/Public/UbaStorageServer.h
@@ -37,6 +37,7 @@ namespace uba
 		using StorageImpl::StoreCasFile;
 
 		void WaitForActiveWork();
+		bool WaitForCasFile(const CasKey& casKey, const tchar* hint, u32 clientId = 0);
 
 		virtual bool CopyOrLink(const CasKey& casKey, const tchar* destination, u32 fileAttributes, bool writeCompressed = false, u32 clientId = 0, const FormattingFunc& formattingFunc = {}, bool isTemp = false, bool allowHardLink = true) override;
 
diff --git a/Engine/Source/Programs/UnrealBuildAccelerator/Test/Private/UbaTestStorage.cpp b/Engine/Source/Programs/UnrealBuildAccelerator/Test/Private/UbaTestStorage.cpp
--- a/Engine/Source/Programs/UnrealBuildAccelerator/Test/Private/UbaTestStorage.cpp
+++ b/Engine/Source/Programs/UnrealBuildAccelerator/Test/Private/UbaTestStorage.cpp
@@ -335,10 +335,12 @@ namespace uba
 		auto copy = [&]()
 			{
 				destFileName.Clear().Append(testRootDir).Append(TCV("Uba")).Append(PathSeparator).Append(TCV("UbaTestFile"));
+				if (!storageServer.WaitForCasFile(key, destFileName.data))
+					return false;
 				if (!storageServer.CopyOrLink(key, destFileName.data, DefaultAttributes(), false, 1, {}, false, true))
 					return logger.Error(TC("Failed to copy cas to file %s"), destFileName.data);
 				if (!FilesEqual(logger, sourceFileName.data, destFileName.data))
-					return false;
+					return logger.Error(TC("Files %s and %s not equal"), sourceFileName.data, destFileName.data);
 				return true;
 			};
 
