[UBA]
* Added retry when writing file and get ERROR_WORKING_SET_QUOTA or ERROR_PAGEFILE_QUOTA errors

diff --git a/Engine/Source/Programs/UnrealBuildAccelerator/Common/Private/UbaFileAccessor.cpp b/Engine/Source/Programs/UnrealBuildAccelerator/Common/Private/UbaFileAccessor.cpp
--- a/Engine/Source/Programs/UnrealBuildAccelerator/Common/Private/UbaFileAccessor.cpp
+++ b/Engine/Source/Programs/UnrealBuildAccelerator/Common/Private/UbaFileAccessor.cpp
@@ -365,6 +365,7 @@ namespace uba
 		}
 #endif
 
+		u32 retryIndex = 0;
 		u64 writeLeft = dataLen;
 		u8* pos = (u8*)data;
 		while (writeLeft)
@@ -377,6 +378,18 @@ namespace uba
 			if (!::WriteFile(asHANDLE(m_fileHandle), pos, toActuallyWrite, &written, NULL))
 			{
 				DWORD lastError = GetLastError();
+				if ((lastError == ERROR_WORKING_SET_QUOTA || lastError == ERROR_PAGEFILE_QUOTA) && retryIndex < 5)
+				{
+					if (retryIndex == 0)
+					{
+						const tchar* quotaType = lastError == ERROR_WORKING_SET_QUOTA ? TC("Workingset") : TC("Pagefile");
+						m_logger.Info(TC("Insufficient %s quota while writing %s. Will retry..."), quotaType, m_fileName);
+					}
+					++retryIndex;
+					Sleep(retryIndex*1000);
+					continue;
+				}
+
 				m_logger.Error(TC("ERROR writing file %s writing %u bytes (%llu bytes written out of %llu. FilePos: %llu) (%s)"), m_fileName, toActuallyWrite, (dataLen - writeLeft), dataLen, GetFilePointer(m_logger, m_fileHandle), LastErrorToText(lastError).data);
 
 				if (lastError == ERROR_DISK_FULL)
