[UBA]
* Added handling for EAGAIN and EINTR in send/recv.

diff --git a/Engine/Source/Programs/UnrealBuildAccelerator/Common/Private/UbaNetworkBackendTcp.cpp b/Engine/Source/Programs/UnrealBuildAccelerator/Common/Private/UbaNetworkBackendTcp.cpp
--- a/Engine/Source/Programs/UnrealBuildAccelerator/Common/Private/UbaNetworkBackendTcp.cpp
+++ b/Engine/Source/Programs/UnrealBuildAccelerator/Common/Private/UbaNetworkBackendTcp.cpp
@@ -1131,7 +1131,15 @@ namespace uba
 			return false;
 
 		u32 lastError = WSAGetLastError();
-		if (lastError != WSAEWOULDBLOCK)
+		bool shouldPoll = lastError == WSAEWOULDBLOCK;
+
+		#if !PLATFORM_WINDOWS
+		if (lastError == EINTR)
+			return true;
+		shouldPoll = shouldPoll || lastError == EAGAIN;
+		#endif
+
+		if (!shouldPoll)
 		{
 			#if UBA_LOG_SOCKET_ERRORS
 			logger.Info(TC("SendSocket - send returned an error for socket %i: %s (%s)"), socket, LastErrorToText(WSAGetLastError()).data, hint);
@@ -1493,7 +1501,16 @@ namespace uba
 
 			UBA_ASSERT(read == SOCKET_ERROR);
 
-			if (WSAGetLastError() == WSAEWOULDBLOCK)
+			u32 lastError = WSAGetLastError();
+			bool shouldPoll = lastError == WSAEWOULDBLOCK;
+
+			#if !PLATFORM_WINDOWS
+			if (lastError == EINTR)
+				continue;
+			shouldPoll = shouldPoll || lastError == EAGAIN;
+			#endif
+
+			if (shouldPoll)
 			{
 				if (retryCount > 0)
 				{
diff --git a/Engine/Source/Programs/UnrealBuildAccelerator/Common/Private/UbaNetworkClient.cpp b/Engine/Source/Programs/UnrealBuildAccelerator/Common/Private/UbaNetworkClient.cpp
--- a/Engine/Source/Programs/UnrealBuildAccelerator/Common/Private/UbaNetworkClient.cpp
+++ b/Engine/Source/Programs/UnrealBuildAccelerator/Common/Private/UbaNetworkClient.cpp
@@ -388,8 +388,13 @@ namespace uba
 				found = m == &msg;
 			lock.Leave();
 			if (!found)
-				return client.m_isDisconnecting ? false : client.m_logger.Error(TC("NetworkClient got receive response body for message that does not exist"));
-
+			{
+				if (client.m_isDisconnecting)
+					return false;
+				u16 messageId = u16(headerData[0] << 8) | u16((*(u32*)(headerData + 1) & 0xff000000) >> 24);
+				u32 messageSize = *(u32*)(headerData + 1) & 0x00FFFFFF;
+				return client.m_logger.Error(TC("NetworkClient got receive response body for message that does not exist (Id: %u, Size: %u)"), u32(messageId), messageSize);
+			}
 			msg.m_error = 2;
 		}
 		#if UBA_TRACK_NETWORK_TIMES
