[UBA]
* Added fileNameLen to FixPath2.. Since NtCreateFile does not always get null terminated string we have to make sure we provide the size and don't read outside

diff --git a/Engine/Source/Programs/UnrealBuildAccelerator/Common/Private/UbaFile.cpp b/Engine/Source/Programs/UnrealBuildAccelerator/Common/Private/UbaFile.cpp
--- a/Engine/Source/Programs/UnrealBuildAccelerator/Common/Private/UbaFile.cpp
+++ b/Engine/Source/Programs/UnrealBuildAccelerator/Common/Private/UbaFile.cpp
@@ -952,7 +952,7 @@ namespace uba
 
 		char fullPath[1024];
 		u32 len;
-		if (!FixPath2(fileName, cwd, cwdlen, fullPath, sizeof(fullPath), &len))
+		if (!FixPath2(fileName, ~0u, cwd, cwdlen, fullPath, sizeof(fullPath), &len))
 			return 0;
 #endif
 		UBA_ASSERT(len < nBufferLength);
diff --git a/Engine/Source/Programs/UnrealBuildAccelerator/Common/Public/UbaProcessStartInfoHolder.h b/Engine/Source/Programs/UnrealBuildAccelerator/Common/Public/UbaProcessStartInfoHolder.h
--- a/Engine/Source/Programs/UnrealBuildAccelerator/Common/Public/UbaProcessStartInfoHolder.h
+++ b/Engine/Source/Programs/UnrealBuildAccelerator/Common/Public/UbaProcessStartInfoHolder.h
@@ -18,7 +18,7 @@ namespace uba
 		u64 workingDirLen = 0;
 		if (workingDir)
 			workingDirLen = TStrlen(workingDir);
-		FixPath2(fileName, workingDir, workingDirLen, buffer, sizeof_array(buffer), &charLen);
+		FixPath2(fileName, ~0u, workingDir, workingDirLen, buffer, sizeof_array(buffer), &charLen);
 		out.Append(buffer);
 	}
 
diff --git a/Engine/Source/Programs/UnrealBuildAccelerator/Core/Public/UbaPathUtils.h b/Engine/Source/Programs/UnrealBuildAccelerator/Core/Public/UbaPathUtils.h
--- a/Engine/Source/Programs/UnrealBuildAccelerator/Core/Public/UbaPathUtils.h
+++ b/Engine/Source/Programs/UnrealBuildAccelerator/Core/Public/UbaPathUtils.h
@@ -13,7 +13,7 @@ namespace uba
 	#endif
 
 
-	inline bool FixPath2(const tchar* fileName, const tchar* workingDir, u64 workingDirCharLen, tchar* buffer, u64 bufferCharCapacity, u32* outBufferCharLen)
+	inline bool FixPath2(const tchar* fileName, u32 fileNameLen, const tchar* workingDir, u64 workingDirCharLen, tchar* buffer, u64 bufferCharCapacity, u32* outBufferCharLen)
 	{
 		buffer[1] = 0;
 		UBA_ASSERT(fileName);
@@ -48,13 +48,15 @@ namespace uba
 			if (!tilde)
 				break;
 
-			u64 fileNameLen = TStrlen(read);
+			if (fileNameLen == ~0u)
+				fileNameLen = TStrlen(read);
 			if (fileNameLen >= 1024)
 				return false;
 
 			// Since this might be a in memory file we can't use the actual full name, so let's find first slash after '~'
 			const tchar* backslash = TStrchr(tilde, '\\');
-			TStrcpy_s(fullName, sizeof_array(fullName), read);
+			memcpy(fullName, read, fileNameLen*sizeof(tchar));
+			fullName[fileNameLen] = 0;
 			if (backslash)
 				fullName[backslash - read] = 0;
 			u32 len = Local_GetLongPathNameW(fullName, fullName, sizeof_array(fullName));
@@ -67,6 +69,7 @@ namespace uba
 			if (backslash)
 				TStrcpy_s(fullName + len, sizeof_array(fullName) - len, backslash);
 			read = fullName;
+			fileNameLen = TStrlen(read);
 		}
 
 		if (read[0] == '\\')
@@ -101,8 +104,13 @@ namespace uba
 
 		bool isAbsolute = false;
 
-		while (tchar c = *read++)
+		const tchar* readEnd = read + fileNameLen;
+		while (read < readEnd)
 		{
+			tchar c = *read++;
+			if (!c)
+				break;
+
 			if (c == '/')
 				c = '\\';
 			else if (c == '\"')
@@ -129,6 +137,7 @@ namespace uba
 			lastLastChar = lastChar;
 			lastChar = c;
 		}
+
 		if (lastChar == '.' && lastLastChar == '\\') // Fix path <path>\.
 			write -= 2;
 
@@ -360,7 +369,7 @@ namespace uba
 	inline bool FixPath(const tchar* fileName, const tchar* workingDir, u64 workingDirCharLen, StringBufferBase& buffer)
 	{
 		u32 count = buffer.count;
-		bool res = FixPath2(fileName, workingDir, workingDirCharLen, buffer.data + count, buffer.capacity - count, &buffer.count);
+		bool res = FixPath2(fileName, ~0u, workingDir, workingDirCharLen, buffer.data + count, buffer.capacity - count, &buffer.count);
 		buffer.count += count;
 		return res;
 	}
diff --git a/Engine/Source/Programs/UnrealBuildAccelerator/Detours/Private/UbaDetoursShared.cpp b/Engine/Source/Programs/UnrealBuildAccelerator/Detours/Private/UbaDetoursShared.cpp
--- a/Engine/Source/Programs/UnrealBuildAccelerator/Detours/Private/UbaDetoursShared.cpp
+++ b/Engine/Source/Programs/UnrealBuildAccelerator/Detours/Private/UbaDetoursShared.cpp
@@ -175,9 +175,9 @@ namespace uba
 	SuppressDetourScope::~SuppressDetourScope() { --t_disallowDetour; }
 
 
-	bool FixPath(StringBufferBase& out, const tchar* path)
+	bool FixPath(StringBufferBase& out, const tchar* path, u32 pathLen)
 	{
-		return FixPath2(path, g_virtualWorkingDir.data, g_virtualWorkingDir.count, out.data, out.capacity, &out.count);
+		return FixPath2(path, pathLen, g_virtualWorkingDir.data, g_virtualWorkingDir.count, out.data, out.capacity, &out.count);
 	}
 
 	struct VfsEntry { StringView vfs; StringView local; VfsEntry() : vfs(NoInit), local(NoInit) {}; };
diff --git a/Engine/Source/Programs/UnrealBuildAccelerator/Detours/Private/UbaDetoursShared.h b/Engine/Source/Programs/UnrealBuildAccelerator/Detours/Private/UbaDetoursShared.h
--- a/Engine/Source/Programs/UnrealBuildAccelerator/Detours/Private/UbaDetoursShared.h
+++ b/Engine/Source/Programs/UnrealBuildAccelerator/Detours/Private/UbaDetoursShared.h
@@ -22,7 +22,7 @@ namespace uba
 	void Rpc_ResolveCallstack(StringBufferBase& out, u32 skipCallstackCount, void* context);
 
 	const tchar* GetApplicationShortName();
-	bool FixPath(StringBufferBase& out, const tchar* path);
+	bool FixPath(StringBufferBase& out, const tchar* path, u32 pathLen = ~0u);
 	
 	void PopulateVfs(BinaryReader& vfsReader);
 	bool IsVfsEnabled();
diff --git a/Engine/Source/Programs/UnrealBuildAccelerator/Detours/Private/Windows/UbaDetoursFunctionsKernelBase.inl b/Engine/Source/Programs/UnrealBuildAccelerator/Detours/Private/Windows/UbaDetoursFunctionsKernelBase.inl
--- a/Engine/Source/Programs/UnrealBuildAccelerator/Detours/Private/Windows/UbaDetoursFunctionsKernelBase.inl
+++ b/Engine/Source/Programs/UnrealBuildAccelerator/Detours/Private/Windows/UbaDetoursFunctionsKernelBase.inl
@@ -125,7 +125,7 @@ void Shared_SetCurrentDirectory(const wchar_t* workingDirBuffer)
 {
 	u32 charLen = 0;
 	wchar_t temp[256];
-	FixPath2(workingDirBuffer, g_virtualWorkingDir.data, g_virtualWorkingDir.count, temp, sizeof_array(temp), &charLen);
+	FixPath2(workingDirBuffer, ~0u, g_virtualWorkingDir.data, g_virtualWorkingDir.count, temp, sizeof_array(temp), &charLen);
 	g_virtualWorkingDir.Clear().Append(temp).Append('\\');
 }
 
diff --git a/Engine/Source/Programs/UnrealBuildAccelerator/Detours/Private/Windows/UbaDetoursFunctionsNtDll.inl b/Engine/Source/Programs/UnrealBuildAccelerator/Detours/Private/Windows/UbaDetoursFunctionsNtDll.inl
--- a/Engine/Source/Programs/UnrealBuildAccelerator/Detours/Private/Windows/UbaDetoursFunctionsNtDll.inl
+++ b/Engine/Source/Programs/UnrealBuildAccelerator/Detours/Private/Windows/UbaDetoursFunctionsNtDll.inl
@@ -786,7 +786,7 @@ NTSTATUS NTAPI Shared_NtCreateFile(bool IsCreateFunc, PHANDLE hFileHandle, ACCES
 		{
 			if (!CanDetour(buf))
 				suppressCreateFileDetour = true;
-			else if (!FixPath(fileName, buf))
+			else if (!FixPath(fileName, buf, bufChars))
 				UBA_ASSERTF(false, L"FixPath failed for string '%ls'", buf + 4);
 		}
 		else
diff --git a/Engine/Source/Programs/UnrealBuildAccelerator/Detours/Private/Windows/UbaDetoursFunctionsUcrtBase.inl b/Engine/Source/Programs/UnrealBuildAccelerator/Detours/Private/Windows/UbaDetoursFunctionsUcrtBase.inl
--- a/Engine/Source/Programs/UnrealBuildAccelerator/Detours/Private/Windows/UbaDetoursFunctionsUcrtBase.inl
+++ b/Engine/Source/Programs/UnrealBuildAccelerator/Detours/Private/Windows/UbaDetoursFunctionsUcrtBase.inl
@@ -35,12 +35,12 @@ wchar_t* Detoured__wfullpath(wchar_t* absPath, const wchar_t* relPath, size_t ma
 			res = (wchar_t*)malloc(maxLength * sizeof(WCHAR)); // Same as wine implementation
 		else
 			res = (wchar_t*)_malloc_base(maxLength * sizeof(WCHAR));
-		FixPath2(relPath, g_virtualWorkingDir.data, g_virtualWorkingDir.count, res, maxLength, nullptr);
+		FixPath2(relPath, ~0u, g_virtualWorkingDir.data, g_virtualWorkingDir.count, res, maxLength, nullptr);
 		UBA_ASSERT(res && wcslen(res) < maxLength - 1);
 	}
 	else if ((relPath && relPath[1] == ':' && relLen < maxLength) || (relLen + g_virtualWorkingDir.count <= maxLength))
 	{
-		FixPath2(relPath, g_virtualWorkingDir.data, g_virtualWorkingDir.count, absPath, maxLength, nullptr);
+		FixPath2(relPath, ~0u, g_virtualWorkingDir.data, g_virtualWorkingDir.count, absPath, maxLength, nullptr);
 		UBA_ASSERT(wcslen(absPath) < maxLength);
 		res = absPath;
 	}
diff --git a/Engine/Source/Programs/UnrealBuildAccelerator/Test/Private/UbaTestBasics.cpp b/Engine/Source/Programs/UnrealBuildAccelerator/Test/Private/UbaTestBasics.cpp
--- a/Engine/Source/Programs/UnrealBuildAccelerator/Test/Private/UbaTestBasics.cpp
+++ b/Engine/Source/Programs/UnrealBuildAccelerator/Test/Private/UbaTestBasics.cpp
@@ -179,10 +179,10 @@ namespace uba
 		tchar buffer[1024];
 		u32 lengthResult;
 
-		auto TestPath = [&](const tchar* path) { return FixPath2(path, workingDir, TStrlen(workingDir), buffer, sizeof_array(buffer), &lengthResult); };
+		auto TestPath = [&](const tchar* path) { return FixPath2(path, ~0u, workingDir, TStrlen(workingDir), buffer, sizeof_array(buffer), &lengthResult); };
 
 #if PLATFORM_WINDOWS
-		if (!FixPath2(TC("\"e:\\temp\""), workingDir, TStrlen(workingDir), buffer, sizeof_array(buffer), &lengthResult))
+		if (!FixPath2(TC("\"e:\\temp\""), ~0u, workingDir, TStrlen(workingDir), buffer, sizeof_array(buffer), &lengthResult))
 			return logger.Error(TC("FixPath2 (1) failed"));
 #else
 
@@ -190,24 +190,24 @@ namespace uba
 			return logger.Error(TC("FixPath2 should have failed"));
 		UBA_TEST_CHECK(Equals(buffer, TC("/")), "Should not contain ..");
 
-		if (!FixPath2(TC("/../Foo"), workingDir, TStrlen(workingDir), buffer, sizeof_array(buffer), &lengthResult))
+		if (!FixPath2(TC("/../Foo"), ~0u, workingDir, TStrlen(workingDir), buffer, sizeof_array(buffer), &lengthResult))
 			return logger.Error(TC("FixPath2 should have failed"));
 		UBA_TEST_CHECK(Equals(buffer, TC("/Foo")), "Should not contain ..");
 
-		if (!FixPath2(TC("/usr/bin//clang++"), workingDir, TStrlen(workingDir), buffer, sizeof_array(buffer), &lengthResult))
+		if (!FixPath2(TC("/usr/bin//clang++"), ~0u, workingDir, TStrlen(workingDir), buffer, sizeof_array(buffer), &lengthResult))
 			return logger.Error(TC("FixPath2 should have failed"));
 		UBA_TEST_CHECK(!Contains(buffer, TC("//")), "Should not contain //");
 #endif
 
-		if (!FixPath2(TC("../Foo"), workingDir, TStrlen(workingDir), buffer, sizeof_array(buffer), &lengthResult))
+		if (!FixPath2(TC("../Foo"), ~0u, workingDir, TStrlen(workingDir), buffer, sizeof_array(buffer), &lengthResult))
 			return logger.Error(TC("FixPath2 (1) failed"));
 		UBA_TEST_CHECK(!Contains(buffer, TC("..")), "Should not contain ..");
 
-		if (!FixPath2(TC("@../Foo"), workingDir, TStrlen(workingDir), buffer, sizeof_array(buffer), &lengthResult))
+		if (!FixPath2(TC("@../Foo"), ~0u, workingDir, TStrlen(workingDir), buffer, sizeof_array(buffer), &lengthResult))
 			return logger.Error(TC("FixPath2 (1) failed"));
 		UBA_TEST_CHECK(Contains(buffer, TC("..")), "Should contain ..");
 
-		if (!FixPath2(TC("..@/Foo"), workingDir, TStrlen(workingDir), buffer, sizeof_array(buffer), &lengthResult))
+		if (!FixPath2(TC("..@/Foo"), ~0u, workingDir, TStrlen(workingDir), buffer, sizeof_array(buffer), &lengthResult))
 			return logger.Error(TC("FixPath2 (1) failed"));
 		UBA_TEST_CHECK(Contains(buffer, TC("..")), "Should contain ..");
 
