[UBA]
* Implemented native socket support when running through wine. This speeds up socket communication a bit by removing wine overhead (socket functions are called very frequently)

diff --git a/Engine/Source/Programs/UnrealBuildAccelerator/Agent/Private/UbaAgent.cpp b/Engine/Source/Programs/UnrealBuildAccelerator/Agent/Private/UbaAgent.cpp
--- a/Engine/Source/Programs/UnrealBuildAccelerator/Agent/Private/UbaAgent.cpp
+++ b/Engine/Source/Programs/UnrealBuildAccelerator/Agent/Private/UbaAgent.cpp
@@ -130,6 +130,8 @@ namespace uba
 		return false;
 	}
 
+	bool UsingNativeSockets();
+
 	ReaderWriterLock* g_exitLock = new ReaderWriterLock();
 	LoggerWithWriter* g_logger;
 	SessionClient* g_sessionClient;
@@ -836,6 +838,10 @@ namespace uba
 			extraInfo.Append(TCV(", MsQuic"));
 		if (hasCrypto)
 			extraInfo.Append(TCV(", Encrypted"));
+		#if PLATFORM_WINDOWS
+		if (IsRunningWine() && UsingNativeSockets())
+			extraInfo.Append(TCV(", NativeSock"));
+		#endif
 
 		if (!description.empty())
 			extraInfo.Append(TCV(", ")).Append(description);
diff --git a/Engine/Source/Programs/UnrealBuildAccelerator/Common/Private/UbaNetworkBackendTcp.cpp b/Engine/Source/Programs/UnrealBuildAccelerator/Common/Private/UbaNetworkBackendTcp.cpp
--- a/Engine/Source/Programs/UnrealBuildAccelerator/Common/Private/UbaNetworkBackendTcp.cpp
+++ b/Engine/Source/Programs/UnrealBuildAccelerator/Common/Private/UbaNetworkBackendTcp.cpp
@@ -3,8 +3,10 @@
 #include "UbaNetworkBackendTcp.h"
 #include "UbaConfig.h"
 #include "UbaDefinitions.h"
+#include "UbaDirectoryIterator.h"
 #include "UbaEnvironment.h"
 #include "UbaEvent.h"
+#include "UbaFile.h"
 #include "UbaHash.h"
 #include "UbaPlatform.h"
 #include "UbaStringBuffer.h"
@@ -19,7 +21,6 @@
 #pragma comment (lib, "Ws2_32.lib")
 #pragma comment(lib, "IPHLPAPI.lib") // For GetAdaptersInfo
 #pragma comment(lib, "qwave.lib")
-#define WSABUFMAX 256
 #define WSABUFINIT(a, b) { u32(b), (char*)a }
 #else
 #include "UbaLinuxNetworkWrappers.h"
@@ -29,28 +30,29 @@
 #include <ifaddrs.h>
 #include <poll.h>
 #include <sys/uio.h>
-#define TIMEVAL timeval
 #define SOCKET_ERROR -1
-#define SOCKET int
-#define INVALID_SOCKET -1
-#define SD_BOTH SHUT_RDWR
 #define WSAHOST_NOT_FOUND 0
-#define WSAENOTCONN ENOTCONN
-#define WSAEWOULDBLOCK EWOULDBLOCK
 #define WSAEADDRINUSE EADDRINUSE
-#define closesocket(a) close(a)
 #define addrinfoW addrinfo
 #define GetAddrInfoW getaddrinfo
 #define FreeAddrInfoW freeaddrinfo
 #define WSAGetLastError() errno
 #define strcpy_s(a, b, c) strcpy(a, c)
-#define WSAPOLLFD pollfd
-#define WSAPoll poll
 #define WSABUF iovec
-#define WSABUFMAX IOV_MAX
 #define WSABUFINIT(a, b) { (void*)a, b }
 #endif
 
+#include "UbaSocket.h"
+
+#if PLATFORM_WINDOWS
+#include "UbaSocketWin.h"
+#else
+namespace uba
+{
+	#include "UbaSocket.inl"
+}
+#endif
+
 #define UBA_LOG_SOCKET_ERRORS UBA_DEBUG
 #define UBA_EMULATE_BAD_INTERNET 0
 
@@ -61,20 +63,21 @@
 namespace uba
 {
 	constexpr u32 MaxHeaderSize = 33;
+	typedef void* WsaEvent;
 
 	////////////////////////////////////////////////////////////////////////////////////////////////////
 
-	bool ShutdownSocket(Logger& logger, SOCKET s, const tchar* hint);
-	bool CloseSocket(Logger& logger, SOCKET s, const tchar* hint);
-	bool SetKeepAlive(Logger& logger, SOCKET socket);
-	bool SetBlocking(Logger& logger, SOCKET socket, bool blocking);
-	bool SetTimeout(Logger& logger, SOCKET socket, u32 timeoutMs);
-	bool SetLinger(Logger& logger, SOCKET socket, u32 lingerSeconds);
-	bool SetRecvBuf(Logger& logger, SOCKET socket, u32 windowSize);
-	bool SetSendBuf(Logger& logger, SOCKET socket, u32 windowSize);
-	bool SetSocketPriority(Logger& logger, SOCKET socket);
-	bool DisableNagle(Logger& logger, SOCKET socket);
-	void LogTcpInfo(Logger& logger, SOCKET socket);
+	bool ShutdownSocket(Logger& logger, Socket s, const tchar* hint);
+	bool CloseSocket(Logger& logger, Socket s, const tchar* hint);
+	bool SetKeepAlive(Logger& logger, Socket socket);
+	bool SetBlocking(Logger& logger, Socket socket, bool blocking);
+	bool SetTimeout(Logger& logger, Socket socket, u32 timeoutMs);
+	bool SetLinger(Logger& logger, Socket socket, u32 lingerSeconds);
+	bool SetRecvBuf(Logger& logger, Socket socket, u32 windowSize);
+	bool SetSendBuf(Logger& logger, Socket socket, u32 windowSize);
+	bool SetSocketPriority(Logger& logger, Socket socket);
+	bool DisableNagle(Logger& logger, Socket socket);
+	void LogTcpInfo(Logger& logger, Socket socket);
 
 	////////////////////////////////////////////////////////////////////////////////////////////////////
 
@@ -84,7 +87,7 @@ namespace uba
 		u16 port;
 		ListenConnectedFunc connectedFunc;
 		Event listening;
-		Atomic<SOCKET> socket = INVALID_SOCKET;
+		Atomic<Socket> socket = InvalidSocket;
 		sockaddr_in addr;
 		Thread thread;
 	};
@@ -110,7 +113,7 @@ namespace uba
 
 	struct NetworkBackendTcp::Connection
 	{
-		Connection(Logger& l, SOCKET s) : logger(l), socket(s), ready(EventResetType_Manual) { CreateGuid(uid); }
+		Connection(Logger& l, Socket s) : logger(l), socket(s), ready(EventResetType_Manual) { CreateGuid(uid); }
 		~Connection()
 		{
 			delete threadedSendContext;
@@ -120,7 +123,7 @@ namespace uba
 		}
 
 		Logger& logger;
-		Atomic<SOCKET> socket;
+		Atomic<Socket> socket;
 
 		Event ready;
 		Guid uid;
@@ -175,7 +178,7 @@ namespace uba
 		WSAEVENT wsaEvent = WSA_INVALID_EVENT;
 		#endif
 
-		void CreateWSAEvent(SOCKET s)
+		void CreateWSAEvent(Socket s)
 		{
 			#if PLATFORM_WINDOWS
 			if (IsRunningWine())
@@ -183,7 +186,7 @@ namespace uba
 			WSAEVENT ev = WSACreateEvent();
 			if (ev == WSA_INVALID_EVENT)
 				return;
-			if (WSAEventSelect(s, ev, FD_WRITE | FD_CLOSE) == SOCKET_ERROR)
+			if (WSAEventSelect(ToPlatformSocket(s), ev, FD_WRITE | FD_CLOSE) == SOCKET_ERROR)
 				return;
 			wsaEvent = ev;
 			#endif
@@ -225,9 +228,6 @@ namespace uba
 	HMODULE g_wineDll;
 	#endif
 
-	bool (*GetCongestionAlgorithm)(SOCKET s, char* out, int outCapacity);
-	bool (*SetCongestionAlgorithm)(SOCKET s, char* value, int valuelen);
-
 	////////////////////////////////////////////////////////////////////////////////////////////////////
 
 	bool InitNetwork(Logger& logger)
@@ -241,16 +241,14 @@ namespace uba
 		{
 			g_wineDll = LoadLibrary(TC("UbaWine.dll.so"));
 			if (g_wineDll)
-			{
-				GetCongestionAlgorithm = (decltype(GetCongestionAlgorithm))GetProcAddress(g_wineDll, "GetCongestionAlgorithm");
-				SetCongestionAlgorithm = (decltype(SetCongestionAlgorithm))GetProcAddress(g_wineDll, "SetCongestionAlgorithm");
-			}
+				SocketInitWine(logger, g_wineDll);
+			else// if (FileExists(logger, TC("UbaWine.dll.so")))
+				logger.Warning(TC("Failed to load UbaWine.dll.so (%s)"), LastErrorToText().data);
+
 		}
 #elif PLATFORM_LINUX
 		struct sigaction sa = { { SIG_IGN } };
 		sigaction(SIGPIPE, &sa, NULL); // Needed for broken pipe that can happen if helpers crash
-		GetCongestionAlgorithm = [](SOCKET s, char* out, int outCapacity) { socklen_t len = outCapacity; return getsockopt(s, IPPROTO_TCP, TCP_CONGESTION, out, &len) == 0; };
-		SetCongestionAlgorithm = [](SOCKET s, char* v, int l) { return setsockopt(s, IPPROTO_TCP, TCP_CONGESTION, v, l) == 0; };
 #endif
 		return true;
 	}
@@ -293,7 +291,7 @@ namespace uba
 		m_disableNagle = info.disableNagle;
 
 		#if PLATFORM_WINDOWS
-		m_useOverlappedSend = info.useOverlappedSend && EventIsNative;
+		m_useOverlappedSend = info.useOverlappedSend && EventIsNative && !IsRunningWine();
 		m_iocpWorkerCount = u16(info.iocpWorkerCount);
 		#endif
 
@@ -315,10 +313,10 @@ namespace uba
 		for (auto& conn : m_connections)
 		{
 			SCOPED_CRITICAL_SECTION(conn.shutdownLock, lock2);
-			if (conn.socket == INVALID_SOCKET)
+			if (conn.socket == InvalidSocket)
 				continue;
-			SOCKET s = conn.socket;
-			conn.socket = INVALID_SOCKET;
+			Socket s = conn.socket;
+			conn.socket = InvalidSocket;
 			auto context = conn.threadedSendContext;
 			if (context)
 				context->blockQueueEvent.Set();
@@ -350,7 +348,7 @@ namespace uba
 	{
 		auto& conn = *(Connection*)connection;
 		SCOPED_CRITICAL_SECTION(conn.shutdownLock, lock);
-		if (conn.socket == INVALID_SOCKET)
+		if (conn.socket == InvalidSocket)
 			return;
 		ShutdownSocket(conn.logger, conn.socket, TC("Shutdown"));
 	}
@@ -513,7 +511,7 @@ namespace uba
 				continue;
 			if (!e.listening.IsSet(4000))
 				success = false;
-			if (e.socket == INVALID_SOCKET)
+			if (e.socket == InvalidSocket)
 				success = false;
 			e.listening.Destroy();
 		}
@@ -525,12 +523,12 @@ namespace uba
 		SCOPED_FUTEX(m_listenEntriesLock, lock);
 		for (auto& e : m_listenEntries)
 		{
-			e.socket = INVALID_SOCKET;
-			SOCKET tempSocket = ::socket(AF_INET, SOCK_STREAM, IPPROTO_TCP); // Create a temporary socket just to connect to listen socket to wakeup WSAPoll
-			if (tempSocket == INVALID_SOCKET)
+			e.socket = InvalidSocket;
+			Socket tempSocket = SocketCreate(); // Create a temporary socket just to connect to listen socket to wakeup WSAPoll
+			if (tempSocket == InvalidSocket)
 				continue;
-			::connect(tempSocket, (sockaddr*)&e.addr, sizeof(e.addr));
-			::closesocket(tempSocket);
+			SocketConnect(tempSocket, (sockaddr*)&e.addr, sizeof(e.addr));
+			SocketClose(tempSocket);
 		}
 		for (auto& e : m_listenEntries)
 			e.thread.Wait();
@@ -561,18 +559,17 @@ namespace uba
 		auto addrGuard = MakeGuard([result]() { FreeAddrInfoW(result); });
 
 		// Create a socket for listening to connections
-		SOCKET listenSocket = ::socket(result->ai_family, result->ai_socktype, result->ai_protocol);
-		if (listenSocket == INVALID_SOCKET)
+		Socket listenSocket = SocketCreate(result->ai_family, result->ai_socktype, result->ai_protocol);
+		if (listenSocket == InvalidSocket)
 			return logger.Error(TC("socket failed (%s)"), LastErrorToText(WSAGetLastError()).data);
 
 		auto listenSocketCleanup = MakeGuard([&]() { CloseSocket(logger, listenSocket, TC("listen cleanup")); });
 
-		u32 reuseAddr = 1;
-		if (::setsockopt(listenSocket, SOL_SOCKET, SO_REUSEADDR, (const char*)&reuseAddr, sizeof reuseAddr) == SOCKET_ERROR)
+		if (SocketReuseAddr(listenSocket) == SOCKET_ERROR)
 			return logger.Error(TC("setsockopt SO_REUSEADDR failed (error: %s)"), LastErrorToText(WSAGetLastError()).data);
 
 		// Setup the TCP listening socket
-		res = bind(listenSocket, result->ai_addr, (socklen_t)result->ai_addrlen);
+		res = SocketBind(listenSocket, result->ai_addr, (socklen_t)result->ai_addrlen);
 
 		if (res == SOCKET_ERROR)
 		{
@@ -589,7 +586,7 @@ namespace uba
 
 		addrGuard.Execute();
 
-		res = listen(listenSocket, SOMAXCONN);
+		res = SocketListen(listenSocket);
 		if (res == SOCKET_ERROR)
 			return logger.Error(TC("Listen failed (%s)"), LastErrorToText(WSAGetLastError()).data);
 
@@ -603,39 +600,35 @@ namespace uba
 
 		while (true)
 		{
-			WSAPOLLFD p;
-			p.fd = listenSocket;
-			p.revents = 0;
-			p.events = POLLIN;
+			int revents = 0;
 			int timeoutMs = 5000;
-			int pollRes = WSAPoll(&p, 1, timeoutMs);
+			int pollRes = SocketPoll(listenSocket, PollIN, timeoutMs, &revents);
 
-			if (entry.socket == INVALID_SOCKET)
+			if (entry.socket == InvalidSocket)
 				break;
 
-			if (pollRes == SOCKET_ERROR)
+			if (pollRes < 0)
 			{
 				int lastError = WSAGetLastError();
 				logger.Warning(TC("WSAPoll returned error %s"), LastErrorToText(lastError).data);
 				break;
 			}
 
-			if (!pollRes)
+			if (!pollRes) // timed out
 				continue;
 
-			if (p.revents & POLLNVAL)
+			if (revents & (PollERR | PollHUP | PollNVAL))
 			{
-				logger.Warning(TC("WSAPoll returned successful but with unexpected flags: %u"), p.revents);
+				logger.Warning(TC("WSAPoll returned successful but with unexpected flags: %u"), revents);
 				continue;
 			}
 
 			sockaddr remoteSockAddr = { 0 }; // for TCP/IP
-			socklen_t remoteSockAddrLen = sizeof(remoteSockAddr);
-			SOCKET clientSocket = accept(listenSocket, (sockaddr*)&remoteSockAddr, &remoteSockAddrLen);
+			Socket clientSocket = SocketAccept(listenSocket, &remoteSockAddr, sizeof(remoteSockAddr));
 
-			if (clientSocket == INVALID_SOCKET)
+			if (clientSocket == InvalidSocket)
 			{
-				if (entry.socket != INVALID_SOCKET)
+				if (entry.socket != InvalidSocket)
 					logger.Info(TC("Accept failed with WSA error: %s"), LastErrorToText(WSAGetLastError()).data);
 				break;
 			}
@@ -669,7 +662,7 @@ namespace uba
 			conn.CreateWSAEvent(clientSocket);
 			#if UBA_USE_IOCP
 			if (m_iocpHandle)
-				CreateIoCompletionPort((HANDLE)clientSocket, m_iocpHandle, (ULONG_PTR)&conn, 0);
+				CreateIoCompletionPort((HANDLE)ToPlatformSocket(clientSocket), m_iocpHandle, (ULONG_PTR)&conn, 0);
 			else
 			#endif
 				conn.recvThread.Start([this, connPtr = &conn] { ThreadRecv(*connPtr); return 0; }, TC("UbaTcpRecv"));
@@ -703,7 +696,7 @@ namespace uba
 			RecvCache recvCache;
 
 			bool isFirst = true;
-			while (connection.socket != INVALID_SOCKET)
+			while (connection.socket != InvalidSocket)
 			{
 				void* bodyContext = nullptr;
 				u8* bodyData = nullptr;
@@ -753,11 +746,11 @@ namespace uba
 		}
 
 		SCOPED_CRITICAL_SECTION(connection.shutdownLock, lock2);
-		SOCKET s = connection.socket;
+		Socket s = connection.socket;
 
 		{
 			SCOPED_FUTEX(connection.sendLock, lock);
-			connection.socket = INVALID_SOCKET;
+			connection.socket = InvalidSocket;
 		}
 
 		if (auto context = connection.threadedSendContext)
@@ -771,7 +764,7 @@ namespace uba
 			cb(context, connection.uid, &connection);
 		}
 
-		if (s == INVALID_SOCKET)
+		if (s == InvalidSocket)
 			return;
 		ShutdownSocket(logger, s, TC("ThreadRecv"));
 		CloseSocket(logger, s, TC("ThreadRecv"));
@@ -822,13 +815,8 @@ namespace uba
 	{
 		// Create a socket for connecting to server
 
-		//TODO: Wrap this up in a better function
-#if PLATFORM_WINDOWS
-		SOCKET socketFd = WSASocketW(remoteSocketAddr.sa_family, SOCK_STREAM, IPPROTO_TCP, NULL, 0, WSA_FLAG_OVERLAPPED);
-#else
-		SOCKET socketFd = socket(remoteSocketAddr.sa_family, SOCK_STREAM, IPPROTO_TCP);
-#endif
-		if (socketFd == INVALID_SOCKET)
+		Socket socketFd = SocketCreate(remoteSocketAddr.sa_family, SOCK_STREAM, IPPROTO_TCP);
+		if (socketFd == InvalidSocket)
 			return logger.Error(TC("socket failed (%s)"), LastErrorToText(WSAGetLastError()).data);
 
 		// Create guard in case we fail to connect (will be cancelled further down if we succeed)
@@ -839,51 +827,35 @@ namespace uba
 			return false;
 
 		// Connect to server.
-		int res = ::connect(socketFd, &remoteSocketAddr, sizeof(remoteSocketAddr));
-
-#if PLATFORM_WINDOWS
-		if (res == SOCKET_ERROR)
-			if (WSAGetLastError() != WSAEWOULDBLOCK)
-				return false;
-#else
+		int res = SocketConnect(socketFd, &remoteSocketAddr, sizeof(remoteSocketAddr), false);
 		if (res != 0)
-		{
-			if (errno != EINPROGRESS)
-			{
-				logger.Error(TC("Connect failed (%d: %s)"), WSAGetLastError(), LastErrorToText(WSAGetLastError()).data);
-				return false;
-			}
-		}
-#endif
+			return logger.Error(TC("SocketConnect failed"));
 
 		int timeoutMs = 2000;
 		if (nameHint && (Equals(nameHint, TC("localhost")) || Equals(nameHint, TC("127.0.0.1"))))
 			timeoutMs = 100;
 
-		WSAPOLLFD p;
-		p.fd = socketFd;
-		p.revents = 0;
-		p.events = POLLOUT;
-		int pollRes = WSAPoll(&p, 1, timeoutMs);
+		int revents = 0;
+		int pollRes = SocketPoll(socketFd, PollOUT, timeoutMs, &revents);
 
-		if (pollRes == SOCKET_ERROR)
+		if (pollRes < 0)
 		{
 			int lastError = WSAGetLastError();
 			logger.Warning(TC("WSAPoll returned error %s (%s)"), LastErrorToText(lastError).data, nameHint);
 			return false;
 		}
 
-		u16 validFlags = POLLERR | POLLHUP; // Treat hangup as timeout (since we want retry if that happens). Also treat error as timeout. This is needed for Wine agent to be able to retry
-		if (!pollRes || p.revents & validFlags)
+		u16 validFlags = PollERR | PollHUP; // Treat hangup as timeout (since we want retry if that happens). Also treat error as timeout. This is needed for Wine agent to be able to retry
+		if (!pollRes || revents & validFlags)
 		{
 			if (timedOut)
 				*timedOut = true;
 			return false;
 		}
 
-		if (p.revents & POLLNVAL)
+		if (revents & PollNVAL)
 		{
-			logger.Warning(TC("WSAPoll returned successful but with unexpected flags: %u"), p.revents);
+			logger.Warning(TC("WSAPoll returned successful but with unexpected flags: %u"), revents);
 			return false;
 		}
 
@@ -891,33 +863,8 @@ namespace uba
 		if (!SetBlocking(logger, socketFd, true))
 			return false;
 
-#if !PLATFORM_WINDOWS
-		// Before we send anything even though the
-		// the socket is writable, but let's make sure
-		// the connection is actually valid by getting
-		// information about what we've connected to
-		struct sockaddr_in junk;
-		socklen_t length = sizeof(junk);
-		memset(&junk, 0, sizeof(junk));
-		if (getpeername(socketFd, (struct sockaddr *)&junk, &length) != 0)
-		{
-			if (timedOut)
-				*timedOut = true;
-			return false;
-		}
-
-		int sent = (int)send(socketFd, nullptr, 0, 0);
-		if (sent == SOCKET_ERROR)
-		{
-			if (errno == ECONNREFUSED || errno == EPIPE)
-			{
-				if (timedOut)
-					*timedOut = true;
-				return false;
-			}
+		if (SocketCheckConnect(socketFd, timedOut) != 0)
 			return false;
-		}
-#endif
 
 		if (m_disableNagle && !DisableNagle(logger, socketFd))
 			return false;
@@ -945,7 +892,7 @@ namespace uba
 		conn.CreateWSAEvent(socketFd);
 		#if UBA_USE_IOCP
 		if (m_iocpHandle)
-			CreateIoCompletionPort((HANDLE)socketFd, m_iocpHandle, (ULONG_PTR)&conn, 0);
+			CreateIoCompletionPort((HANDLE)ToPlatformSocket(socketFd), m_iocpHandle, (ULONG_PTR)&conn, 0);
 		else
 		#endif
 			conn.recvThread.Start([this, connPtr = &conn] { ThreadRecv(*connPtr); return 0; }, TC("UbaTcpRecv"));
@@ -1047,7 +994,7 @@ namespace uba
 					first->sent.Set();
 			});
 
-		while (connection.socket != INVALID_SOCKET)
+		while (connection.socket != InvalidSocket)
 		{
 			SCOPED_FUTEX(context.blockQueueLock, queueLock);
 			SendBlock* newFirst = context.firstBlock;
@@ -1079,25 +1026,19 @@ namespace uba
 				}
 			}
 
-			SOCKET socket = connection.socket;
-			if (socket == INVALID_SOCKET)
+			Socket socket = connection.socket;
+			if (socket == InvalidSocket)
 				break;
 
-			WSABUF buf[WSABUFMAX];
-			u32 count = 0;
-			for (SendBlock* it = first; it && count < WSABUFMAX; it = it->next, ++count)
-				buf[count] = WSABUFINIT(it->data, it->dataLen);
-
-			#if PLATFORM_WINDOWS
-			DWORD bytesSent = 0;
-			int res = WSASend(socket, buf, count, &bytesSent, 0, NULL, NULL);
-			#else
-			ssize_t res = writev(socket, buf, count);
-			u64 bytesSent = (u64)res;
-			#endif
+			SendBuf bufs[256];
+			u32 bufCount = 0;
+			for (SendBlock* it = first; it && bufCount < sizeof_array(bufs); it = it->next, ++bufCount)
+				bufs[bufCount] = { it->dataLen, it->data };
 
+			int res = SocketSend(socket, bufs, bufCount);
 			if (res != SOCKET_ERROR)
 			{
+				u32 bytesSent = res;
 				while (bytesSent)
 				{
 					if (bytesSent >= first->dataLen)
@@ -1126,18 +1067,14 @@ namespace uba
 
 	bool NetworkBackendTcp::HandleSendError(Connection& connection, Logger& logger, const tchar* hint)
 	{
-		SOCKET socket = connection.socket;
-		if (socket == INVALID_SOCKET)
+		Socket socket = connection.socket;
+		if (socket == InvalidSocket)
 			return false;
 
-		u32 lastError = WSAGetLastError();
-		bool shouldPoll = lastError == WSAEWOULDBLOCK;
-
-		#if !PLATFORM_WINDOWS
-		if (lastError == EINTR)
+		bool retry = false;
+		bool shouldPoll = SocketShouldPoll(&retry);
+		if (retry)
 			return true;
-		shouldPoll = shouldPoll || lastError == EAGAIN;
-		#endif
 
 		if (!shouldPoll)
 		{
@@ -1173,7 +1110,7 @@ namespace uba
 					#endif
 				}
 				WSANETWORKEVENTS ne{};
-				if (WSAEnumNetworkEvents(socket, connection.wsaEvent, &ne) == SOCKET_ERROR)
+				if (WSAEnumNetworkEvents(ToPlatformSocket(socket), connection.wsaEvent, &ne) == SOCKET_ERROR)
 				{
 					#if UBA_LOG_SOCKET_ERRORS
 					logger.Info(TC("SendSocket - WSAEnumNetworkEvents returned an error for socket %i: %s (%s)"), socket, LastErrorToText(WSAGetLastError()).data, hint);
@@ -1190,11 +1127,8 @@ namespace uba
 		{
 			for (u32 i = 0; i != 2; ++i)
 			{
-				WSAPOLLFD p;
-				p.fd = socket;
-				p.revents = 0;
-				p.events = POLLWRNORM;
-				int res = WSAPoll(&p, 1, int(timeoutMs));
+				int revents = 0;
+				int res = SocketPoll(socket, PollOUT, timeoutMs, &revents);
 				if (!res)
 				{
 					if (i == 0)
@@ -1206,13 +1140,20 @@ namespace uba
 					logger.Info(TC("SendSocket - WSAPoll returned timeout after 10 minutes for socket %i (%s)"), socket, hint);
 					return false;
 				}
-				if (res == SOCKET_ERROR)
+				if (res < 0)
 				{
 					#if UBA_LOG_SOCKET_ERRORS
 					logger.Info(TC("SendSocket - WSAPoll returned an error for socket %i: %s (%s)"), socket, LastErrorToText(WSAGetLastError()).data, hint);
 					#endif
 					return false;
 				}
+				if (revents & (PollERR | PollHUP | PollNVAL))
+				{
+					#if UBA_LOG_SOCKET_ERRORS
+					logger.Info(TC("SendSocket - WSAPoll revents contained error flags for socket %i: %s (%s)"), socket, LastErrorToText(WSAGetLastError()).data, hint);
+					#endif
+					return false;
+				}
 				break;
 			}
 		}
@@ -1277,10 +1218,10 @@ namespace uba
 				SCOPED_FUTEX(connection.sendLock, lock);
 				#endif
 
-				SOCKET socket = connection.socket;
-				if (socket == INVALID_SOCKET)
+				Socket socket = connection.socket;
+				if (socket == InvalidSocket)
 					return false;
-				int res = WSASend(socket, &buf, 1, NULL, 0, &overlapped, NULL);
+				int res = WSASend(ToPlatformSocket(socket), &buf, 1, NULL, 0, &overlapped, NULL);
 
 				#if !UBA_USE_OVERLAPPED_SEND_WITH_LOCK
 				lock.Leave();
@@ -1327,7 +1268,7 @@ namespace uba
 
 				DWORD bytesSent;
 				DWORD flags;
-				if (!WSAGetOverlappedResult(socket, &overlapped, &bytesSent, FALSE, &flags))
+				if (!WSAGetOverlappedResult(ToPlatformSocket(socket), &overlapped, &bytesSent, FALSE, &flags))
 				{
 					#if UBA_LOG_SOCKET_ERRORS
 					logger.Info(TC("WSAGetOverlappedResult - error for socket %i: %s (%s)"), socket, LastErrorToText().data, hint);
@@ -1354,15 +1295,16 @@ namespace uba
 #endif
 
 		SCOPED_FUTEX(connection.sendLock, lock);
-		SOCKET socket = connection.socket;
-		if (socket == INVALID_SOCKET)
+		Socket socket = connection.socket;
+		if (socket == InvalidSocket)
 			return false;
 
 		char* buffer = (char*)b;
 		u64 left = bufferLen;
 		while (left)
 		{
-			int sent = (int)send(socket, (char*)buffer, u32(left), 0);
+			SendBuf buf { u32(left), buffer };
+			int sent = SocketSend(socket, &buf, 1);
 			if (sent != SOCKET_ERROR)
 			{
 				buffer += sent;
@@ -1386,7 +1328,7 @@ namespace uba
 		}
 		#endif
 
-		SOCKET socket = connection.socket;
+		Socket socket = connection.socket;
 
 		u32 retryCount = 4;
 
@@ -1421,7 +1363,7 @@ namespace uba
 
 				DWORD bytesReceived = 0;
 				DWORD flags = 0;
-				read = WSARecv(socket, bufs, bufCount, &bytesReceived, &flags, NULL, NULL);
+				read = WSARecv(ToPlatformSocket(socket), bufs, bufCount, &bytesReceived, &flags, NULL, NULL);
 				if (read == 0)
 					read = int(bytesReceived);
 				else
@@ -1461,7 +1403,7 @@ namespace uba
 					writeCapacity = sizeof(recvCache.bytes);
 				}
 
-				read = (int)recv(socket, (char*)writePos, writeCapacity, 0);
+				read = SocketRecv(socket, writePos, writeCapacity);
 				if (read > 0)
 				{
 					if (writePos == buffer)
@@ -1501,14 +1443,10 @@ namespace uba
 
 			UBA_ASSERT(read == SOCKET_ERROR);
 
-			u32 lastError = WSAGetLastError();
-			bool shouldPoll = lastError == WSAEWOULDBLOCK;
-
-			#if !PLATFORM_WINDOWS
-			if (lastError == EINTR)
+			bool retry = false;
+			bool shouldPoll = SocketShouldPoll(&retry);
+			if (retry)
 				continue;
-			shouldPoll = shouldPoll || lastError == EAGAIN;
-			#endif
 
 			if (shouldPoll)
 			{
@@ -1521,11 +1459,8 @@ namespace uba
 				if (!timeoutMs)
 					timeoutMs = DefaultNetworkReceiveTimeoutSeconds * 1000; // 10 minutes time out
 
-				WSAPOLLFD p;
-				p.fd = socket;
-				p.revents = 0;
-				p.events = POLLRDNORM;
-				int res = WSAPoll(&p, 1, int(timeoutMs));
+				int revents = 0;
+				int res = SocketPoll(socket, PollIN, timeoutMs, &revents);
 				if (!res)
 				{
 					if (connection.recvTimeoutCallback)
@@ -1537,13 +1472,20 @@ namespace uba
 					connection.logger.Info(TC("RecvSocket - WSAPoll returned timeout for socket %i and connection %s after %s (%s%s)"), socket, GuidToString(connection.uid).str, TimeToText(MsToTime(timeoutMs)).str, connection.recvHint, hint);
 					return false;
 				}
-				if (res == SOCKET_ERROR)
+				if (res < 0)
 				{
 					#if UBA_LOG_SOCKET_ERRORS
 					connection.logger.Info(TC("RecvSocket - WSAPoll returned an err-or for socket %i and connection %s: %s (%s%s)"), socket, GuidToString(connection.uid).str, LastErrorToText(WSAGetLastError()).data, connection.recvHint, hint);
 					#endif
 					return false;
 				}
+				if ((revents & PollIN) == 0)
+				{
+					#if UBA_LOG_SOCKET_ERRORS
+					connection.logger.Info(TC("RecvSocket - WSAPoll revents contained error for socket %i and connection %s: %s (%s%s)"), socket, GuidToString(connection.uid).str, LastErrorToText(WSAGetLastError()).data, connection.recvHint, hint);
+					#endif
+					return false;
+				}
 				continue;
 			}
 
@@ -1586,12 +1528,12 @@ namespace uba
 			if (bytesTransferred == 0 || !result)
 			{
 				SCOPED_CRITICAL_SECTION(connection.shutdownLock, lock2);
-				SOCKET s = connection.socket;
-				CancelIoEx((HANDLE)s, overlapped);
+				Socket s = connection.socket;
+				CancelIoEx((HANDLE)ToPlatformSocket(s), overlapped);
 
 				{
 					SCOPED_FUTEX(connection.sendLock, lock);
-					connection.socket = INVALID_SOCKET;
+					connection.socket = InvalidSocket;
 				}
 
 				// We need to marshal shutdown and callback handling out of iocp thread since callback can call things that requires iocp to loop
@@ -1609,7 +1551,7 @@ namespace uba
 							cb(context, connection.uid, &connection);
 						}
 
-						if (s != INVALID_SOCKET)
+						if (s != InvalidSocket)
 						{
 							ShutdownSocket(logger, s, TC("ThreadRecv"));
 							CloseSocket(logger, s, TC("ThreadRecv"));
@@ -1651,7 +1593,7 @@ namespace uba
 				if (!hc(connection.recvContext, connection.uid, connection.header, connection.bodyContext, bodyData, bodySize))
 					continue;
 
-				if (connection.socket == INVALID_SOCKET)
+				if (connection.socket == InvalidSocket)
 					continue;
 
 				if (!bodySize)
@@ -1680,7 +1622,7 @@ namespace uba
 
 				m_totalRecv += connection.bodySize;
 
-				if (connection.socket == INVALID_SOCKET)
+				if (connection.socket == InvalidSocket)
 					continue;
 
 				connection.receivingHeader = true;
@@ -1697,10 +1639,10 @@ namespace uba
 		connection.wsaBuf.buf = (char*)data;
 		connection.wsaBuf.len = dataSize;
 
-		SOCKET socket = connection.socket;
+		Socket socket = connection.socket;
 
 		DWORD flags = 0;
-		int ret = WSARecv(socket, &connection.wsaBuf, 1, NULL, &flags, &connection.overlapped, NULL);
+		int ret = WSARecv(ToPlatformSocket(socket), &connection.wsaBuf, 1, NULL, &flags, &connection.overlapped, NULL);
 		if (ret == SOCKET_ERROR && WSAGetLastError() != WSA_IO_PENDING)
 		{
 			connection.logger.Info(TC("WSARecv failed for socket %i trying to receive %u bytes (%s)"), socket, dataSize, LastErrorToText(WSAGetLastError()).data);
@@ -1729,32 +1671,31 @@ namespace uba
 			for (auto& conn : m_connections)
 			{
 				SCOPED_CRITICAL_SECTION(conn.shutdownLock, lock2);
-				if (conn.socket == INVALID_SOCKET)
+				if (conn.socket == InvalidSocket)
 					continue;
 				SCOPED_FUTEX(conn.sendLock, sendLock);
-				SOCKET s = conn.socket;
-				if (s == INVALID_SOCKET)
+				Socket s = conn.socket;
+				if (s == InvalidSocket)
 					continue;
 				//LogTcpInfo(m_logger, s);
 
 				#if 1
-				socklen_t optlen = sizeof(int);
 				int sendBufSize = 0;
-				int recvBufSize = 0;
-				int prio = 0;
-				char alg[16] = { 0 };
+				SocketGetSendBuf(s, sendBufSize);
 
-				getsockopt(s, SOL_SOCKET, SO_SNDBUF, (char*)&sendBufSize, &optlen);
-				getsockopt(s, SOL_SOCKET, SO_RCVBUF, (char*)&recvBufSize, &optlen);
+				int recvBufSize = 0;
+				SocketGetRecvBuf(s, recvBufSize);
 
-				if (GetCongestionAlgorithm)
-					GetCongestionAlgorithm(s, alg, sizeof(alg));
+				char alg[16] = { 0 };
+				SocketGetCongestionAlgorithm(s, alg, sizeof(alg));
 
+				int prio = 0;
 				#if PLATFORM_LINUX
-				getsockopt(s, SOL_SOCKET, SO_PRIORITY, (char*)&prio, &optlen);
+				socklen_t optlen = sizeof(int);
+				getsockopt(ToPlatformSocket(s), SOL_SOCKET, SO_PRIORITY, (char*)&prio, &optlen);
 				#endif
 
-				m_logger.Info(TC("%i - Send: %s, Recv: %s, Prio %i%hs%hs"), (int)s, BytesToText(sendBufSize).str, BytesToText(recvBufSize).str, prio, *alg ? ", Alg: " : "", alg);
+				m_logger.Info(TC("%llu - Send: %s, Recv: %s, Prio %i%hs%hs"), s.internal, BytesToText(sendBufSize).str, BytesToText(recvBufSize).str, prio, *alg ? ", Alg: " : "", alg);
 				#endif
 
 			}
@@ -1778,13 +1719,12 @@ namespace uba
 		GetTcpAutoTuning = UnixGetTcpAutoTuning;
 		#endif
 
-		if (GetCongestionAlgorithm)
+		if (!IsWindows || IsRunningWine())
 		{
-			SOCKET tempSocket = ::socket(AF_INET, SOCK_STREAM, IPPROTO_TCP);
+			Socket tempSocket = SocketCreate();
 			char alg[16] = { 0 };
-			GetCongestionAlgorithm(tempSocket, alg, sizeof(alg));
-			::closesocket(tempSocket);
-
+			SocketGetCongestionAlgorithm(tempSocket, alg, sizeof(alg));
+			SocketClose(tempSocket);
 			if (strncmp(alg, "bbr", 3) != 0)
 				logger.Info(TC("  TIP: Enable BBR congestion algorithm and fq (Fair Queue) queuing discipline (Current: %hs)"), alg);
 		}
@@ -1948,16 +1888,16 @@ namespace uba
 
 	HttpConnection::HttpConnection()
 	{
-		m_socket = INVALID_SOCKET;
+		m_socket = InvalidSocket.internal;
 		*m_host = 0;
 	}
 
 	HttpConnection::~HttpConnection()
 	{
-		if (m_socket != INVALID_SOCKET)
+		if (m_socket != InvalidSocket.internal)
 		{
 			LoggerWithWriter logger(g_nullLogWriter);
-			CloseSocket(logger, m_socket, TC("HttpDtor"));
+			CloseSocket(logger, {m_socket}, TC("HttpDtor"));
 		}
 	}
 
@@ -1981,10 +1921,10 @@ namespace uba
 		if (protoent == NULL)
 			return logger.Error(TC("HttpConnection: getprotobyname returned null for tcp"));
 
-		SOCKET sock = socket(AF_INET, SOCK_STREAM, protoent->p_proto);
-		if (sock == INVALID_SOCKET)
+		Socket sock = SocketCreate(AF_INET, SOCK_STREAM, protoent->p_proto);
+		if (sock == InvalidSocket)
 			return logger.Error(TC("HttpConnection: socket failed (%s)"), LastErrorToText(WSAGetLastError()).data);
-		auto socketClose = MakeGuard([sock]() { closesocket(sock); });
+		auto socketClose = MakeGuard([sock]() { SocketClose(sock); });
 
 		if (m_connectTimeOutMs)
 			SetTimeout(logger, sock, m_connectTimeOutMs);
@@ -1994,7 +1934,7 @@ namespace uba
 		sockaddr_in.sin_family = AF_INET;
 		sockaddr_in.sin_port = htons(80);
 
-		if (connect(sock, (struct sockaddr*)&sockaddr_in, sizeof(sockaddr_in)) == -1)
+		if (SocketConnect(sock, (struct sockaddr*)&sockaddr_in, sizeof(sockaddr_in)) == -1)
 			return false;// logger.Error(TC("HttpConnection: connect error"));
 
 		socketClose.Cancel();
@@ -2003,7 +1943,7 @@ namespace uba
 			SetTimeout(logger, sock, 10000); // 10 seconds timeout
 
 		strcpy_s(m_host, sizeof_array(m_host), host);
-		m_socket = sock;
+		m_socket = sock.internal;
 		return true;
 	}
 
@@ -2014,12 +1954,12 @@ namespace uba
 
 		if (*m_host)// && _stricmp(m_host, host) != 0)
 		{
-			CloseSocket(logger, m_socket, TC("HttpQuery"));
-			m_socket = INVALID_SOCKET;
+			CloseSocket(logger, {m_socket}, TC("HttpQuery"));
+			m_socket = InvalidSocket.internal;
 			*m_host = 0;
 		}
 
-		if (m_socket == INVALID_SOCKET)
+		if (m_socket == InvalidSocket.internal)
 			if (!Connect(logger, host))
 				return false;
 
@@ -2028,11 +1968,13 @@ namespace uba
 		if (requestLen >= sizeof_array(request))
 			return logger.Error(TC("STACK BUFFER TOO SMALL!"));
 
-		SetTimeout(logger, m_socket, timeoutMs);
+		SetTimeout(logger, {m_socket}, timeoutMs);
 
 		int totalBytesSent = 0;
-		while (totalBytesSent < requestLen) {
-			int bytesSent = send(m_socket, request + totalBytesSent, requestLen - totalBytesSent, 0);
+		while (totalBytesSent < requestLen)
+		{
+			SendBuf buf { u32(requestLen - totalBytesSent), request + totalBytesSent };
+			int bytesSent = SocketSend({m_socket}, &buf, 1);
 			if (bytesSent == -1)
 				return logger.Error(TC("HttpConnection: socket send error (%hs)"), host);
 			totalBytesSent += bytesSent;
@@ -2048,7 +1990,7 @@ namespace uba
 		u32 readPos = 0;
 		char buf[4*1024];
 		int bytesRead = 0;
-		while ((bytesRead = recv(m_socket, buf + readPos, sizeof(buf) - readPos, 0)) > 0)
+		while ((bytesRead = SocketRecv({m_socket}, buf + readPos, sizeof(buf) - readPos)) > 0)
 			readPos += bytesRead;
 
 		if (bytesRead == SOCKET_ERROR)
@@ -2094,164 +2036,77 @@ namespace uba
 
 	////////////////////////////////////////////////////////////////////////////////////////////////////
 
-	bool ShutdownSocket(Logger& logger, SOCKET s, const tchar* hint)
+	bool ShutdownSocket(Logger& logger, Socket s, const tchar* hint)
 	{
-		if (s == INVALID_SOCKET)
-			return true;
-		if (shutdown(s, SD_BOTH) != SOCKET_ERROR)
-			return true;
-		if (WSAGetLastError() == WSAENOTCONN)
+		if (SocketShutdown(s) == 0)
 			return true;
-		logger.Info(TC("Failed to shutdown socket %llu in %s (%s)"), u64(s), hint, LastErrorToText(WSAGetLastError()).data);
+		logger.Info(TC("Failed to shutdown socket %llu in %s (%s)"), s.internal, hint, LastErrorToText(WSAGetLastError()).data);
 		return false;
 	}
 
-	bool CloseSocket(Logger& logger, SOCKET s, const tchar* hint)
+	bool CloseSocket(Logger& logger, Socket s, const tchar* hint)
 	{
-		if (s == INVALID_SOCKET)
+		if (s == InvalidSocket)
 			return true;
-		if (closesocket(s) != SOCKET_ERROR)
+		if (SocketClose(s) != SOCKET_ERROR)
 			return true;
-		logger.Info(TC("Failed to close socket %llu in %s (%s)"), u64(s), hint, LastErrorToText(WSAGetLastError()).data);
+		logger.Info(TC("Failed to close socket %llu in %s (%s)"), s.internal, hint, LastErrorToText(WSAGetLastError()).data);
 		return false;
 	}
 
-	bool SetBlocking(Logger& logger, SOCKET socket, bool blocking)
+	bool SetBlocking(Logger& logger, Socket socket, bool blocking)
 	{
-#if PLATFORM_WINDOWS
-		u_long value = blocking ? 0 : 1;
-		if (ioctlsocket(socket, FIONBIO, &value) == SOCKET_ERROR)
-			return logger.Error(TC("Setting non blocking socket failed (error: %s)"), LastErrorToText(WSAGetLastError()).data);
-#else
-		int flags = fcntl(socket, F_GETFL, 0);
-		if (flags == -1) return false;
-		flags = blocking ? (flags & ~O_NONBLOCK) : (flags | O_NONBLOCK);
-		if (fcntl(socket, F_SETFL, flags) != 0)
+		if (SocketSetBlocking(socket, blocking) != 0)
 			return logger.Error(TC("Setting non blocking socket failed (error: %s)"), LastErrorToText(WSAGetLastError()).data);
-#endif
 		return true;
 	}
 
-	bool DisableNagle(Logger& logger, SOCKET socket)
+	bool DisableNagle(Logger& logger, Socket socket)
 	{
-#if !PLATFORM_MAC
-		u32 value = 1;
-		if (setsockopt(socket, IPPROTO_TCP, TCP_NODELAY, (const char*)&value, sizeof(value)) == SOCKET_ERROR)
+		if (SocketSetNoDelay(socket) != 0)
 			return logger.Error(TC("setsockopt TCP_NODELAY error: (error: %s)"), LastErrorToText(WSAGetLastError()).data);
-#endif
 		return true;
 	}
 
-	bool SetTimeout(Logger& logger, SOCKET socket, u32 timeoutMs)
+	bool SetTimeout(Logger& logger, Socket socket, u32 timeoutMs)
 	{
-		#if PLATFORM_WINDOWS
-		DWORD timeout = timeoutMs;
-		#else
-		struct timeval timeout;
-		timeout.tv_sec = timeoutMs / 1000;
-		timeout.tv_usec = (timeoutMs % 1000)*1000;
-		#endif
-
-		if (setsockopt(socket, SOL_SOCKET, SO_SNDTIMEO, (const char*)&timeout, sizeof timeout) == SOCKET_ERROR)
+		if (SocketSetTimeout(socket, timeoutMs) != 0)
 			return logger.Error(TC("setsockopt SO_SNDTIMEO error: (error: %s)"), LastErrorToText(WSAGetLastError()).data);
-		if (setsockopt(socket, SOL_SOCKET, SO_RCVTIMEO, (const char*)&timeout, sizeof timeout) == SOCKET_ERROR)
-			return logger.Error(TC("setsockopt SO_RCVTIMEO error: (error: %s)"), LastErrorToText(WSAGetLastError()).data);
 		return true;
 	}
 
-	bool SetLinger(Logger& logger, SOCKET socket, u32 lingerSeconds)
+	bool SetLinger(Logger& logger, Socket socket, u32 lingerSeconds)
 	{
-#if PLATFORM_MAC // Mac does not seem to automatically close the socket if the process crashes
-		struct linger so_linger;
-		so_linger.l_onoff = 1;
-		so_linger.l_linger = u16(lingerSeconds);
-
-		if (setsockopt(socket, SOL_SOCKET, SO_LINGER, (const char*)&so_linger, sizeof(so_linger)) < 0)
+		if (SocketSetLinger(socket, lingerSeconds) != 0)
 			return logger.Error(TC("setsockopt SO_LINGER error: (error: %s)"), LastErrorToText(WSAGetLastError()).data);
-#endif
 		return true;
 	}
 
-	bool SetRecvBuf(Logger& logger, SOCKET socket, u32 windowSize)
+	bool SetRecvBuf(Logger& logger, Socket socket, u32 windowSize)
 	{
-		int optname = SO_RCVBUF;
-		#if PLATFORM_LINUX
-		optname = SO_RCVBUFFORCE;
-		#endif
-		if (setsockopt(socket, SOL_SOCKET, optname, (const char*)&windowSize, sizeof(windowSize)) == SOCKET_ERROR)
-			return logger.Error(TC("setsockopt SO_RCVBUF error: (error: %s)"), LastErrorToText(WSAGetLastError()).data);
-		return true;
+		if (SocketSetRecvBuf(socket, windowSize) == 0)
+			return true;
+		return logger.Error(TC("setsockopt SO_RCVBUF error: (error: %s)"), LastErrorToText(WSAGetLastError()).data);
 	}
 
-	bool SetSendBuf(Logger& logger, SOCKET socket, u32 windowSize)
+	bool SetSendBuf(Logger& logger, Socket socket, u32 windowSize)
 	{
-		int optname = SO_SNDBUF;
-		#if PLATFORM_LINUX
-		optname = SO_SNDBUFFORCE;
-		#endif
-		if (setsockopt(socket, SOL_SOCKET, optname, (const char*)&windowSize, sizeof(windowSize)) == SOCKET_ERROR)
-			return logger.Error(TC("setsockopt SO_SNDBUF error: (error: %s)"), LastErrorToText(WSAGetLastError()).data);
-		return true;
+		if (SocketSetSendBuf(socket, windowSize) == 0)
+			return true;
+		return logger.Error(TC("setsockopt SO_SNDBUF error: (error: %s)"), LastErrorToText(WSAGetLastError()).data);
 	}
 
-	bool SetSocketPriority(Logger& logger, SOCKET socket)
+	bool SetSocketPriority(Logger& logger, Socket socket)
 	{
-#if PLATFORM_WINDOWS
-		static HANDLE s_qosHandle = [](Logger& logger)
-			{
-				HANDLE h = INVALID_HANDLE_VALUE;
-				QOS_VERSION ver = {1, 0};
-				if (!QOSCreateHandle(&ver, &h))
-					logger.Error(TC("QOSCreateHandle failed"));
-				return h;
-			}(logger);
-		if (s_qosHandle == INVALID_HANDLE_VALUE)
-			return false;
-		QOS_FLOWID flowId = 0;
-		if (!QOSAddSocketToFlow(s_qosHandle, socket, NULL, QOSTrafficTypeExcellentEffort, QOS_NON_ADAPTIVE_FLOW, &flowId))
-			return logger.Error(TC("QOSAddSocketToFlow failed for socket %i (%s)"), socket, LastErrorToText().data);
-#elif PLATFORM_LINUX
-		int prio = 6;
-		if (setsockopt(socket, SOL_SOCKET, SO_PRIORITY, &prio, sizeof(prio)) == -1)
-			return logger.Error(TC("setsockopt SO_PRIORITY failed (%s)"), LastErrorToText(WSAGetLastError()).data);
-		//uint8_t tos = 46 << 2; // EF = 0b101110xx
-		//if (setsockopt(socket, IPPROTO_IP, IP_TOS, &tos, sizeof(tos)) == -1)
-		//	return logger.Error(TC("setsockopt IP_TOS failed (%s)"), LastErrorToText(WSAGetLastError()).data);
-#endif
-		return true;
+		if (SocketSetPriority(socket) == 0)
+			return true;
+		return logger.Error(TC("setsockopt SO_PRIORITY failed (%s)"), LastErrorToText(WSAGetLastError()).data);
 	}
 
-	bool SetKeepAlive(Logger& logger, SOCKET socket) // This will make sure that WSAPoll exits when the network cable is pulled
+	bool SetKeepAlive(Logger& logger, Socket socket) // This will make sure that WSAPoll exits when the network cable is pulled
 	{
-		
-		int keepAliveTime = KeepAliveIdleSeconds; // Time of idle before it starts sending keep alive
-		int keepAliveTimeInterval = KeepAliveIntervalSeconds; // Time between each check (10 failures mean disconnect)
-
-
-		u32 value = 1;
-		if (setsockopt(socket, SOL_SOCKET, SO_KEEPALIVE, (const char*)&value, sizeof(value)) == SOCKET_ERROR)
-			return logger.Error(TC("setsockopt SO_KEEPALIVE (error: %s)"), LastErrorToText(WSAGetLastError()).data);
-
-#if PLATFORM_WINDOWS	
-		struct tcp_keepalive kaSettings;
-		DWORD bytesReturned;
-		kaSettings.onoff = 1;
-		kaSettings.keepalivetime = keepAliveTime * 1000;
-		kaSettings.keepaliveinterval = keepAliveTimeInterval * 1000; 
-		if (WSAIoctl(socket, SIO_KEEPALIVE_VALS, &kaSettings, sizeof(kaSettings), NULL, 0, &bytesReturned, NULL, NULL) == SOCKET_ERROR)
-			return logger.Error(TC("WSAIoctl SIO_KEEPALIVE_VALS (error: %s)"), LastErrorToText(WSAGetLastError()).data);
-#elif PLATFORM_LINUX
-		if (setsockopt(socket, IPPROTO_TCP, TCP_KEEPIDLE, &keepAliveTime, sizeof(int)) < 0)
-			return logger.Error(TC("setsockopt TCP_KEEPIDLE (error: %s)"), LastErrorToText(WSAGetLastError()).data);
-		if (setsockopt(socket, IPPROTO_TCP, TCP_KEEPINTVL, &keepAliveTimeInterval, sizeof(int)) < 0)
-			return logger.Error(TC("setsockopt TCP_KEEPINTVL (error: %s)"), LastErrorToText(WSAGetLastError()).data);
-		int keepAliveProbes = KeepAliveProbeCount; // Number of tests before timing out
-		if (setsockopt(socket, IPPROTO_TCP, TCP_KEEPCNT, &keepAliveProbes, sizeof(int)) < 0)
-			return logger.Error(TC("setsockopt TCP_KEEPCNT (error: %s)"), LastErrorToText(WSAGetLastError()).data);
-#else // PLATFORM_MAC
-		if (setsockopt(socket, IPPROTO_TCP, TCP_KEEPALIVE, &keepAliveTime, sizeof(int)) < 0)
-			return logger.Error(TC("setsockopt TCP_KEEPALIVE (error: %s)"), LastErrorToText(WSAGetLastError()).data);
-#endif
+		if (SocketSetKeepAlive(socket, KeepAliveIdleSeconds, KeepAliveIntervalSeconds) != 0)
+			return logger.Error(TC("SocketSetKeepAlive (error: %s)"), LastErrorToText(WSAGetLastError()).data);
 		return true;
 	}
 
@@ -2360,7 +2215,7 @@ namespace uba
 #endif
 
 
-	void LogTcpInfo(Logger& logger, SOCKET socket)
+	void LogTcpInfo(Logger& logger, Socket socket)
 	{
 	#if PLATFORM_WINDOWS
 		#if 0
@@ -2369,7 +2224,7 @@ namespace uba
 			static HMODULE wineDll = LoadLibrary(TC("UbaWine.dll.so"));
 			if (!wineDll)
 				return false;
-			using GetLinuxTcpInfoFunc = int WINAPI(SOCKET, void*, int*);
+			using GetLinuxTcpInfoFunc = int WINAPI(Socket, void*, int*);
 			static auto GetLinuxTcpInfo = (GetLinuxTcpInfoFunc*)GetProcAddress(wineDll, "GetLinuxTcpInfo");
 			if (!GetLinuxTcpInfo)
 				return logger.Error(TC("GetLinuxTcpInfo not found in UbaWine dll"));
diff --git a/Engine/Source/Programs/UnrealBuildAccelerator/Common/Public/UbaNetworkBackend.h b/Engine/Source/Programs/UnrealBuildAccelerator/Common/Public/UbaNetworkBackend.h
--- a/Engine/Source/Programs/UnrealBuildAccelerator/Common/Public/UbaNetworkBackend.h
+++ b/Engine/Source/Programs/UnrealBuildAccelerator/Common/Public/UbaNetworkBackend.h
@@ -111,12 +111,8 @@ namespace uba
 		bool Connect(Logger& logger, const char* host);
 
 		char m_host[256];
-		#if PLATFORM_WINDOWS
 		u64 m_socket;
 		bool m_wsaInitDone = false;
-		#else
-		int m_socket;
-		#endif
 		u32 m_connectTimeOutMs = 0;
 	};
 }
diff --git a/Engine/Source/Programs/UnrealBuildAccelerator/Common/Public/UbaSocket.h b/Engine/Source/Programs/UnrealBuildAccelerator/Common/Public/UbaSocket.h
new file mode 100644
--- /dev/null
+++ b/Engine/Source/Programs/UnrealBuildAccelerator/Common/Public/UbaSocket.h
@@ -0,0 +1,67 @@
+// Copyright Epic Games, Inc. All Rights Reserved.
+
+#pragma once
+
+struct sockaddr;
+
+namespace uba
+{
+	////////////////////////////////////////////////////////////////////////////////////////////////////
+
+	struct Socket { u64 internal = ~0ull; bool operator==(const Socket& o) const { return internal == o.internal; } };
+	inline constexpr Socket InvalidSocket;
+
+	#if PLATFORM_WINDOWS
+	inline SOCKET ToPlatformSocket(Socket socket) { return (SOCKET)socket.internal; }
+	inline Socket FromPlatformSocket(SOCKET socket) { return { u64(socket) }; }
+	#else
+	inline int ToPlatformSocket(Socket socket) { return (int)socket.internal; }
+	inline Socket FromPlatformSocket(int socket) { return { u64(socket) }; }
+	#endif
+
+	////////////////////////////////////////////////////////////////////////////////////////////////////
+
+	inline constexpr int PollIN      = 0x0001;
+	inline constexpr int PollPRI     = 0x0002;
+	inline constexpr int PollOUT     = 0x0004;
+	inline constexpr int PollERR     = 0x0008;
+	inline constexpr int PollHUP     = 0x0010;
+	inline constexpr int PollNVAL    = 0x0020;
+	inline constexpr int PollRDNORM  = 0x0040;
+	inline constexpr int PollRDBAND  = 0x0080;
+	inline constexpr int PollWRNORM  = 0x0100;
+	inline constexpr int PollWRBAND  = 0x0200;
+
+	struct SendBuf
+	{
+		u32 len;
+		const void* data;
+	};
+
+	Socket SocketCreate(int family = AF_INET, int type = SOCK_STREAM, int protocol = IPPROTO_TCP);
+	int SocketConnect(Socket s, const sockaddr* addr, u32 addrSize, bool isBlocking = true);
+	int SocketClose(Socket s);
+	int SocketSetBlocking(Socket s, bool blocking);
+	int SocketReuseAddr(Socket s);
+	int SocketBind(Socket s, const sockaddr* addr, u32 addrSize);
+	int SocketListen(Socket s);
+	int SocketSetKeepAlive(Socket s, int keepAliveTime, int keepAliveTimeInterval);
+	int SocketSetNoDelay(Socket s);
+	int SocketSetLinger(Socket s, int lingerSeconds);
+	int SocketSetTimeout(Socket s, int timeoutMs);
+	int SocketSetRecvBuf(Socket s, int size);
+	int SocketSetSendBuf(Socket s, int size);
+	int SocketSetPriority(Socket s);
+	int SocketPoll(Socket s, int events, int timeoutMs, int* outRevents);
+	Socket SocketAccept(Socket s, sockaddr* outAddr, int addrLen);
+	int SocketRecv(Socket s, void* data, int len);
+	int SocketSend(Socket s, SendBuf* bufs, int bufCount);
+	bool SocketShouldPoll(bool* outRetry);
+	int SocketCheckConnect(Socket s, bool* outTimedOut);
+	int SocketShutdown(Socket s);
+	int SocketGetCongestionAlgorithm(Socket s, char* out, int outCapacity);
+	int SocketGetRecvBuf(Socket s, int& outSize);
+	int SocketGetSendBuf(Socket s, int& outSize);
+
+	////////////////////////////////////////////////////////////////////////////////////////////////////
+}
diff --git a/Engine/Source/Programs/UnrealBuildAccelerator/Common/Public/UbaSocket.inl b/Engine/Source/Programs/UnrealBuildAccelerator/Common/Public/UbaSocket.inl
new file mode 100644
--- /dev/null
+++ b/Engine/Source/Programs/UnrealBuildAccelerator/Common/Public/UbaSocket.inl
@@ -0,0 +1,336 @@
+// Copyright Epic Games, Inc. All Rights Reserved.
+
+#pragma once
+
+////////////////////////////////////////////////////////////////////////////////////////////////////
+
+Socket SocketCreate(int family, int type, int protocol)
+{
+#if PLATFORM_WINDOWS
+	return FromPlatformSocket(WSASocketW(family, type, protocol, NULL, 0, WSA_FLAG_OVERLAPPED));
+#else
+	return FromPlatformSocket(::socket(family, type, protocol));
+#endif
+}
+
+int SocketConnect(Socket s, const sockaddr* addr, u32 addrSize, bool isBlocking)
+{
+	int res = ::connect(ToPlatformSocket(s), addr, addrSize);
+	if (isBlocking)
+		return res;
+#if PLATFORM_WINDOWS
+	if (WSAGetLastError() != WSAEWOULDBLOCK)
+		return -1;
+#else
+	if (errno != EINPROGRESS)
+		return -1;
+#endif
+	return 0;
+}
+
+int SocketClose(Socket s)
+{
+#if PLATFORM_WINDOWS
+	return closesocket(ToPlatformSocket(s));
+#else
+	return close(ToPlatformSocket(s));
+#endif
+}
+
+int SocketSetBlocking(Socket s, bool blocking)
+{
+#if PLATFORM_WINDOWS
+	u_long value = blocking ? 0 : 1;
+	return ioctlsocket(ToPlatformSocket(s), FIONBIO, &value);
+#else
+	int flags = fcntl(ToPlatformSocket(s), F_GETFL, 0);
+	if (flags == -1) return -1;
+	flags = blocking ? (flags & ~O_NONBLOCK) : (flags | O_NONBLOCK);
+	if (fcntl(ToPlatformSocket(s), F_SETFL, flags) != 0)
+		return -1;
+	return 0;
+#endif
+}
+
+int SocketReuseAddr(Socket s)
+{
+	u32 reuseAddr = 1;
+	return ::setsockopt(ToPlatformSocket(s), SOL_SOCKET, SO_REUSEADDR, (const char*)&reuseAddr, sizeof reuseAddr);
+}
+
+int SocketBind(Socket s, const sockaddr* addr, u32 addrSize)
+{
+	return bind(ToPlatformSocket(s), addr, (socklen_t)addrSize);
+}
+int SocketListen(Socket s)
+{
+	return listen(ToPlatformSocket(s), SOMAXCONN);
+}
+int SocketSetKeepAlive(Socket s, int keepAliveTime, int keepAliveTimeInterval)
+{
+	u32 value = 1;
+	int res = setsockopt(ToPlatformSocket(s), SOL_SOCKET, SO_KEEPALIVE, (const char*)&value, sizeof(value));
+	if (res != 0)
+		return res;
+
+#if PLATFORM_WINDOWS	
+	struct tcp_keepalive kaSettings;
+	DWORD bytesReturned;
+	kaSettings.onoff = 1;
+	kaSettings.keepalivetime = keepAliveTime * 1000;
+	kaSettings.keepaliveinterval = keepAliveTimeInterval * 1000; 
+	res = WSAIoctl(ToPlatformSocket(s), SIO_KEEPALIVE_VALS, &kaSettings, sizeof(kaSettings), NULL, 0, &bytesReturned, NULL, NULL);
+	if (res != 0)
+		return res;
+#elif PLATFORM_LINUX
+	if (setsockopt(ToPlatformSocket(s), IPPROTO_TCP, TCP_KEEPIDLE, &keepAliveTime, sizeof(int)) < 0)
+		return -1;
+	if (setsockopt(ToPlatformSocket(s), IPPROTO_TCP, TCP_KEEPINTVL, &keepAliveTimeInterval, sizeof(int)) < 0)
+		return -1;
+	int keepAliveProbes = KeepAliveProbeCount; // Number of tests before timing out
+	if (setsockopt(ToPlatformSocket(s), IPPROTO_TCP, TCP_KEEPCNT, &keepAliveProbes, sizeof(int)) < 0)
+		return -1;
+#else // PLATFORM_MAC
+	if (setsockopt(s, IPPROTO_TCP, TCP_KEEPALIVE, &keepAliveTime, sizeof(int)) < 0)
+		return logger.Error(TC("setsockopt TCP_KEEPALIVE (error: %s)"), LastErrorToText(WSAGetLastError()).data);
+#endif
+	return 0;
+}
+
+int SocketSetNoDelay(Socket s)
+{
+#if !PLATFORM_MAC
+	u32 value = 1;
+	return setsockopt(ToPlatformSocket(s), IPPROTO_TCP, TCP_NODELAY, (const char*)&value, sizeof(value));
+#else
+	return 0;
+#endif
+}
+
+int SocketSetLinger(Socket s, int lingerSeconds)
+{
+#if PLATFORM_MAC // Mac does not seem to automatically close the socket if the process crashes
+	struct linger so_linger;
+	so_linger.l_onoff = 1;
+	so_linger.l_linger = u16(lingerSeconds);
+	return setsockopt(ToPlatformSocket(s), SOL_SOCKET, SO_LINGER, (const char*)&so_linger, sizeof(so_linger));
+#endif
+	return 0;
+}
+
+int SocketSetTimeout(Socket s, int timeoutMs)
+{
+#if PLATFORM_WINDOWS
+	DWORD timeout = timeoutMs;
+#else
+	struct timeval timeout;
+	timeout.tv_sec = timeoutMs / 1000;
+	timeout.tv_usec = (timeoutMs % 1000)*1000;
+#endif
+	if (setsockopt(ToPlatformSocket(s), SOL_SOCKET, SO_SNDTIMEO, (const char*)&timeout, sizeof timeout) != 0)
+		return -1;
+	if (setsockopt(ToPlatformSocket(s), SOL_SOCKET, SO_RCVTIMEO, (const char*)&timeout, sizeof timeout) != 0)
+		return -1;
+	return 0;
+}
+
+int SocketSetRecvBuf(Socket s, int size)
+{
+	int optname = SO_RCVBUF;
+#if PLATFORM_LINUX
+	optname = SO_RCVBUFFORCE;
+#endif
+	return setsockopt(ToPlatformSocket(s), SOL_SOCKET, optname, (const char*)&size, sizeof(size));
+}
+
+int SocketSetSendBuf(Socket s, int size)
+{
+	int optname = SO_SNDBUF;
+#if PLATFORM_LINUX
+	optname = SO_SNDBUFFORCE;
+#endif
+	return setsockopt(ToPlatformSocket(s), SOL_SOCKET, optname, (const char*)&size, sizeof(size));
+}
+
+int SocketSetPriority(Socket s)
+{
+#if PLATFORM_WINDOWS
+	static HANDLE s_qosHandle = []()
+		{
+			HANDLE h = INVALID_HANDLE_VALUE;
+			QOS_VERSION ver = {1, 0};
+			QOSCreateHandle(&ver, &h);
+			return h;
+		}();
+	if (s_qosHandle == INVALID_HANDLE_VALUE)
+		return -1;
+	QOS_FLOWID flowId = 0;
+	if (!QOSAddSocketToFlow(s_qosHandle, ToPlatformSocket(s), NULL, QOSTrafficTypeExcellentEffort, QOS_NON_ADAPTIVE_FLOW, &flowId))
+		return -1;
+	return 0;
+#elif PLATFORM_LINUX
+	int prio = 6;
+	return setsockopt(ToPlatformSocket(s), SOL_SOCKET, SO_PRIORITY, &prio, sizeof(prio));
+	//uint8_t tos = 46 << 2; // EF = 0b101110xx
+	//if (setsockopt(socket, IPPROTO_IP, IP_TOS, &tos, sizeof(tos)) == -1)
+	//	return logger.Error(TC("setsockopt IP_TOS failed (%s)"), LastErrorToText(WSAGetLastError()).data);
+#else
+	return 0;
+#endif
+}
+
+int SocketPoll(Socket s, int events, int timeoutMs, int* outRevents)
+{
+	int realEvents = 0;
+	realEvents |= (events & PollIN) ? POLLIN : 0;
+	realEvents |= (events & PollPRI) ? POLLPRI : 0;
+	realEvents |= (events & PollOUT) ? POLLOUT : 0;
+	realEvents |= (events & PollRDNORM) ? POLLRDNORM : 0;
+	realEvents |= (events & PollWRNORM) ? POLLWRNORM : 0;
+	realEvents |= (events & PollWRBAND) ? POLLWRBAND : 0;
+
+#if PLATFORM_WINDOWS
+	WSAPOLLFD p;
+	p.fd = ToPlatformSocket(s);
+	p.revents = 0;
+	p.events = (SHORT)realEvents;
+	int res = WSAPoll(&p, 1, timeoutMs);
+#else
+	pollfd p;
+	p.fd = ToPlatformSocket(s);
+	p.revents = 0;
+	p.events = (short)realEvents;
+	int res = poll(&p, 1, timeoutMs);
+
+#endif
+	if (!outRevents)
+		return res;
+	realEvents = 0;
+	realEvents |= (p.revents & POLLIN) ? PollIN : 0;
+	realEvents |= (p.revents & POLLPRI) ? PollPRI : 0;
+	realEvents |= (p.revents & POLLOUT) ? PollOUT : 0;
+	realEvents |= (p.revents & POLLERR) ? PollERR : 0;
+	realEvents |= (p.revents & POLLHUP) ? PollHUP : 0;
+	realEvents |= (p.revents & POLLNVAL) ? PollNVAL : 0;
+	realEvents |= (p.revents & POLLRDNORM) ? PollRDNORM : 0;
+	realEvents |= (p.revents & POLLWRNORM) ? PollWRNORM : 0;
+	realEvents |= (p.revents & POLLWRBAND) ? PollWRBAND : 0;
+	*outRevents = realEvents;
+	return res;
+}
+
+Socket SocketAccept(Socket s, sockaddr* outAddr, int addrLen)
+{
+	socklen_t remoteSockAddrLen = addrLen;
+	return FromPlatformSocket(accept(ToPlatformSocket(s), outAddr, &remoteSockAddrLen));
+}
+
+int SocketRecv(Socket s, void* data, int len)
+{
+	return (int)recv(ToPlatformSocket(s), (char*)data, len, 0);
+}
+
+int SocketSend(Socket s, SendBuf* bufs, int bufCount)
+{
+	if (bufCount == 1)
+		return (int)send(ToPlatformSocket(s), (char*)bufs[0].data, bufs[0].len, 0);
+#if PLATFORM_WINDOWS
+	DWORD bytesSent = 0;
+	int res = WSASend(ToPlatformSocket(s), (LPWSABUF)bufs, bufCount, &bytesSent, 0, NULL, NULL);
+	if (!res)
+		res = int(bytesSent);
+	return res;
+#else
+	iovec bufs2[256];
+	for (int i=0; i!=bufCount; ++i)
+		bufs2[i] = { (void*)bufs[i].data, bufs[i].len };
+	return writev(ToPlatformSocket(s), bufs2, bufCount);
+#endif
+
+}
+
+bool SocketShouldPoll(bool* outRetry)
+{
+	*outRetry = false;
+#if PLATFORM_WINDOWS
+	return WSAGetLastError() == WSAEWOULDBLOCK;
+#else
+	if (errno == EINTR)
+		*outRetry = true;
+	return errno == EWOULDBLOCK || errno == EAGAIN;
+#endif
+}
+
+int SocketCheckConnect(Socket s, bool* outTimedOut)
+{
+#if !PLATFORM_WINDOWS
+	// Before we send anything even though the
+	// the socket is writable, but let's make sure
+	// the connection is actually valid by getting
+	// information about what we've connected to
+	struct sockaddr_in junk = { 0 };
+	socklen_t length = sizeof(junk);
+	if (getpeername(ToPlatformSocket(s), (struct sockaddr *)&junk, &length) != 0)
+	{
+		if (outTimedOut)
+			*outTimedOut = true;
+		return -1;
+	}
+
+	int sent = (int)send(ToPlatformSocket(s), nullptr, 0, 0);
+	if (sent < 0)
+	{
+		if (errno == ECONNREFUSED || errno == EPIPE)
+		{
+			if (outTimedOut)
+				*outTimedOut = true;
+			return -1;
+		}
+		return -1;
+	}
+#endif
+	return 0;
+}
+
+int SocketShutdown(Socket s)
+{
+	if (s == InvalidSocket)
+		return 0;
+#if PLATFORM_WINDOWS
+	if (shutdown(ToPlatformSocket(s), SD_BOTH) != SOCKET_ERROR)
+		return 0;
+	if (WSAGetLastError() == WSAENOTCONN)
+		return 0;
+#else
+	if (shutdown(ToPlatformSocket(s), SHUT_RDWR) == 0)
+		return 0;
+	if (errno == ENOTCONN)
+		return 0;
+#endif
+	return -1;
+}
+
+int SocketGetCongestionAlgorithm(Socket s, char* out, int outCapacity)
+{
+#if PLATFORM_WINDOWS
+	strcpy_s(out, outCapacity, "Default");
+	return 0;
+#else
+	socklen_t len = outCapacity;
+	return getsockopt(ToPlatformSocket(s), IPPROTO_TCP, TCP_CONGESTION, out, &len);
+#endif
+}
+
+int SocketGetRecvBuf(Socket s, int& outSize)
+{
+	socklen_t optlen = sizeof(int);
+	return getsockopt(ToPlatformSocket(s), SOL_SOCKET, SO_RCVBUF, (char*)&outSize, &optlen);
+}
+
+int SocketGetSendBuf(Socket s, int& outSize)
+{
+	socklen_t optlen = sizeof(int);
+	return getsockopt(ToPlatformSocket(s), SOL_SOCKET, SO_SNDBUF, (char*)&outSize, &optlen);
+}
+
+////////////////////////////////////////////////////////////////////////////////////////////////////
diff --git a/Engine/Source/Programs/UnrealBuildAccelerator/Common/Public/UbaSocketWin.h b/Engine/Source/Programs/UnrealBuildAccelerator/Common/Public/UbaSocketWin.h
new file mode 100644
--- /dev/null
+++ b/Engine/Source/Programs/UnrealBuildAccelerator/Common/Public/UbaSocketWin.h
@@ -0,0 +1,96 @@
+// Copyright Epic Games, Inc. All Rights Reserved.
+
+#pragma once
+
+namespace uba
+{
+	////////////////////////////////////////////////////////////////////////////////////////////////////
+
+	namespace impl
+	{
+		#include "UbaSocket.inl"
+	}
+
+	////////////////////////////////////////////////////////////////////////////////////////////////////
+
+	#define SOCKET_FUNCTIONS \
+		SOCKET_FUNCTION(Create) \
+		SOCKET_FUNCTION(Connect) \
+		SOCKET_FUNCTION(Close) \
+		SOCKET_FUNCTION(SetBlocking) \
+		SOCKET_FUNCTION(ReuseAddr) \
+		SOCKET_FUNCTION(Bind) \
+		SOCKET_FUNCTION(Listen) \
+		SOCKET_FUNCTION(SetKeepAlive) \
+		SOCKET_FUNCTION(Poll) \
+		SOCKET_FUNCTION(Accept) \
+		SOCKET_FUNCTION(SetNoDelay) \
+		SOCKET_FUNCTION(SetLinger) \
+		SOCKET_FUNCTION(SetTimeout) \
+		SOCKET_FUNCTION(SetRecvBuf) \
+		SOCKET_FUNCTION(SetSendBuf) \
+		SOCKET_FUNCTION(SetPriority) \
+		SOCKET_FUNCTION(Recv) \
+		SOCKET_FUNCTION(Send) \
+		SOCKET_FUNCTION(ShouldPoll) \
+		SOCKET_FUNCTION(CheckConnect) \
+		SOCKET_FUNCTION(Shutdown) \
+		SOCKET_FUNCTION(GetCongestionAlgorithm) \
+		SOCKET_FUNCTION(GetRecvBuf) \
+		SOCKET_FUNCTION(GetSendBuf) \
+
+	#define SOCKET_FUNCTION(func) decltype(impl::Socket##func)* g_socket##func##Func = impl::Socket##func;
+	SOCKET_FUNCTIONS
+	#undef SOCKET_FUNCTION
+
+	////////////////////////////////////////////////////////////////////////////////////////////////////
+
+	Socket SocketCreate(int family, int type, int protocol) { return g_socketCreateFunc(family, type, protocol); }
+	int SocketConnect(Socket s, const sockaddr* addr, u32 addrSize, bool isBlocking) { return g_socketConnectFunc(s, addr, addrSize, isBlocking); }
+	int SocketClose(Socket s) { return g_socketCloseFunc(s); }
+	int SocketSetBlocking(Socket s, bool blocking) { return g_socketSetBlockingFunc(s, blocking); }
+	int SocketReuseAddr(Socket s) { return g_socketReuseAddrFunc(s); }
+	int SocketBind(Socket s, const sockaddr* addr, u32 addrSize) { return g_socketBindFunc(s, addr, addrSize); }
+	int SocketListen(Socket s) { return g_socketListenFunc(s); }
+	int SocketSetKeepAlive(Socket s, int keepAliveTime, int keepAliveTimeInterval) { return g_socketSetKeepAliveFunc(s, keepAliveTime, keepAliveTimeInterval); }
+	int SocketSetNoDelay(Socket s) { return g_socketSetNoDelayFunc(s); }
+	int SocketSetLinger(Socket s, int lingerSeconds) { return g_socketSetLingerFunc(s, lingerSeconds); }
+	int SocketSetTimeout(Socket s, int timeoutMs) { return g_socketSetTimeoutFunc(s, timeoutMs); }
+	int SocketSetRecvBuf(Socket s, int size) { return g_socketSetRecvBufFunc(s, size); }
+	int SocketSetSendBuf(Socket s, int size) { return g_socketSetSendBufFunc(s, size); }
+	int SocketPoll(Socket s, int events, int timeoutMs, int* outRevents) { return g_socketPollFunc(s, events, timeoutMs, outRevents); }
+	Socket SocketAccept(Socket s, sockaddr* outAddr, int addrLen) { return g_socketAcceptFunc(s, outAddr, addrLen); }
+	int SocketRecv(Socket s, void* data, int len) { return g_socketRecvFunc(s, data, len); }
+	int SocketSend(Socket s, SendBuf* bufs, int bufCount) { return g_socketSendFunc(s, bufs, bufCount); }
+	int SocketCheckConnect(Socket s, bool* outTimedOut) { return g_socketCheckConnectFunc(s, outTimedOut); }
+	bool SocketShouldPoll(bool* outRetry) { return g_socketShouldPollFunc(outRetry); }
+	int SocketShutdown(Socket s) { return g_socketShutdownFunc(s); }
+	int SocketGetCongestionAlgorithm(Socket s, char* out, int outCapacity) { return g_socketGetCongestionAlgorithmFunc(s, out, outCapacity); } 
+	int SocketGetRecvBuf(Socket s, int& outSize) { return g_socketSetRecvBufFunc(s, outSize); }
+	int SocketGetSendBuf(Socket s, int& outSize) { return g_socketSetSendBufFunc(s, outSize); }
+	int SocketSetPriority(Socket s) { return g_socketSetPriorityFunc(s); };
+
+	////////////////////////////////////////////////////////////////////////////////////////////////////
+
+	bool SocketInitWine(Logger& logger, HMODULE wineDll)
+	{
+		#define SOCKET_FUNCTION(func) \
+			auto func##Func = (decltype(Socket##func)*)GetProcAddress(wineDll, "Socket" #func); \
+			if (!func##Func) return logger.Warning(TC("Socket" #func " is not exported from wine dll"));
+		SOCKET_FUNCTIONS
+		#undef SOCKET_FUNCTION
+
+		#define SOCKET_FUNCTION(func) g_socket##func##Func = func##Func;
+		SOCKET_FUNCTIONS
+		#undef SOCKET_FUNCTION
+
+		return true;
+	}
+
+	bool UsingNativeSockets()
+	{
+		return g_socketCreateFunc != impl::SocketCreate;
+	}
+
+	////////////////////////////////////////////////////////////////////////////////////////////////////
+}
diff --git a/Engine/Source/Programs/UnrealBuildAccelerator/Wine/Build.sh b/Engine/Source/Programs/UnrealBuildAccelerator/Wine/Build.sh
--- a/Engine/Source/Programs/UnrealBuildAccelerator/Wine/Build.sh
+++ b/Engine/Source/Programs/UnrealBuildAccelerator/Wine/Build.sh
@@ -1,4 +1,3 @@
 #!/bin/sh
 cd "$(dirname "$0")"
-winegcc -shared -m64 Windows.c Linux.c Windows.def -o ../../../../Binaries/Win64/UnrealBuildAccelerator/x64/UbaWine.dll.so
-
+winegcc -shared -m64 Windows.cpp Linux.cpp Windows.def -o ../../../../Binaries/Win64/UnrealBuildAccelerator/x64/UbaWine.dll.so
\ No newline at end of file
diff --git a/Engine/Source/Programs/UnrealBuildAccelerator/Wine/Linux.c b/Engine/Source/Programs/UnrealBuildAccelerator/Wine/Linux.c
deleted file mode 100644
diff --git a/Engine/Source/Programs/UnrealBuildAccelerator/Wine/Linux.cpp b/Engine/Source/Programs/UnrealBuildAccelerator/Wine/Linux.cpp
new file mode 100644
--- /dev/null
+++ b/Engine/Source/Programs/UnrealBuildAccelerator/Wine/Linux.cpp
@@ -0,0 +1,46 @@
+// Copyright Epic Games, Inc. All Rights Reserved.
+
+
+#include <netinet/tcp.h>
+#include <netinet/in.h>
+#include <errno.h>
+#include <stdio.h>
+#include <stdlib.h>
+//#include "../Common/Public/UbaLinuxNetworkWrappers.h"
+
+int UnixGetsockopt(int fd, int level, int optname, char* optval, int* optlen)
+{
+	return getsockopt(fd, level, optname, optval, (socklen_t*)optlen);
+}
+
+int UnixSetsockopt(int fd, int level, int optname, char* optval, int optlen)
+{
+	return setsockopt(fd, level, optname, optval, (socklen_t)optlen);
+}
+
+int UnixGetTcpInfo(int fd, void* buf, int* len)
+{
+	socklen_t l = *len;
+	if (getsockopt(fd, IPPROTO_TCP, TCP_INFO, buf, &l) == -1)
+		return -errno;
+
+	*len = (int)l;
+	return 0;
+}
+
+#define PLATFORM_WINDOWS 0
+#define PLATFORM_MAC 0
+#define PLATFORM_LINUX 1
+
+#include <unistd.h>
+#include <fcntl.h>
+#include <poll.h>
+#include <sys/uio.h>
+#include "../Core/Public/UbaBase.h"
+#include "../Common/Public/UbaSocket.h"
+
+namespace uba
+{
+	static constexpr u32 KeepAliveProbeCount = 10;
+	#include "../Common/Public/UbaSocket.inl"
+}
diff --git a/Engine/Source/Programs/UnrealBuildAccelerator/Wine/Windows.c b/Engine/Source/Programs/UnrealBuildAccelerator/Wine/Windows.c
deleted file mode 100644
diff --git a/Engine/Source/Programs/UnrealBuildAccelerator/Wine/Windows.cpp b/Engine/Source/Programs/UnrealBuildAccelerator/Wine/Windows.cpp
new file mode 100644
--- /dev/null
+++ b/Engine/Source/Programs/UnrealBuildAccelerator/Wine/Windows.cpp
@@ -0,0 +1,156 @@
+// Copyright Epic Games, Inc. All Rights Reserved.
+
+#include <winsock2.h>
+#include <stdint.h>
+#include <unistd.h>
+#include <stdbool.h>
+#include <stdio.h>
+#include "../Core/Public/UbaBase.h"
+#include "../Common/Public/UbaSocket.h"
+
+// NOTE, README
+// Must use "WINAPI" on function and add Windows.def
+
+
+extern "C" NTSYSAPI NTSTATUS CDECL wine_server_handle_to_fd( HANDLE handle, unsigned int access, int *unix_fd, unsigned int *options );
+
+extern int UnixGetsockopt(int fd, int level, int optname, char* optval, int* optlen);
+extern int UnixSetsockopt(int fd, int level, int optname, char* optval, int optlen);
+extern int UnixGetTcpInfo(int fd, void* buf, int* len);
+//extern bool UnixGetTcpAutoTuning(int* outReadMin, int* outReadDefault, int* outReadMax, int* outWriteMin, int* outWriteDefault, int* outWriteMax);
+
+extern "C" INT WINAPI GetLinuxTcpInfo(SOCKET s, void* buf, INT* len)
+{
+	int fd = -1;
+	if (wine_server_handle_to_fd((HANDLE)s, 0, &fd, NULL) != 0)
+		return -1;
+	int res = UnixGetTcpInfo(fd, buf, len);
+	close(fd);
+	return res;
+}
+
+int UNIX_IPPROTO_TCP = 6;
+int TCP_CONGESTION = 13;
+
+extern "C" bool WINAPI GetCongestionAlgorithm(SOCKET s, char* out, int outCapacity)
+{
+	int fd = -1;
+	if (wine_server_handle_to_fd((HANDLE)s, 0, &fd, NULL) != 0)
+		return -1;
+	int res = UnixGetsockopt(fd, UNIX_IPPROTO_TCP, TCP_CONGESTION, out, &outCapacity);
+	close(fd);
+	return res == 0;
+}
+
+extern "C" bool WINAPI SetCongestionAlgorithm(SOCKET s, char* value, int valuelen)
+{
+	int fd = -1;
+	if (wine_server_handle_to_fd((HANDLE)s, 0, &fd, NULL) != 0)
+		return -1;
+	int res = UnixSetsockopt(fd, UNIX_IPPROTO_TCP, TCP_CONGESTION, value, valuelen);
+	close(fd);
+	return res == 0;
+}
+
+
+extern "C" bool WINAPI GetTcpAutoTuning(int* outReadMin, int* outReadDefault, int* outReadMax, int* outWriteMin, int* outWriteDefault, int* outWriteMax)
+{
+	return false;//UnixGetTcpAutoTuning(outReadMin, outReadDefault, outReadMax, outWriteMin, outWriteDefault, outWriteMax);
+}
+
+extern "C" uba::Socket WINAPI SocketCreate(int family, int type, int protocol)
+{
+	return uba::SocketCreate(family, type, protocol);
+}
+extern "C" int WINAPI SocketConnect(uba::Socket s, const sockaddr* addr, unsigned int addrSize, bool isBlocking)
+{
+	return uba::SocketConnect(s, addr, addrSize, isBlocking);
+}
+extern "C" int WINAPI SocketClose(uba::Socket s)
+{
+	return uba::SocketClose(s);
+}
+extern "C" int WINAPI SocketSetBlocking(uba::Socket s, bool blocking)
+{
+	return uba::SocketSetBlocking(s, blocking);
+}
+extern "C" int WINAPI SocketReuseAddr(uba::Socket s)
+{
+	return uba::SocketReuseAddr(s);
+}
+extern "C" int WINAPI SocketBind(uba::Socket s, const sockaddr* addr, unsigned int addrSize)
+{
+	return uba::SocketBind(s, addr, addrSize);
+}
+extern "C" int WINAPI SocketListen(uba::Socket s)
+{
+	return uba::SocketListen(s);
+}
+extern "C" int WINAPI SocketSetKeepAlive(uba::Socket s, int keepAliveTime, int keepAliveTimeInterval)
+{
+	return uba::SocketSetKeepAlive(s, keepAliveTime, keepAliveTimeInterval);
+}
+extern "C" int WINAPI SocketSetNoDelay(uba::Socket s)
+{
+	return uba::SocketSetNoDelay(s);
+}
+extern "C" int WINAPI SocketSetLinger(uba::Socket s, int lingerSeconds)
+{
+	return uba::SocketSetLinger(s, lingerSeconds);
+}
+extern "C" int WINAPI SocketSetTimeout(uba::Socket s, int timeoutMs)
+{
+	return uba::SocketSetTimeout(s, timeoutMs);
+}
+extern "C" int WINAPI SocketSetRecvBuf(uba::Socket s, int size)
+{
+	return uba::SocketSetRecvBuf(s, size);
+}
+extern "C" int WINAPI SocketSetSendBuf(uba::Socket s, int size)
+{
+	return uba::SocketSetSendBuf(s, size);
+}
+extern "C" int WINAPI SocketSetPriority(uba::Socket s)
+{
+	return uba::SocketSetPriority(s);
+}
+extern "C" int WINAPI SocketPoll(uba::Socket s, int events, int timeoutMs, int* outRevents)
+{
+	return uba::SocketPoll(s, events, timeoutMs, outRevents);
+}
+extern "C" uba::Socket WINAPI SocketAccept(uba::Socket s, sockaddr* outAddr, int addrLen)
+{
+	return uba::SocketAccept(s, outAddr, addrLen);
+}
+extern "C" int WINAPI SocketRecv(uba::Socket s, void* data, int len)
+{
+	return uba::SocketRecv(s, data, len);
+}
+extern "C" int WINAPI SocketSend(uba::Socket s, uba::SendBuf* bufs, int bufCount)
+{
+	return uba::SocketSend(s, bufs, bufCount);
+}
+extern "C" bool WINAPI SocketShouldPoll(bool* outRetry)
+{
+	return uba::SocketShouldPoll(outRetry);
+}
+extern "C" int WINAPI SocketCheckConnect(uba::Socket s, bool* outTimedOut)
+{
+	return uba::SocketCheckConnect(s, outTimedOut);
+}
+extern "C" int WINAPI SocketShutdown(uba::Socket s)
+{
+	return uba::SocketShutdown(s);
+}
+extern "C" int WINAPI SocketGetCongestionAlgorithm(uba::Socket s, char* out, int outCapacity)
+{
+	return uba::SocketGetCongestionAlgorithm(s, out, outCapacity);
+}
+extern "C" int WINAPI SocketGetRecvBuf(uba::Socket s, int& outSize)
+{
+	return uba::SocketGetRecvBuf(s, outSize);
+}
+extern "C" int WINAPI SocketGetSendBuf(uba::Socket s, int& outSize)
+{
+	return uba::SocketGetSendBuf(s, outSize);
+}
diff --git a/Engine/Source/Programs/UnrealBuildAccelerator/Wine/Windows.def b/Engine/Source/Programs/UnrealBuildAccelerator/Wine/Windows.def
--- a/Engine/Source/Programs/UnrealBuildAccelerator/Wine/Windows.def
+++ b/Engine/Source/Programs/UnrealBuildAccelerator/Wine/Windows.def
@@ -3,4 +3,28 @@ EXPORTS
 	GetLinuxTcpInfo
 	GetCongestionAlgorithm
 	SetCongestionAlgorithm
-	GetTcpAutoTuning
\ No newline at end of file
+	GetTcpAutoTuning
+	SocketCreate
+	SocketConnect
+	SocketClose
+	SocketSetBlocking
+	SocketReuseAddr
+	SocketBind
+	SocketListen
+	SocketSetKeepAlive
+	SocketPoll
+	SocketAccept
+	SocketSetNoDelay
+	SocketSetLinger
+	SocketSetTimeout
+	SocketSetRecvBuf
+	SocketSetSendBuf
+	SocketSetPriority
+	SocketRecv
+	SocketSend
+	SocketCheckConnect
+	SocketShouldPoll
+	SocketShutdown
+	SocketGetCongestionAlgorithm
+	SocketGetRecvBuf
+	SocketGetSendBuf
