[UBA]
* Renamed source code related to building UbaWine to be prefixed with UbaWine

diff --git a/Engine/Source/Programs/UnrealBuildAccelerator/Wine/Build.sh b/Engine/Source/Programs/UnrealBuildAccelerator/Wine/Build.sh
--- a/Engine/Source/Programs/UnrealBuildAccelerator/Wine/Build.sh
+++ b/Engine/Source/Programs/UnrealBuildAccelerator/Wine/Build.sh
@@ -1,3 +1,3 @@
 #!/bin/sh
 cd "$(dirname "$0")"
-winegcc -shared -m64 Windows.cpp Linux.cpp Windows.def -o ../../../../Binaries/Win64/UnrealBuildAccelerator/x64/UbaWine.dll.so
\ No newline at end of file
+winegcc -shared -m64 UbaWineWindows.cpp UbaWineLinux.cpp UbaWineWindows.def -o ../../../../Binaries/Win64/UnrealBuildAccelerator/x64/UbaWine.dll.so
\ No newline at end of file
diff --git a/Engine/Source/Programs/UnrealBuildAccelerator/Wine/Linux.cpp b/Engine/Source/Programs/UnrealBuildAccelerator/Wine/Linux.cpp
deleted file mode 100644
diff --git a/Engine/Source/Programs/UnrealBuildAccelerator/Wine/UbaWineLinux.cpp b/Engine/Source/Programs/UnrealBuildAccelerator/Wine/UbaWineLinux.cpp
new file mode 100644
--- /dev/null
+++ b/Engine/Source/Programs/UnrealBuildAccelerator/Wine/UbaWineLinux.cpp
@@ -0,0 +1,46 @@
+// Copyright Epic Games, Inc. All Rights Reserved.
+
+
+#include <netinet/tcp.h>
+#include <netinet/in.h>
+#include <errno.h>
+#include <stdio.h>
+#include <stdlib.h>
+//#include "../Common/Public/UbaLinuxNetworkWrappers.h"
+
+int UnixGetsockopt(int fd, int level, int optname, char* optval, int* optlen)
+{
+	return getsockopt(fd, level, optname, optval, (socklen_t*)optlen);
+}
+
+int UnixSetsockopt(int fd, int level, int optname, char* optval, int optlen)
+{
+	return setsockopt(fd, level, optname, optval, (socklen_t)optlen);
+}
+
+int UnixGetTcpInfo(int fd, void* buf, int* len)
+{
+	socklen_t l = *len;
+	if (getsockopt(fd, IPPROTO_TCP, TCP_INFO, buf, &l) == -1)
+		return -errno;
+
+	*len = (int)l;
+	return 0;
+}
+
+#define PLATFORM_WINDOWS 0
+#define PLATFORM_MAC 0
+#define PLATFORM_LINUX 1
+
+#include <unistd.h>
+#include <fcntl.h>
+#include <poll.h>
+#include <sys/uio.h>
+#include "../Core/Public/UbaBase.h"
+#include "../Common/Public/UbaSocket.h"
+
+namespace uba
+{
+	static constexpr u32 KeepAliveProbeCount = 10;
+	#include "../Common/Public/UbaSocket.inl"
+}
diff --git a/Engine/Source/Programs/UnrealBuildAccelerator/Wine/UbaWineWindows.cpp b/Engine/Source/Programs/UnrealBuildAccelerator/Wine/UbaWineWindows.cpp
new file mode 100644
--- /dev/null
+++ b/Engine/Source/Programs/UnrealBuildAccelerator/Wine/UbaWineWindows.cpp
@@ -0,0 +1,156 @@
+// Copyright Epic Games, Inc. All Rights Reserved.
+
+#include <winsock2.h>
+#include <stdint.h>
+#include <unistd.h>
+#include <stdbool.h>
+#include <stdio.h>
+#include "../Core/Public/UbaBase.h"
+#include "../Common/Public/UbaSocket.h"
+
+// NOTE, README
+// Must use "WINAPI" on function and add Windows.def
+
+
+extern "C" NTSYSAPI NTSTATUS CDECL wine_server_handle_to_fd( HANDLE handle, unsigned int access, int *unix_fd, unsigned int *options );
+
+extern int UnixGetsockopt(int fd, int level, int optname, char* optval, int* optlen);
+extern int UnixSetsockopt(int fd, int level, int optname, char* optval, int optlen);
+extern int UnixGetTcpInfo(int fd, void* buf, int* len);
+//extern bool UnixGetTcpAutoTuning(int* outReadMin, int* outReadDefault, int* outReadMax, int* outWriteMin, int* outWriteDefault, int* outWriteMax);
+
+extern "C" INT WINAPI GetLinuxTcpInfo(SOCKET s, void* buf, INT* len)
+{
+	int fd = -1;
+	if (wine_server_handle_to_fd((HANDLE)s, 0, &fd, NULL) != 0)
+		return -1;
+	int res = UnixGetTcpInfo(fd, buf, len);
+	close(fd);
+	return res;
+}
+
+int UNIX_IPPROTO_TCP = 6;
+int TCP_CONGESTION = 13;
+
+extern "C" bool WINAPI GetCongestionAlgorithm(SOCKET s, char* out, int outCapacity)
+{
+	int fd = -1;
+	if (wine_server_handle_to_fd((HANDLE)s, 0, &fd, NULL) != 0)
+		return -1;
+	int res = UnixGetsockopt(fd, UNIX_IPPROTO_TCP, TCP_CONGESTION, out, &outCapacity);
+	close(fd);
+	return res == 0;
+}
+
+extern "C" bool WINAPI SetCongestionAlgorithm(SOCKET s, char* value, int valuelen)
+{
+	int fd = -1;
+	if (wine_server_handle_to_fd((HANDLE)s, 0, &fd, NULL) != 0)
+		return -1;
+	int res = UnixSetsockopt(fd, UNIX_IPPROTO_TCP, TCP_CONGESTION, value, valuelen);
+	close(fd);
+	return res == 0;
+}
+
+
+extern "C" bool WINAPI GetTcpAutoTuning(int* outReadMin, int* outReadDefault, int* outReadMax, int* outWriteMin, int* outWriteDefault, int* outWriteMax)
+{
+	return false;//UnixGetTcpAutoTuning(outReadMin, outReadDefault, outReadMax, outWriteMin, outWriteDefault, outWriteMax);
+}
+
+extern "C" uba::Socket WINAPI SocketCreate(int family, int type, int protocol)
+{
+	return uba::SocketCreate(family, type, protocol);
+}
+extern "C" int WINAPI SocketConnect(uba::Socket s, const sockaddr* addr, unsigned int addrSize, bool isBlocking)
+{
+	return uba::SocketConnect(s, addr, addrSize, isBlocking);
+}
+extern "C" int WINAPI SocketClose(uba::Socket s)
+{
+	return uba::SocketClose(s);
+}
+extern "C" int WINAPI SocketSetBlocking(uba::Socket s, bool blocking)
+{
+	return uba::SocketSetBlocking(s, blocking);
+}
+extern "C" int WINAPI SocketReuseAddr(uba::Socket s)
+{
+	return uba::SocketReuseAddr(s);
+}
+extern "C" int WINAPI SocketBind(uba::Socket s, const sockaddr* addr, unsigned int addrSize)
+{
+	return uba::SocketBind(s, addr, addrSize);
+}
+extern "C" int WINAPI SocketListen(uba::Socket s)
+{
+	return uba::SocketListen(s);
+}
+extern "C" int WINAPI SocketSetKeepAlive(uba::Socket s, int keepAliveTime, int keepAliveTimeInterval)
+{
+	return uba::SocketSetKeepAlive(s, keepAliveTime, keepAliveTimeInterval);
+}
+extern "C" int WINAPI SocketSetNoDelay(uba::Socket s)
+{
+	return uba::SocketSetNoDelay(s);
+}
+extern "C" int WINAPI SocketSetLinger(uba::Socket s, int lingerSeconds)
+{
+	return uba::SocketSetLinger(s, lingerSeconds);
+}
+extern "C" int WINAPI SocketSetTimeout(uba::Socket s, int timeoutMs)
+{
+	return uba::SocketSetTimeout(s, timeoutMs);
+}
+extern "C" int WINAPI SocketSetRecvBuf(uba::Socket s, int size)
+{
+	return uba::SocketSetRecvBuf(s, size);
+}
+extern "C" int WINAPI SocketSetSendBuf(uba::Socket s, int size)
+{
+	return uba::SocketSetSendBuf(s, size);
+}
+extern "C" int WINAPI SocketSetPriority(uba::Socket s)
+{
+	return uba::SocketSetPriority(s);
+}
+extern "C" int WINAPI SocketPoll(uba::Socket s, int events, int timeoutMs, int* outRevents)
+{
+	return uba::SocketPoll(s, events, timeoutMs, outRevents);
+}
+extern "C" uba::Socket WINAPI SocketAccept(uba::Socket s, sockaddr* outAddr, int addrLen)
+{
+	return uba::SocketAccept(s, outAddr, addrLen);
+}
+extern "C" int WINAPI SocketRecv(uba::Socket s, void* data, int len)
+{
+	return uba::SocketRecv(s, data, len);
+}
+extern "C" int WINAPI SocketSend(uba::Socket s, uba::SendBuf* bufs, int bufCount)
+{
+	return uba::SocketSend(s, bufs, bufCount);
+}
+extern "C" bool WINAPI SocketShouldPoll(bool* outRetry)
+{
+	return uba::SocketShouldPoll(outRetry);
+}
+extern "C" int WINAPI SocketCheckConnect(uba::Socket s, bool* outTimedOut)
+{
+	return uba::SocketCheckConnect(s, outTimedOut);
+}
+extern "C" int WINAPI SocketShutdown(uba::Socket s)
+{
+	return uba::SocketShutdown(s);
+}
+extern "C" int WINAPI SocketGetCongestionAlgorithm(uba::Socket s, char* out, int outCapacity)
+{
+	return uba::SocketGetCongestionAlgorithm(s, out, outCapacity);
+}
+extern "C" int WINAPI SocketGetRecvBuf(uba::Socket s, int& outSize)
+{
+	return uba::SocketGetRecvBuf(s, outSize);
+}
+extern "C" int WINAPI SocketGetSendBuf(uba::Socket s, int& outSize)
+{
+	return uba::SocketGetSendBuf(s, outSize);
+}
diff --git a/Engine/Source/Programs/UnrealBuildAccelerator/Wine/UbaWineWindows.def b/Engine/Source/Programs/UnrealBuildAccelerator/Wine/UbaWineWindows.def
new file mode 100644
--- /dev/null
+++ b/Engine/Source/Programs/UnrealBuildAccelerator/Wine/UbaWineWindows.def
@@ -0,0 +1,30 @@
+LIBRARY libtcpinfo.dll
+EXPORTS
+	GetLinuxTcpInfo
+	GetCongestionAlgorithm
+	SetCongestionAlgorithm
+	GetTcpAutoTuning
+	SocketCreate
+	SocketConnect
+	SocketClose
+	SocketSetBlocking
+	SocketReuseAddr
+	SocketBind
+	SocketListen
+	SocketSetKeepAlive
+	SocketPoll
+	SocketAccept
+	SocketSetNoDelay
+	SocketSetLinger
+	SocketSetTimeout
+	SocketSetRecvBuf
+	SocketSetSendBuf
+	SocketSetPriority
+	SocketRecv
+	SocketSend
+	SocketCheckConnect
+	SocketShouldPoll
+	SocketShutdown
+	SocketGetCongestionAlgorithm
+	SocketGetRecvBuf
+	SocketGetSendBuf
diff --git a/Engine/Source/Programs/UnrealBuildAccelerator/Wine/Windows.cpp b/Engine/Source/Programs/UnrealBuildAccelerator/Wine/Windows.cpp
deleted file mode 100644
diff --git a/Engine/Source/Programs/UnrealBuildAccelerator/Wine/Windows.def b/Engine/Source/Programs/UnrealBuildAccelerator/Wine/Windows.def
deleted file mode 100644
