[UBA]
* Fixed macos compile error

diff --git a/Engine/Source/Programs/UnrealBuildAccelerator/Common/Public/UbaSocket.inl b/Engine/Source/Programs/UnrealBuildAccelerator/Common/Public/UbaSocket.inl
--- a/Engine/Source/Programs/UnrealBuildAccelerator/Common/Public/UbaSocket.inl
+++ b/Engine/Source/Programs/UnrealBuildAccelerator/Common/Public/UbaSocket.inl
@@ -92,7 +92,7 @@ int SocketSetKeepAlive(Socket s, int keepAliveTime, int keepAliveTimeInterval)
 		return -1;
 #else // PLATFORM_MAC
 	if (setsockopt(s, IPPROTO_TCP, TCP_KEEPALIVE, &keepAliveTime, sizeof(int)) < 0)
-		return logger.Error(TC("setsockopt TCP_KEEPALIVE (error: %s)"), LastErrorToText(WSAGetLastError()).data);
+		return -1;
 #endif
 	return 0;
 }
@@ -312,12 +312,12 @@ int SocketShutdown(Socket s)
 
 int SocketGetCongestionAlgorithm(Socket s, char* out, int outCapacity)
 {
-#if PLATFORM_WINDOWS
-	strcpy_s(out, outCapacity, "Default");
-	return 0;
-#else
+#if PLATFORM_LINUX
 	socklen_t len = outCapacity;
 	return getsockopt(ToPlatformSocket(s), IPPROTO_TCP, TCP_CONGESTION, out, &len);
+#else
+	strcpy_s(out, outCapacity, "Default");
+	return 0;
 #endif
 }
 
