[UBA]
* Fixed another macos compile error

diff --git a/Engine/Source/Programs/UnrealBuildAccelerator/Common/Public/UbaSocket.inl b/Engine/Source/Programs/UnrealBuildAccelerator/Common/Public/UbaSocket.inl
--- a/Engine/Source/Programs/UnrealBuildAccelerator/Common/Public/UbaSocket.inl
+++ b/Engine/Source/Programs/UnrealBuildAccelerator/Common/Public/UbaSocket.inl
@@ -91,7 +91,7 @@ int SocketSetKeepAlive(Socket s, int keepAliveTime, int keepAliveTimeInterval)
 	if (setsockopt(ToPlatformSocket(s), IPPROTO_TCP, TCP_KEEPCNT, &keepAliveProbes, sizeof(int)) < 0)
 		return -1;
 #else // PLATFORM_MAC
-	if (setsockopt(s, IPPROTO_TCP, TCP_KEEPALIVE, &keepAliveTime, sizeof(int)) < 0)
+	if (setsockopt(ToPlatformSocket(s), IPPROTO_TCP, TCP_KEEPALIVE, &keepAliveTime, sizeof(int)) < 0)
 		return -1;
 #endif
 	return 0;
