[UBA]
* Promoted IndexContainer to UbaWorkManager.h to be able to use it in other places than in UbaStorageUtils.cpp

diff --git a/Engine/Source/Programs/UnrealBuildAccelerator/Common/Private/UbaStorageUtils.cpp b/Engine/Source/Programs/UnrealBuildAccelerator/Common/Private/UbaStorageUtils.cpp
--- a/Engine/Source/Programs/UnrealBuildAccelerator/Common/Private/UbaStorageUtils.cpp
+++ b/Engine/Source/Programs/UnrealBuildAccelerator/Common/Private/UbaStorageUtils.cpp
@@ -34,25 +34,6 @@ namespace uba
 
 	////////////////////////////////////////////////////////////////////////////////////////////////////
 
-	struct IndexContainer
-	{
-		struct iterator
-		{
-			iterator() : index(0) {}
-			iterator(u32 i) : index(i) {}
-			iterator operator++(int) { return iterator(index++); }
-			bool operator==(const iterator& o) const { return index == o.index; }
-			u32 index;
-		};
-		IndexContainer(u64 c) : count(u32(c)) {}
-		iterator begin() const { return iterator(0); }
-		iterator end() const { return iterator(count); }
-		u32 size() const { return count; }
-		u32 count;
-	};
-
-	////////////////////////////////////////////////////////////////////////////////////////////////////
-
 	u8 GetCompressor(const tchar* str)
 	{
 		#define OODLE_COMPRESSOR(x) if (Equals(str, TC(#x))) return u8(OodleLZ_Compressor_##x);
diff --git a/Engine/Source/Programs/UnrealBuildAccelerator/Common/Public/UbaWorkManager.h b/Engine/Source/Programs/UnrealBuildAccelerator/Common/Public/UbaWorkManager.h
--- a/Engine/Source/Programs/UnrealBuildAccelerator/Common/Public/UbaWorkManager.h
+++ b/Engine/Source/Programs/UnrealBuildAccelerator/Common/Public/UbaWorkManager.h
@@ -16,11 +16,15 @@ namespace uba
 	struct StringView;
 	struct TrackWorkScope;
 
+	////////////////////////////////////////////////////////////////////////////////////////////////////
+
 	using Color = u32;
 	constexpr inline Color ToColor(u8 r, u8 g, u8 b) { return (r << 16) + (g << 8) + b; }
 	constexpr Color ColorWhite = ToColor(255, 255, 255);
 	constexpr Color ColorWork = ToColor(70, 70, 100);
 
+	////////////////////////////////////////////////////////////////////////////////////////////////////
+
 	class WorkTracker
 	{
 	public:
@@ -29,12 +33,16 @@ namespace uba
 		virtual void TrackWorkEnd(u32 id) = 0;
 	};
 
+	////////////////////////////////////////////////////////////////////////////////////////////////////
+
 	struct WorkContext
 	{
 		TrackWorkScope& tracker;
 		Thread* thread = nullptr;
 	};
 
+	////////////////////////////////////////////////////////////////////////////////////////////////////
+
 	using WorkFunction = Function<void(const WorkContext& context)>;
 
 	enum WorkPriority : u8
@@ -44,6 +52,8 @@ namespace uba
 		WorkPriority_Low,
 	};
 
+	////////////////////////////////////////////////////////////////////////////////////////////////////
+
 	class WorkManager : public WorkTracker
 	{
 	public:
@@ -66,6 +76,8 @@ namespace uba
 		Atomic<WorkTracker*> m_workTracker = nullptr;
 	};
 
+	////////////////////////////////////////////////////////////////////////////////////////////////////
+
 	struct Work { WorkFunction func; Bottleneck* bottleneck; TString desc; Work* next; };
 
 	struct WorkQueue
@@ -84,6 +96,7 @@ namespace uba
 		Work* free = nullptr;
 	};
 
+	////////////////////////////////////////////////////////////////////////////////////////////////////
 
 	class WorkManagerImpl : public WorkManager
 	{
@@ -110,6 +123,8 @@ namespace uba
 		Worker* m_firstAvailableWorker = nullptr;
 	};
 
+	////////////////////////////////////////////////////////////////////////////////////////////////////
+
 	struct TrackWorkScope
 	{
 		TrackWorkScope();
@@ -122,6 +137,8 @@ namespace uba
 		u32 id;
 	};
 
+	////////////////////////////////////////////////////////////////////////////////////////////////////
+
 	struct TrackHintScope
 	{
 		TrackHintScope(TrackWorkScope& t, StringView h) : tws(t), hint(h), startTime(GetTime()) {}
@@ -131,6 +148,8 @@ namespace uba
 		u64 startTime;
 	};
 
+	////////////////////////////////////////////////////////////////////////////////////////////////////
+
 	template<u32 BatchSize, typename TContainer, typename TFunc>
 	void WorkManager::ParallelFor(u32 workCount, TContainer& container, TFunc&& func, const StringView& description, WorkPriority priority, Bottleneck* bottleneck)
 	{
@@ -233,4 +252,25 @@ namespace uba
 
 		#endif
 	}
+
+	////////////////////////////////////////////////////////////////////////////////////////////////////
+
+	struct IndexContainer
+	{
+		struct iterator
+		{
+			iterator() : index(0) {}
+			iterator(u32 i) : index(i) {}
+			iterator operator++(int) { return iterator(index++); }
+			bool operator==(const iterator& o) const { return index == o.index; }
+			u32 index;
+		};
+		IndexContainer(u64 c) : count(u32(c)) {}
+		iterator begin() const { return iterator(0); }
+		iterator end() const { return iterator(count); }
+		u32 size() const { return count; }
+		u32 count;
+	};
+
+	////////////////////////////////////////////////////////////////////////////////////////////////////
 }
