[UBA]
* Changed so cas files created for cache upload are only kept in memory and dropped after sent. This should reduce crazy churn on disk for cache-populating machines

diff --git a/Engine/Source/Programs/UnrealBuildAccelerator/Common/Private/UbaStorage.cpp b/Engine/Source/Programs/UnrealBuildAccelerator/Common/Private/UbaStorage.cpp
--- a/Engine/Source/Programs/UnrealBuildAccelerator/Common/Private/UbaStorage.cpp
+++ b/Engine/Source/Programs/UnrealBuildAccelerator/Common/Private/UbaStorage.cpp
@@ -355,7 +355,7 @@ namespace uba
 		return w.CompressToFile(out.storedSize, from, readHandle, readMem, fileSize, to, header, headerSize, lastWriteTime, m_tempPath.data);
 	}
 
-	bool StorageImpl::WriteCasFileNoCheck(WriteCasResult& out, const StringKey& fileNameKey, const tchar* fileName, const CasKey& casKey, const tchar* casFile, bool storeCompressed)
+	bool StorageImpl::WriteCasFileNoCheck(WriteCasResult& out, const StringKey& fileNameKey, const tchar* fileName, const CasKey& casKey, const tchar* casFile, bool storeCompressed, bool useMapping)
 	{
 		StorageStats& stats = Stats();
 		TimerScope ts(stats.createCas);
@@ -436,11 +436,41 @@ namespace uba
 		}
 
 		if (!readIsCompressed && storeCompressed)
+		{
+			if (useMapping)
+			{
+				CompressWriter w(m_logger, m_sharedMemory, m_bufferSlots, m_workManager, Stats(), m_casCompressor, m_casCompressionLevel, m_asyncUnmapViewOfFile);
+				out.memoryOffset = 0;
+				out.isIndependent = true;
+				out.dropAfterUse = true; // This is only used for cache client right now.. so we know this is an output file that can be dropped
+				return w.CompressToMappingFromFile(out.memoryHandle, out.storedSize, readHandle, readFileSize, fileName);
+			}
 			return WriteCompressedFile(out, fileName, readHandle, 0, readFileSize, casFile, nullptr, 0);
-
+		}
 		if (readIsCompressed && !storeCompressed)
 			return m_logger.Error(TC("[WriteCasFileNoCheck] Writing compressed file to non-compressed storage not supported (%s)"), fileName);
 
+		if (useMapping)
+		{
+			SharedMemoryHandle destMapping = m_sharedMemory.CreateHandle(m_logger, fileName);
+			auto destGuard = MakeGuard([&]() { m_sharedMemory.CloseHandle(m_logger, destMapping, fileName); });
+			m_sharedMemory.ExtendMemory(destMapping, readFileSize);
+			u8* destData = m_sharedMemory.MapView(destMapping, fileName, SharedMemoryMapType_ReadWrite);
+			if (!destData)
+				return false;
+			auto unmapGuard = MakeGuard([&]() { m_sharedMemory.UnmapView(destMapping, destData); });
+			if (!ReadFile(m_logger, fileName, readHandle, destData, readFileSize))
+				return false;
+			destGuard.Cancel();
+
+			out.memoryHandle = destMapping;
+			out.storedSize = readFileSize;
+			out.memoryOffset = 0;
+			out.isIndependent = true;
+			out.dropAfterUse = true; // This is only used for cache client right now.. so we know this is an output file that can be dropped
+			return true;
+		}
+
 		FileAccessor destinationFile(m_logger, casFile);
 		if (!destinationFile.CreateWrite(false, DefaultAttributes(), readFileSize, m_tempPath.data))
 			return false;
@@ -467,7 +497,7 @@ namespace uba
 		return true;
 	}
 
-	bool StorageImpl::WriteCasFile(WriteCasResult& out, const tchar* fileName, const CasKey& casKey)
+	bool StorageImpl::WriteCasFile(WriteCasResult& out, const tchar* fileName, const CasKey& casKey, bool useMapping)
 	{
 		UBA_ASSERT(IsCompressed(casKey) == m_storeCompressed);
 		StringBuffer<> casFile;
@@ -481,7 +511,7 @@ namespace uba
 		if (CaseInsensitiveFs)
 			temp.MakeLower();
 		StringKey fileNameKey = ToStringKey(temp);
-		return WriteCasFileNoCheck(out, fileNameKey, fileName, casKey, casFile.data, IsCompressed(casKey));
+		return WriteCasFileNoCheck(out, fileNameKey, fileName, casKey, casFile.data, IsCompressed(casKey), useMapping);
 	}
 
 	void StorageImpl::CasEntryAccessed(const CasKey& casKey)
@@ -640,7 +670,7 @@ namespace uba
 		casEntry.exists = false;
 
 		WriteCasResult res;
-		if (!WriteCasFileNoCheck(res, fileNameKey, fileName, casKey, casFile.data, IsCompressed(casKey)))
+		if (!WriteCasFileNoCheck(res, fileNameKey, fileName, casKey, casFile.data, IsCompressed(casKey), false))
 			return false;
 		casEntry.memoryHandle = res.memoryHandle;
 		casEntry.memoryOffset = res.memoryOffset;
@@ -1822,6 +1852,29 @@ namespace uba
 		return m_casLookup.find(casKey) != m_casLookup.end();
 	}
 
+	void StorageImpl::ReleaseCasEntry(CasEntry& casEntry, bool dropIfZero)
+	{
+		CasEntryWriteLock casEntryLock(casEntry);
+		ReleaseCasEntryNoLock(casEntry, dropIfZero);
+	}
+
+	void StorageImpl::ReleaseCasEntryNoLock(CasEntry& casEntry, bool dropIfZero)
+	{
+		if (--casEntry.refCount)
+			return;
+		if (!casEntry.dropped || !dropIfZero)
+			return;
+		if (!casEntry.memoryHandle.IsValid() || !m_casDataBuffer.IsIndependent(casEntry.memoryHandle))
+			return;
+		SharedMemoryHandle memoryHandle = casEntry.memoryHandle;
+		casEntry.memoryHandle = {};
+		casEntry.verified = true;
+		casEntry.exists = false;
+		casEntry.dropped = false;
+
+		m_casDataBuffer.CloseIndependent(memoryHandle, casEntry.memorySize, TC("ReleaseCasEntryNoLock"));
+	}
+
 	bool StorageImpl::StoreCasKey(CasKey& out, const tchar* fileName, const CasKey& casKeyOverride)
 	{
 		StringBuffer<> forKey;
@@ -1929,7 +1982,7 @@ namespace uba
 		return false;
 	}
 
-	bool StorageImpl::HasCasFile(const CasKey& casKey, CasEntry** out)
+	bool StorageImpl::HasCasFile(const CasKey& casKey, CasEntry** out, bool useMapping)
 	{
 		SCOPED_READ_LOCK(m_casLookupLock, lookupLock);
 		auto it = m_casLookup.find(casKey);
@@ -1973,7 +2026,7 @@ namespace uba
 		casEntry.verified = true;
 		deferredLock.Leave();
 		WriteCasResult res;
-		if (!WriteCasFile(res, deferredCreation.data, casKey))
+		if (!WriteCasFile(res, deferredCreation.data, casKey, useMapping))
 			return m_logger.Error(TC("Failed to write deferred cas %s (%s)"), CasKeyString(casKey).str, deferredCreation.data);
 		casEntry.dropped = res.dropAfterUse;
 
@@ -2047,7 +2100,7 @@ namespace uba
 		if (!fileName)
 			return false;
 		WriteCasResult res;
-		if (!WriteCasFile(res, fileName, casKey))
+		if (!WriteCasFile(res, fileName, casKey, false))
 			return false;
 		casEntry.exists = true;
 		casEntryLock.Leave();
diff --git a/Engine/Source/Programs/UnrealBuildAccelerator/Common/Private/UbaStorageClient.cpp b/Engine/Source/Programs/UnrealBuildAccelerator/Common/Private/UbaStorageClient.cpp
--- a/Engine/Source/Programs/UnrealBuildAccelerator/Common/Private/UbaStorageClient.cpp
+++ b/Engine/Source/Programs/UnrealBuildAccelerator/Common/Private/UbaStorageClient.cpp
@@ -895,7 +895,7 @@ namespace uba
 		return true;
 	}
 
-	bool StorageClient::HasCasFile(const CasKey& casKey, CasEntry** out)
+	bool StorageClient::HasCasFile(const CasKey& casKey, CasEntry** out, bool useMapping)
 	{
 		CasKey localKey = AsCompressed(casKey, false);
 		SCOPED_READ_LOCK(m_localStorageFilesLock, lock);
diff --git a/Engine/Source/Programs/UnrealBuildAccelerator/Common/Private/UbaStorageServer.cpp b/Engine/Source/Programs/UnrealBuildAccelerator/Common/Private/UbaStorageServer.cpp
--- a/Engine/Source/Programs/UnrealBuildAccelerator/Common/Private/UbaStorageServer.cpp
+++ b/Engine/Source/Programs/UnrealBuildAccelerator/Common/Private/UbaStorageServer.cpp
@@ -184,7 +184,7 @@ namespace uba
 			memoryBegin = nullptr;
 		}
 
-		server.ReleaseCasEntry(*casEntry);
+		server.ReleaseCasEntry(*casEntry, server.m_dropIndependentMappingsOnCopy);
 	}
 
 	void StorageServer::OnDisconnected(u32 clientId)
@@ -284,7 +284,7 @@ namespace uba
 				m_activeRefCounts.erase(findIt);
 				refCountLock.Leave();
 				for (CasEntry* casEntry : refCounts)
-					ReleaseCasEntry(*casEntry);
+					ReleaseCasEntry(*casEntry, m_dropIndependentMappingsOnCopy);
 			}
 		}
 
@@ -294,29 +294,6 @@ namespace uba
 		}
 	}
 
-	void StorageServer::ReleaseCasEntry(CasEntry& casEntry)
-	{
-		CasEntryWriteLock casEntryLock(casEntry);
-		ReleaseCasEntryNoLock(casEntry);
-	}
-
-	void StorageServer::ReleaseCasEntryNoLock(CasEntry& casEntry)
-	{
-		if (--casEntry.refCount)
-			return;
-		if (!casEntry.dropped || !m_dropIndependentMappingsOnCopy)
-			return;
-		if (!casEntry.memoryHandle.IsValid() || !m_casDataBuffer.IsIndependent(casEntry.memoryHandle))
-			return;
-		SharedMemoryHandle memoryHandle = casEntry.memoryHandle;
-		casEntry.memoryHandle = {};
-		casEntry.verified = true;
-		casEntry.exists = false;
-		casEntry.dropped = false;
-
-		m_casDataBuffer.CloseIndependent(memoryHandle, casEntry.memorySize, TC("ReleaseCasEntryNoLock"));
-	}
-
 	void StorageServer::ReleaseActiveRefCount(const CasKey& casKey, u32 clientId)
 	{
 		SCOPED_WRITE_LOCK(m_activeRefCountsLock, refCountLock);
@@ -330,7 +307,7 @@ namespace uba
 				continue;
 			findIt->second.erase(it);
 			refCountLock.Leave();
-			ReleaseCasEntry(*casEntry);
+			ReleaseCasEntry(*casEntry, m_dropIndependentMappingsOnCopy);
 			break;
 		}
 	}
@@ -683,7 +660,7 @@ namespace uba
 				if (!casEntry.refCount)
 					m_logger.Error(TC("CasEntry refcount wrap-around during fetch. This is bad, report to uba devs"));
 				casEntryLock.Leave();
-				auto rlg = MakeGuard([&]() { casEntryLock.Enter(); ReleaseCasEntryNoLock(casEntry); });
+				auto rlg = MakeGuard([&]() { casEntryLock.Enter(); ReleaseCasEntryNoLock(casEntry, m_dropIndependentMappingsOnCopy); });
 
 				if (!casEntry.exists)
 				{
diff --git a/Engine/Source/Programs/UnrealBuildAccelerator/Common/Private/UbaStorageUtils.cpp b/Engine/Source/Programs/UnrealBuildAccelerator/Common/Private/UbaStorageUtils.cpp
--- a/Engine/Source/Programs/UnrealBuildAccelerator/Common/Private/UbaStorageUtils.cpp
+++ b/Engine/Source/Programs/UnrealBuildAccelerator/Common/Private/UbaStorageUtils.cpp
@@ -1037,9 +1037,42 @@ namespace uba
 		LinearWriterMem destination(compressedData, mappingSize);
 		if (!CompressFromMemOrFile(destination, hint, InvalidFileHandle, readMem, fileSize))
 			return false;
+
+		unmapGuard.Execute();
+		destGuard.Cancel();
+
+		u64 written = destination.GetWritten();
+		m_allocator.TrimMemory(destMapping, written);
+
+		outMappingHandle = destMapping;
+		outMappingSize = written;
+		return true;
+	}
+
+	bool CompressWriter::CompressToMappingFromFile(SharedMemoryHandle& outMappingHandle, u64& outMappingSize, FileHandle readHandle, u64 fileSize, const tchar* hint)
+	{
+		u64 mappingSize = fileSize + 1024;
+
+		SharedMemoryHandle destMapping = m_allocator.CreateHandle(m_logger, hint);
+		m_allocator.ExtendMemory(destMapping, mappingSize);
+		auto destGuard = MakeGuard([&]() { m_allocator.CloseHandle(m_logger, destMapping, hint); });
+		u8* compressedData = m_allocator.MapView(destMapping, hint, SharedMemoryMapType_ReadWrite);
+		if (!compressedData)
+			return false;
+		auto unmapGuard = MakeGuard([&]() { m_allocator.UnmapView(destMapping, compressedData); });
+
+		LinearWriterMem destination(compressedData, mappingSize);
+		if (!CompressFromMemOrFile(destination, hint, readHandle, nullptr, fileSize))
+			return false;
+
+		unmapGuard.Execute();
 		destGuard.Cancel();
+
+		u64 written = destination.GetWritten();
+		m_allocator.TrimMemory(destMapping, written);
+
 		outMappingHandle = destMapping;
-		outMappingSize = destination.GetWritten();
+		outMappingSize = written;
 		return true;
 	}
 
diff --git a/Engine/Source/Programs/UnrealBuildAccelerator/Common/Public/UbaStorage.h b/Engine/Source/Programs/UnrealBuildAccelerator/Common/Public/UbaStorage.h
--- a/Engine/Source/Programs/UnrealBuildAccelerator/Common/Public/UbaStorage.h
+++ b/Engine/Source/Programs/UnrealBuildAccelerator/Common/Public/UbaStorage.h
@@ -187,7 +187,7 @@ namespace uba
 		struct CasEntry;
 		struct CasEntryWriteLock;
 
-		virtual bool HasCasFile(const CasKey& casKey, CasEntry** out = nullptr);
+		virtual bool HasCasFile(const CasKey& casKey, CasEntry** out = nullptr, bool useMapping = false);
 		bool EnsureCasFile(const CasKey& casKey, const tchar* fileName);
 		CasKey CalculateCasKey(const tchar* fileName, FileHandle fileHandle, u64 fileSize, bool storeCompressed);
 		CasKey CalculateCasKey(const u8* fileMem, u64 fileSize, bool storeCompressed);
@@ -196,6 +196,8 @@ namespace uba
 		bool IsFileVerified(const StringKey& fileNameKey);
 		void ReportFileInfoWeak(const StringKey& fileNameKey, u64 verifiedLastWriteTime, u64 verifiedSize);
 		bool HasBeenSeen(const CasKey& casKey);
+		void ReleaseCasEntry(CasEntry& casEntry, bool dropIfZero);
+		void ReleaseCasEntryNoLock(CasEntry& casEntry, bool dropIfZero);
 
 		struct WriteCasResult : WriteResult
 		{
@@ -204,8 +206,8 @@ namespace uba
 
 		virtual void RegisterExternalFileMappingsProvider(ExternalFileMappingsProvider&& provider) override;
 		virtual bool WriteCompressedFile(WriteResult& out, const tchar* from, FileHandle readHandle, u8* readMem, u64 fileSize, const tchar* toFile, const void* header, u64 headerSize, u64 lastWriteTime = 0) final;
-		bool WriteCasFileNoCheck(WriteCasResult& out, const StringKey& fileNameKey, const tchar* fileName, const CasKey& casKey, const tchar* casFile, bool storeCompressed);
-		bool WriteCasFile(WriteCasResult& out, const tchar* fileName, const CasKey& casKey);
+		bool WriteCasFileNoCheck(WriteCasResult& out, const StringKey& fileNameKey, const tchar* fileName, const CasKey& casKey, const tchar* casFile, bool storeCompressed, bool useMapping);
+		bool WriteCasFile(WriteCasResult& out, const tchar* fileName, const CasKey& casKey, bool useMapping);
 		bool VerifyExisting(bool& outReturnValue, CasEntryWriteLock& casEntryLock, const CasKey& casKey, CasEntry& casEntry, StringBufferBase& casFile);
 		bool AddCasFile(StringKey fileNameKey, const tchar* fileName, const CasKey& casKey, bool deferCreation, u64 fileSize = InvalidValue, u64 lastWriteTime = 0);
 		void CasEntryAccessed(const CasKey& casKey);
diff --git a/Engine/Source/Programs/UnrealBuildAccelerator/Common/Public/UbaStorageClient.h b/Engine/Source/Programs/UnrealBuildAccelerator/Common/Public/UbaStorageClient.h
--- a/Engine/Source/Programs/UnrealBuildAccelerator/Common/Public/UbaStorageClient.h
+++ b/Engine/Source/Programs/UnrealBuildAccelerator/Common/Public/UbaStorageClient.h
@@ -59,7 +59,7 @@ namespace uba
 		virtual bool RetrieveCasFile(RetrieveResult& out, const CasKey& casKey, StringKey keyHint, const tchar* hint, FileMappingBuffer* mappingBuffer = nullptr, u64 memoryMapAlignment = 1, bool allowProxy = true, u32 clientId = 0) override;
 		virtual bool StoreCasFile(CasKey& out, const tchar* fileName, const CasKey& casKeyOverride, bool deferCreation) override;
 		virtual bool StoreCasFileClient(CasKey& out, StringKey fileNameKey, const tchar* fileName, SharedMemoryHandle mappingFile, u64 mappingOffset, u64 fileSize, const tchar* hint, bool keepMappingInMemory = false, bool storeCompressed = true) override;
-		virtual bool HasCasFile(const CasKey& casKey, CasEntry** out = nullptr) override;
+		virtual bool HasCasFile(const CasKey& casKey, CasEntry** out = nullptr, bool useMapping = false) override;
 		virtual void Ping() override;
 		virtual void PrintSummary(Logger& logger) override;
 
diff --git a/Engine/Source/Programs/UnrealBuildAccelerator/Common/Public/UbaStorageServer.h b/Engine/Source/Programs/UnrealBuildAccelerator/Common/Public/UbaStorageServer.h
--- a/Engine/Source/Programs/UnrealBuildAccelerator/Common/Public/UbaStorageServer.h
+++ b/Engine/Source/Programs/UnrealBuildAccelerator/Common/Public/UbaStorageServer.h
@@ -50,8 +50,6 @@ namespace uba
 		virtual bool HasProxy(u32 clientId) override;
 
 		void OnDisconnected(u32 clientId);
-		void ReleaseCasEntry(CasEntry& casEntry);
-		void ReleaseCasEntryNoLock(CasEntry& casEntry);
 		void ReleaseActiveRefCount(const CasKey& casKey, u32 clientId);
 		void UpdateWaitEntries(const CasKey& casKey, bool success);
 		bool HandleMessage(const ConnectionInfo& connectionInfo, const WorkContext& workContext, u8 messageType, BinaryReader& reader, BinaryWriter& writer);
diff --git a/Engine/Source/Programs/UnrealBuildAccelerator/Common/Public/UbaStorageUtils.h b/Engine/Source/Programs/UnrealBuildAccelerator/Common/Public/UbaStorageUtils.h
--- a/Engine/Source/Programs/UnrealBuildAccelerator/Common/Public/UbaStorageUtils.h
+++ b/Engine/Source/Programs/UnrealBuildAccelerator/Common/Public/UbaStorageUtils.h
@@ -79,6 +79,7 @@ namespace uba
 
 		bool CompressToFile(u64& outCompressedSize, const tchar* from, FileHandle readHandle, u8* readMem, u64 fileSize, const tchar* toFile, const void* header, u64 headerSize, u64 lastWriteTime, const tchar* tempPath);
 		bool CompressToMapping(SharedMemoryHandle& outMappingHandle, u64& outMappingSize, u8* readMem, u64 fileSize, const tchar* hint);
+		bool CompressToMappingFromFile(SharedMemoryHandle& outMappingHandle, u64& outMappingSize, FileHandle readHandle, u64 fileSize, const tchar* hint);
 		bool CompressFromMem(LinearWriter& destination, u32 workCount, const u8* uncompressedData, u64 fileSize, u64 maxUncompressedBlock, u64& totalWritten);
 		bool CompressFromMemOrFile(LinearWriter& destination, const tchar* from, FileHandle readHandle, u8* readMem, u64 fileSize);
 
