[UBA]
* Rolled back storage version from 5 to 4 in order to be able to release new binaries

diff --git a/Engine/Source/Programs/UnrealBuildAccelerator/Common/Private/UbaStorageClient.cpp b/Engine/Source/Programs/UnrealBuildAccelerator/Common/Private/UbaStorageClient.cpp
--- a/Engine/Source/Programs/UnrealBuildAccelerator/Common/Private/UbaStorageClient.cpp
+++ b/Engine/Source/Programs/UnrealBuildAccelerator/Common/Private/UbaStorageClient.cpp
@@ -372,7 +372,8 @@ namespace uba
 				NetworkMessage msg(*client, ServiceId, StorageMessageType_FetchBegin, writer);
 				writer.WriteByte(wantsProxy ? 1 : 0);
 				writer.WriteCasKey(casKey);
-				writer.WriteStringKey(keyHint);
+				if constexpr (StorageNetworkVersion >= 5)
+					writer.WriteStringKey(keyHint);
 				writer.WriteString(hint);
 				BinaryReader reader(slot, 0, SendMaxSize);
 				if (!msg.Send(reader))
diff --git a/Engine/Source/Programs/UnrealBuildAccelerator/Common/Private/UbaStorageProxy.cpp b/Engine/Source/Programs/UnrealBuildAccelerator/Common/Private/UbaStorageProxy.cpp
--- a/Engine/Source/Programs/UnrealBuildAccelerator/Common/Private/UbaStorageProxy.cpp
+++ b/Engine/Source/Programs/UnrealBuildAccelerator/Common/Private/UbaStorageProxy.cpp
@@ -183,7 +183,9 @@ namespace uba
 		UBA_ASSERT((recvFlags & 2) == 0);(void)recvFlags;
 
 		CasKey casKey = reader.ReadCasKey();
-		StringKey keyHint = reader.ReadStringKey();
+		StringKey keyHint;
+		if (connectionInfo.GetStorageVersion() >= 5)
+			keyHint = reader.ReadStringKey();
 		StringBuffer<> hint;
 		reader.ReadString(hint);
 
@@ -223,7 +225,8 @@ namespace uba
 			auto& writer2 = file.beginMessage->writer;
 			writer2.WriteByte(2); // Does not want proxy but informs it is proxy
 			writer2.WriteCasKey(casKey);
-			writer2.WriteStringKey(keyHint);
+			if constexpr (StorageNetworkVersion >= 5)
+				writer2.WriteStringKey(keyHint);
 			writer2.WriteString(hint);
 			writer2.WriteBytes(reader.GetPositionData(), reader.GetLeft());
 
diff --git a/Engine/Source/Programs/UnrealBuildAccelerator/Common/Public/UbaNetwork.h b/Engine/Source/Programs/UnrealBuildAccelerator/Common/Public/UbaNetwork.h
--- a/Engine/Source/Programs/UnrealBuildAccelerator/Common/Public/UbaNetwork.h
+++ b/Engine/Source/Programs/UnrealBuildAccelerator/Common/Public/UbaNetwork.h
@@ -15,7 +15,7 @@ namespace uba
 	static constexpr u8 CacheServiceId = 3;
 	
 	static constexpr u32 SystemNetworkVersion = 1339;
-	static constexpr u32 StorageNetworkVersion = 5;
+	static constexpr u32 StorageNetworkVersion = 4;
 	static constexpr u32 SessionNetworkVersion = 47;
 	static constexpr u32 CacheNetworkVersion = 5;
 	static constexpr u32 CacheBucketVersion = 3;
