[UBA]
* Added missing file from previous changelist where CacheClient keeps cas files in memory and drop after send

diff --git a/Engine/Source/Programs/UnrealBuildAccelerator/Common/Private/UbaCacheClient.cpp b/Engine/Source/Programs/UnrealBuildAccelerator/Common/Private/UbaCacheClient.cpp
--- a/Engine/Source/Programs/UnrealBuildAccelerator/Common/Private/UbaCacheClient.cpp
+++ b/Engine/Source/Programs/UnrealBuildAccelerator/Common/Private/UbaCacheClient.cpp
@@ -1327,14 +1327,13 @@ namespace uba
 			TimerScope createCasTs(stats.createCas);
 			StorageImpl::CasEntry* casEntry;
 
-			if (m_storage.HasCasFile(casKey, &casEntry))
+			if (m_storage.HasCasFile(casKey, &casEntry, true))
 			{
 				createCasTs.Leave();
 
 				auto refCountGuard = MakeGuard([&]()
 					{
-						StorageImpl::CasEntryWriteLock casEntryLock(*casEntry);
-						--casEntry->refCount; // HasCasFile increase refcount
+						m_storage.ReleaseCasEntry(*casEntry, true); // HasCasFile increase refcount
 					});
 
 				UBA_TRACE_WRITE_HINT("SendFile");
@@ -1526,7 +1525,7 @@ namespace uba
 			Entry(u8* slot, u32 i, u32 messageMaxSize) : reader(slot + i * messageMaxSize, 0, SendMaxSize), done(EventResetType_Manual) {}
 			NetworkMessage message;
 			BinaryReader reader;
-			Event done;
+			EventSlim done;
 		};
 
 		bool success = true;
@@ -1539,7 +1538,7 @@ namespace uba
 			e.message.Init(m_client, CacheServiceId, messageType, writer, HasPriority);
 			writer.Write7BitEncoded(MakeId(bucketId));
 			writer.WriteU32(tableSize + i*messageFetchSize);
-			if (e.message.SendAsync(e.reader, [](bool error, void* userData) { ((Event*)userData)->Set(); }, &e.done))
+			if (e.message.SendAsync(e.reader, [](bool error, void* userData) { ((EventSlim*)userData)->Set(); }, &e.done))
 				continue;
 			entries.pop_back();
 			success = false;
