[UBA]
* Made it possible to run UbaTest.exe on remote machines

diff --git a/Engine/Source/Programs/UnrealBuildAccelerator/Cli/Private/UbaCli.cpp b/Engine/Source/Programs/UnrealBuildAccelerator/Cli/Private/UbaCli.cpp
--- a/Engine/Source/Programs/UnrealBuildAccelerator/Cli/Private/UbaCli.cpp
+++ b/Engine/Source/Programs/UnrealBuildAccelerator/Cli/Private/UbaCli.cpp
@@ -1245,9 +1245,6 @@ namespace uba
 			if (!networkServer.StartListen(networkBackend, port, listenIp.data))
 				return false;
 		}
-		auto stopServer = MakeGuard([&]() { networkServer.DisconnectClients(); });
-
-		auto stopListen = MakeGuard([&]() { networkBackend.StopListen(); });
 
 		auto RunLocal = [&](const TString& app, const TString& arg, bool enableDetour)
 		{
@@ -1340,6 +1337,7 @@ namespace uba
 		};
 
 		CoordinatorWrapper coordinator;
+		auto cg = MakeGuard([&]() { coordinator.Destroy(); });
 
 		auto CreateCoordinator = [&](Scheduler* scheduler)
 			{
@@ -1364,6 +1362,9 @@ namespace uba
 				return true;
 			};
 
+		auto stopServer = MakeGuard([&]() { networkServer.DisconnectClients(); });
+		auto stopListen = MakeGuard([&]() { networkBackend.StopListen(); });
+
 		auto RunScheduler = [&](const TString& app, const TString& arg)
 		{
 			/*
@@ -1408,7 +1409,6 @@ namespace uba
 			if (!CreateCoordinator(&scheduler))
 				return false;
 
-
 			if (StringView(app).EndsWith(TCV(".yaml")))
 			{
 				if (!scheduler.EnqueueFromFile(app.c_str(), [&](EnqueueProcessInfo& epi)
@@ -1487,8 +1487,6 @@ namespace uba
 				return RunQueue();
 		};
 
-		auto cg = MakeGuard([&]() { coordinator.Destroy(); });
-
 #if PLATFORM_WINDOWS	// Annoying that link.exe/lld-link.exe needs path to windows folder.. 
 		if (!useScheduler)
 		{
diff --git a/Engine/Source/Programs/UnrealBuildAccelerator/Common/Private/UbaNetworkServer.cpp b/Engine/Source/Programs/UnrealBuildAccelerator/Common/Private/UbaNetworkServer.cpp
--- a/Engine/Source/Programs/UnrealBuildAccelerator/Common/Private/UbaNetworkServer.cpp
+++ b/Engine/Source/Programs/UnrealBuildAccelerator/Common/Private/UbaNetworkServer.cpp
@@ -21,6 +21,7 @@ namespace uba
 		table.GetValueAsU32(sendBufferSize, TC("SendBufferSize"));
 		table.GetValueAsU32(workerCount, TC("WorkerCount"));
 		table.GetValueAsU32(sendSize, TC("SendSize"));
+		table.GetValueAsU32(receiveTimeoutSeconds, TC("ReceiveTimeoutSeconds"));
 		table.GetValueAsBool(logConnections, TC("LogConnections"));
 		table.GetValueAsBool(useKeepAlive, TC("UseKeepAlive"));
 		table.GetValueAsBool(useMessagePriority, TC("UseMessagePriority"));
diff --git a/Engine/Source/Programs/UnrealBuildAccelerator/Common/Private/UbaSessionClient.cpp b/Engine/Source/Programs/UnrealBuildAccelerator/Common/Private/UbaSessionClient.cpp
--- a/Engine/Source/Programs/UnrealBuildAccelerator/Common/Private/UbaSessionClient.cpp
+++ b/Engine/Source/Programs/UnrealBuildAccelerator/Common/Private/UbaSessionClient.cpp
@@ -1723,13 +1723,15 @@ namespace uba
 		
 		response.done.IsSet(); // Wait forever (until message trigger error or finishes)
 
+		bool asyncRes = msg.ProcessAsyncResults(reader);
+
 		if (reportPing)
 		{
 			LoggerWithWriter logger(g_consoleLogWriter);
-			logger.Info(TC("Ping finished after %s"), TimeToText(GetTime() - time).str);
+			logger.Info(TC("Ping finished after %s (%u)"), TimeToText(GetTime() - time).str, msg.GetError());
 		}
 
-		if (!msg.ProcessAsyncResults(reader) || msg.GetError())
+		if (!asyncRes || msg.GetError())
 		{
 			m_loop = false;
 			return;
@@ -2074,7 +2076,9 @@ namespace uba
 
 					void* env = GetProcessEnvironmentVariables();
 
-					auto process = new ProcessImpl(*this, startInfo.processId, nullptr, true);
+					bool detourEnabled = !StringView(startInfo.applicationStr).GetFileName().Contains(TCV("ubatest.exe"));
+
+					auto process = new ProcessImpl(*this, startInfo.processId, nullptr, detourEnabled);
 
 					activeProcesses.emplace_back(process);
 					ProcessRec* rec = &activeProcesses.back();
diff --git a/Engine/Source/Programs/UnrealBuildAccelerator/Common/Private/UbaSessionServer.cpp b/Engine/Source/Programs/UnrealBuildAccelerator/Common/Private/UbaSessionServer.cpp
--- a/Engine/Source/Programs/UnrealBuildAccelerator/Common/Private/UbaSessionServer.cpp
+++ b/Engine/Source/Programs/UnrealBuildAccelerator/Common/Private/UbaSessionServer.cpp
@@ -1079,6 +1079,19 @@ namespace uba
 		if (!GetBinaryModules(modules, applicationName.data))
 			return false;
 
+		#if PLATFORM_WINDOWS
+		if (applicationName.EndsWith(TCV("ubatest.exe"))) // Since we don't detour ubatest.exe we need to make sure to transfer all the files needed to do the unit tests
+		{
+			StringBuffer<> extra;
+			extra.AppendDir(applicationName).EnsureEndsWithSlash().Append(TCV("UbaWine.dll.so"));
+			modules.emplace_back(TC("UbaWine.dll.so"), extra.data, DefaultAttributes());
+			extra.Clear().AppendDir(applicationName).EnsureEndsWithSlash().Append(TCV("UbaDetours.dll"));
+			modules.emplace_back(TC("UbaDetours.dll"), extra.data, DefaultAttributes());
+			extra.Clear().AppendDir(applicationName).EnsureEndsWithSlash().Append(TCV("UbaTestApp.exe"));
+			modules.emplace_back(TC("UbaTestApp.exe"), extra.data, DefaultAttributes());
+		}
+		#endif
+
 		writer.WriteU32(m_systemPath.count);
 		writer.WriteU32(u32(modules.size()));
 		for (BinaryModule& m : modules)
diff --git a/Engine/Source/Programs/UnrealBuildAccelerator/Test/Private/UbaTestNetwork.cpp b/Engine/Source/Programs/UnrealBuildAccelerator/Test/Private/UbaTestNetwork.cpp
--- a/Engine/Source/Programs/UnrealBuildAccelerator/Test/Private/UbaTestNetwork.cpp
+++ b/Engine/Source/Programs/UnrealBuildAccelerator/Test/Private/UbaTestNetwork.cpp
@@ -19,7 +19,7 @@ namespace uba
 		Thread t([&]()
 			{
 				logger.Info(TC("Starting to Listen"));
-				tcp.StartListen(logger, port, TC("127.0.0.1"), [&](void* connection, const sockaddr& remoteSocketAddr)
+				if (!tcp.StartListen(logger, port, TC("127.0.0.1"), [&](void* connection, const sockaddr& remoteSocketAddr)
 					{
 						logger.Info(TC("Listen got connection"));
 						tcp.SetDisconnectCallback(connection, nullptr, [](void*, const Guid&, void*) {});
@@ -32,21 +32,25 @@ namespace uba
 							},
 							nullptr, TC(""));
 						return true;
-					});
+					}))
+					logger.Error(TC("Failed to listen"));
+
 				return 0;
 			});
 
 		Sleep(100);
 		logger.Info(TC("Starting to Connect"));
-		tcp.Connect(logger, TC("127.0.0.1"), [&](void* connection, const sockaddr& remoteSocketAddr, bool* timedOut)
+		if (!tcp.Connect(logger, TC("127.0.0.1"), [&](void* connection, const sockaddr& remoteSocketAddr, bool* timedOut)
 			{
 				tcp.SetDisconnectCallback(connection, nullptr, [](void*, const Guid&, void*) {});
 				tcp.SetRecvCallbacks(connection, &tcp, 1, [](void* context, const Guid& connectionUid, u8* headerData, void*& outBodyContext, u8*& outBodyData, u32& outBodySize) { return true; }, nullptr, TC(""));
 				u8 b = 42;
 				NetworkBackend::SendContext sc;
-				tcp.Send(logger, connection, &b, 1, sc, TC(""));
+				if (!tcp.Send(logger, connection, &b, 1, sc, TC("")))
+					logger.Error(TC("Failed to send"));
 				return true;
-			}, port);
+			}, port))
+			return logger.Error(TC("Failed to connect"));
 
 		Sleep(200);
 		t.Wait();
