[UBA]
* Disabled unit test that only works with storage network version 5

diff --git a/Engine/Source/Programs/UnrealBuildAccelerator/Test/Private/UbaTestStorage.cpp b/Engine/Source/Programs/UnrealBuildAccelerator/Test/Private/UbaTestStorage.cpp
--- a/Engine/Source/Programs/UnrealBuildAccelerator/Test/Private/UbaTestStorage.cpp
+++ b/Engine/Source/Programs/UnrealBuildAccelerator/Test/Private/UbaTestStorage.cpp
@@ -231,37 +231,38 @@ namespace uba
 		if (!storageClient.RetrieveCasFile(result, casKey, {}, fileName.data))
 			return logger.Error(TC("Failed to retrieve file %s"), fileName.data);
 
-
-
-		StringKey testFileNameKey = ToStringKeyLower(fileName);
-
-		auto& sharedMem = storageServer.GetSharedMemory();
-		auto memHandle = sharedMem.CreateHandle(logger, TC(""));
-		sharedMem.ExtendMemory(memHandle, 1);
-		u8* mem = sharedMem.MapView(memHandle, TC(""), SharedMemoryMapType_ReadWrite);
-		*mem = 'H';
-		sharedMem.UnmapView(memHandle, mem);
-
-		storageServer.RegisterExternalFileMappingsProvider([&](Storage::ExternalFileMapping& out, StringKey fileNameKey, const tchar* fileName)
-			{
-				if (fileNameKey != testFileNameKey)
-					return false;
-				out.handle = memHandle;
-				out.size = 1;
-				out.dropCasAfterUse = true;
-				out.createIndependentMapping = true;
-				return true;
-			});
-
-		storageClient.DropCasFile(casKey, true, fileName.data, false);
-		storageServer.DropCasFile(casKey, true, fileName.data, false);
-		if (!storageClient.RetrieveCasFile(result, casKey, testFileNameKey, TC("KnownInput")))
-			return logger.Error(TC("Failed to retrieve file %s"), fileName.data);
-
-		storageClient.DropCasFile(casKey, true, fileName.data, false);
-		storageServer.DropCasFile(casKey, true, fileName.data, false);
-		if (!storageClient.RetrieveCasFile(result, casKey, testFileNameKey, TC("KnownInput")))
-			return logger.Error(TC("Failed to retrieve file %s"), fileName.data);
+		if constexpr (StorageNetworkVersion >= 5)
+		{
+			StringKey testFileNameKey = ToStringKeyLower(fileName);
+
+			auto& sharedMem = storageServer.GetSharedMemory();
+			auto memHandle = sharedMem.CreateHandle(logger, TC(""));
+			sharedMem.ExtendMemory(memHandle, 1);
+			u8* mem = sharedMem.MapView(memHandle, TC(""), SharedMemoryMapType_ReadWrite);
+			*mem = 'H';
+			sharedMem.UnmapView(memHandle, mem);
+
+			storageServer.RegisterExternalFileMappingsProvider([&](Storage::ExternalFileMapping& out, StringKey fileNameKey, const tchar* fileName)
+				{
+					if (fileNameKey != testFileNameKey)
+						return false;
+					out.handle = memHandle;
+					out.size = 1;
+					out.dropCasAfterUse = true;
+					out.createIndependentMapping = true;
+					return true;
+				});
+
+			storageClient.DropCasFile(casKey, true, fileName.data, false);
+			storageServer.DropCasFile(casKey, true, fileName.data, false);
+			if (!storageClient.RetrieveCasFile(result, casKey, testFileNameKey, TC("KnownInput")))
+				return logger.Error(TC("Failed to retrieve file %s"), fileName.data);
+
+			storageClient.DropCasFile(casKey, true, fileName.data, false);
+			storageServer.DropCasFile(casKey, true, fileName.data, false);
+			if (!storageClient.RetrieveCasFile(result, casKey, testFileNameKey, TC("KnownInput")))
+				return logger.Error(TC("Failed to retrieve file %s"), fileName.data);
+		}
 
 		return true;
 	}
