[UBA]
* Added support for merged modules on native linux

#Synced-CL 48930946

diff --git a/Engine/Source/Programs/UnrealBuildAccelerator/Common/Private/UbaSession.cpp b/Engine/Source/Programs/UnrealBuildAccelerator/Common/Private/UbaSession.cpp
--- a/Engine/Source/Programs/UnrealBuildAccelerator/Common/Private/UbaSession.cpp
+++ b/Engine/Source/Programs/UnrealBuildAccelerator/Common/Private/UbaSession.cpp
@@ -3313,7 +3313,7 @@ namespace uba
 				lastWriteTime = writtenFile.lastWriteTime;
 			}
 
-			if ((msg.process.m_extractExports || m_extractObjFilesSymbols) && rules.ShouldExtractSymbols(activeFile.name) && fileSize != 0)
+			if ((msg.process.m_extractExports || m_extractObjFilesSymbols) && rules.ShouldExtractSymbols(name))
 				if (!ExtractSymbolsFromObjectFile(msg, name.data, fileSize))
 					return false;
 		}
@@ -3967,11 +3967,30 @@ namespace uba
 
 	bool Session::ExtractSymbolsFromObjectFile(const CloseFileMessage& msg, const tchar* fileName, u64 fileSize)
 	{
-		if (!msg.memoryHandle.IsValid())
-			return m_logger.Error(TC("Can't extract symbols from obj file that is written directly to disk (%s writing %s)"), msg.process.m_startInfo.application, fileName);
+		u8* mem = nullptr;
+		
+		auto memClose = MakeGuard([&]() { if (msg.memoryHandle.IsValid()) m_sharedMemory.UnmapView(msg.memoryHandle, mem); });
 
-		u8* mem = m_sharedMemory.MapView(msg.memoryHandle, fileName);
-		auto memClose = MakeGuard([&]() { m_sharedMemory.UnmapView(msg.memoryHandle, mem); });
+		FileAccessor objFile(m_logger, fileName);
+
+		if (msg.memoryHandle.IsValid())
+		{
+			if (!fileSize)
+				return true;
+			mem = m_sharedMemory.MapView(msg.memoryHandle, fileName);
+		}
+		else
+		{
+			if constexpr (IsWindows)
+				return m_logger.Error(TC("Can't extract symbols from obj file that is written directly to disk (%s writing %s)"), msg.process.m_startInfo.application, fileName);
+			
+			if (!objFile.OpenMemoryRead())
+				return false;
+			mem = objFile.GetData();
+			fileSize = objFile.GetSize();
+			if (!fileSize)
+				return true;
+		}
 
 		ObjectFile* objectFile = ObjectFile::Parse(m_logger, ObjectFileParseMode_All, mem, fileSize, fileName);
 		if (!objectFile)
@@ -3991,24 +4010,39 @@ namespace uba
 		if (!objectFile->WriteImportsAndExports(m_logger, memoryBlock, verbose))
 			return false;
 
-		SharedMemoryHandle symHandle = m_sharedMemory.CreateHandle(m_logger, TC("SymHandle"));
-		auto mg = MakeGuard([&]() { m_sharedMemory.CloseHandle(m_logger, symHandle, TC("SymHandle")); });
-		m_sharedMemory.ExtendMemory(symHandle, memoryBlock.writtenSize);
-		u8* mem2 = m_sharedMemory.MapView(symHandle, TC("SymHandle"), SharedMemoryMapType_ReadWrite);
-		if (!mem2)
-			return false;
-
-		MapMemoryCopy(mem2, memoryBlock.memory, memoryBlock.writtenSize);
-
-		m_sharedMemory.UnmapView(symHandle, mem2);
-
+		SharedMemoryHandle symHandle;
 		StringKey symFileKey = CaseInsensitiveFs ? ToStringKeyLower(exportsFile) : ToStringKey(exportsFile);
-		u64 lastWriteTime = GetSystemTimeAsFileTime();
+		u64 lastWriteTime = 0;
 
-		if (!RegisterCreateFileForWrite(symFileKey, exportsFile, false, memoryBlock.writtenSize, lastWriteTime))
-			return false;
+		if (m_allowMemoryMaps)
+		{
+			symHandle = m_sharedMemory.CreateHandle(m_logger, TC("SymHandle"));
+			auto mg = MakeGuard([&]() { m_sharedMemory.CloseHandle(m_logger, symHandle, TC("SymHandle")); });
+			m_sharedMemory.ExtendMemory(symHandle, memoryBlock.writtenSize);
+			u8* mem2 = m_sharedMemory.MapView(symHandle, TC("SymHandle"), SharedMemoryMapType_ReadWrite);
+			if (!mem2)
+				return false;
 
-		mg.Cancel();
+			MapMemoryCopy(mem2, memoryBlock.memory, memoryBlock.writtenSize);
+			m_sharedMemory.UnmapView(symHandle, mem2);
+
+			lastWriteTime = GetSystemTimeAsFileTime();
+			if (!RegisterCreateFileForWrite(symFileKey, exportsFile, false, memoryBlock.writtenSize, lastWriteTime))
+				return false;
+			mg.Cancel();
+		}
+		else
+		{
+			FileAccessor destFile(m_logger, exportsFile.data);
+			if (!destFile.CreateWrite(false, DefaultAttributes(), memoryBlock.writtenSize))
+				return false;
+			if (!destFile.Write(memoryBlock.memory, memoryBlock.writtenSize))
+				return false;
+			if (!destFile.Close(&lastWriteTime))
+				return false;
+			if (!RegisterCreateFileForWrite(symFileKey, exportsFile, true, memoryBlock.writtenSize, lastWriteTime))
+				return false;
+		}
 
 		auto insres = msg.process.m_shared.writtenFiles.try_emplace(symFileKey);
 		WrittenFile& writtenFile = insres.first->second;
