[UBA]
* Fixed visualization of long mac traces where time could wrap around

diff --git a/Engine/Source/Programs/UnrealBuildAccelerator/Visualizer/Private/UbaTraceReader.h b/Engine/Source/Programs/UnrealBuildAccelerator/Visualizer/Private/UbaTraceReader.h
--- a/Engine/Source/Programs/UnrealBuildAccelerator/Visualizer/Private/UbaTraceReader.h
+++ b/Engine/Source/Programs/UnrealBuildAccelerator/Visualizer/Private/UbaTraceReader.h
@@ -201,6 +201,7 @@ namespace uba
 		u64 traceSystemStartTimeUs = 0;
 		u64 startTime = 0;
 		u64 frequency = 0;
+		double timeMultiplier = 0;
 		u64 lastKillProcessTime = 0;
 		u64 lastSpawningDelayStartTime = 0;
 		u64 lastSpawningDelayEndTime = 0;
diff --git a/Engine/Source/Programs/UnrealBuildAccelerator/Visualizer/Private/UbaVisualizer.cpp b/Engine/Source/Programs/UnrealBuildAccelerator/Visualizer/Private/UbaVisualizer.cpp
--- a/Engine/Source/Programs/UnrealBuildAccelerator/Visualizer/Private/UbaVisualizer.cpp
+++ b/Engine/Source/Programs/UnrealBuildAccelerator/Visualizer/Private/UbaVisualizer.cpp
@@ -1628,7 +1628,7 @@ namespace uba
 
 							RECT rect;
 							rect.left = left;
-							rect.right = posX + int(float(TimeToS(stop)) * scaleX) - 1;
+							rect.right = Max(posX + int(float(TimeToS(stop)) * scaleX) - 1, left + 1);
 							rect.top = posY;
 							rect.bottom = rectBottom - 1;
 
@@ -1639,8 +1639,6 @@ namespace uba
 								continue;
 							++visibleBoxes;
 
-							rect.right = Max(int(rect.right), left + 1);
-
 							bool selected = m_processSelected && m_processSelectedLocation == processLocation;
 							if (selected)
 								process.bitmapDirty = true;
