[UBA]
* Fixed visualization of long mac traces where time could wrap around

diff --git a/Engine/Source/Programs/UnrealBuildAccelerator/Visualizer/Private/UbaTraceReader.cpp b/Engine/Source/Programs/UnrealBuildAccelerator/Visualizer/Private/UbaTraceReader.cpp
--- a/Engine/Source/Programs/UnrealBuildAccelerator/Visualizer/Private/UbaTraceReader.cpp
+++ b/Engine/Source/Programs/UnrealBuildAccelerator/Visualizer/Private/UbaTraceReader.cpp
@@ -98,7 +98,7 @@ namespace uba
 
 	u64 ConvertTime(const TraceView& view, u64 time)
 	{
-		return time * GetFrequency() / view.frequency;
+		return u64(time * view.timeMultiplier);
 	}
 
 	bool TraceReader::ReadFile(TraceView& out, const tchar* fileName, bool replay)
@@ -142,6 +142,7 @@ namespace uba
 			out.frequency = reader.Read7BitEncoded();
 		else
 			out.frequency = GetFrequency();
+		out.timeMultiplier = double(GetFrequency()) / double(out.frequency);
 
 		out.realStartTime = reader.Read7BitEncoded();
 
@@ -368,6 +369,7 @@ namespace uba
 				out.frequency = traceReader.Read7BitEncoded();
 			else
 				out.frequency = GetFrequency();
+			out.timeMultiplier = double(GetFrequency()) / double(out.frequency);
 
 			out.realStartTime = traceReader.Read7BitEncoded();
 			if (traceSystemStartTimeUs)
