[UBA]
* Added so log entry for failing to create file for write includes read-only attribute if it is set

diff --git a/Engine/Source/Programs/UnrealBuildAccelerator/Common/Private/UbaFileAccessor.cpp b/Engine/Source/Programs/UnrealBuildAccelerator/Common/Private/UbaFileAccessor.cpp
--- a/Engine/Source/Programs/UnrealBuildAccelerator/Common/Private/UbaFileAccessor.cpp
+++ b/Engine/Source/Programs/UnrealBuildAccelerator/Common/Private/UbaFileAccessor.cpp
@@ -116,6 +116,13 @@ namespace uba
 				if (retryCount == 1)
 					GetProcessHoldingFile(additionalInfo, m_fileName);
 
+				// File can also be read-only for whatever reason
+				FileBasicInformation info;
+				if (lastError == ERROR_ACCESS_DENIED && !additionalInfo.count)
+					if (GetFileBasicInformation(info, m_logger, realFileName, false))
+						if (IsReadOnly(info.attributes))
+							additionalInfo.Append(TCV(" - File has read-only attribute set"));
+
 				if (retryCount < 40)
 				{
 					Sleep(500);
@@ -125,7 +132,7 @@ namespace uba
 			}
 			#endif
 			const tchar* retryText = retryCount ? TC(" after retrying for 20 seconds") : TC("");
-			return m_logger.Error(TC("ERROR opening file %s for write%s (%s%s)"), realFileName, retryText, LastErrorToText(lastError).data, additionalInfo.data);
+			return m_logger.Warning(TC("ERROR opening file %s for write%s (%s%s)"), realFileName, retryText, LastErrorToText(lastError).data, additionalInfo.data);
 		}
 
 #if PLATFORM_WINDOWS
