[UBA]
* Fixed so cas files are kept alive until after callback for process exit is called

diff --git a/Engine/Source/Programs/UnrealBuildAccelerator/Common/Private/UbaCacheClient.cpp b/Engine/Source/Programs/UnrealBuildAccelerator/Common/Private/UbaCacheClient.cpp
--- a/Engine/Source/Programs/UnrealBuildAccelerator/Common/Private/UbaCacheClient.cpp
+++ b/Engine/Source/Programs/UnrealBuildAccelerator/Common/Private/UbaCacheClient.cpp
@@ -302,6 +302,13 @@ namespace uba
 
 		TimerScope buildTs(stats.build);
 
+		Vector<StorageImpl::CasEntry*> refCountedCas;
+		auto refCountedCasGuard = MakeGuard([&]()
+			{
+				for (auto casEntry : refCountedCas)
+					m_storage.ReleaseCasEntry(*casEntry, true);
+			});
+
 		// Traverse all inputs and outputs. to create cache entry that we can send to server
 		while (true)
 		{
@@ -430,8 +437,10 @@ namespace uba
 
 				if (isOutput)
 				{
-					if (!m_storage.StoreCasFile(casKey, path.data, CasKeyZero, deferCreation))
+					StorageImpl::CasEntry* casEntry = nullptr;
+					if (!m_storage.StoreCasFile(casKey, path.data, CasKeyZero, deferCreation, &casEntry))
 						return m_logger.Info(TC("Failed to store output file %s when writing to cache"), path.data).ToFalse();
+					refCountedCas.push_back(casEntry);
 				}
 				else
 				{
diff --git a/Engine/Source/Programs/UnrealBuildAccelerator/Common/Private/UbaSessionServer.cpp b/Engine/Source/Programs/UnrealBuildAccelerator/Common/Private/UbaSessionServer.cpp
--- a/Engine/Source/Programs/UnrealBuildAccelerator/Common/Private/UbaSessionServer.cpp
+++ b/Engine/Source/Programs/UnrealBuildAccelerator/Common/Private/UbaSessionServer.cpp
@@ -19,6 +19,7 @@ namespace uba
 		StringKey destinationKey;
 		TString destination;
 		u32 attributes;
+		bool casRefcounted;
 	};
 
 	class SessionServer::RemoteProcess final : public Process
@@ -1294,10 +1295,16 @@ namespace uba
 				++writeFilesTimer.count;
 			}, AsView(TC("WriteReceivedFiles")));
 		--writeFilesTimer.count;
-		process.m_receivedFiles.clear();
 		return writeSuccess;
 	}
 
+	void SessionServer::ReleaseReceivedFiles(const Vector<ReceivedFile>& receivedFiles)
+	{
+		for (auto& file : receivedFiles)
+			if (file.casRefcounted)
+				m_storage.ReleaseRef(file.casKey, true);
+	}
+
 	bool SessionServer::HandleSendFileToServer(const ConnectionInfo& connectionInfo, const WorkContext& workContext, BinaryReader& reader, BinaryWriter& writer)
 	{
 		u32 clientId = connectionInfo.GetId();
@@ -1307,6 +1314,7 @@ namespace uba
 		file.destination = reader.ReadString();
 		file.destinationKey = reader.ReadStringKey();
 		file.attributes = reader.ReadU32();
+		file.casRefcounted = false;
 		CasKey casKey = reader.ReadCasKey();
 
 		Storage::RetrieveResult res;
@@ -1335,6 +1343,9 @@ namespace uba
 			return false;
 		}
 
+		m_storage.AddRef(res.casKey);
+		file.casRefcounted = true;
+
 		auto& process = *(RemoteProcess*)h.m_process;
 		process.m_receivedFiles.push_back(file);
 		return true;
@@ -1671,14 +1682,17 @@ namespace uba
 			process.m_logLines.push_back({ std::move(text), type });
 		}
 
-		if (auto func = process.m_startInfo.logLineFunc)
+		auto& startInfo = process.m_startInfo;
+
+		if (auto func = startInfo.logLineFunc)
 			for (auto& line : process.m_logLines)
-				func(process.m_startInfo.logLineUserData, line.text.c_str(), u32(line.text.size()), line.type);
+				func(startInfo.logLineUserData, line.text.c_str(), u32(line.text.size()), line.type);
 
 		Timer writeFilesTimer;
-		if (exitCode == 0)
+		if (startInfo.writeOutputFilesOnFail || startInfo.rules->IsExitCodeSuccess(exitCode))
 			if (!WriteReceivedFiles(process, writeFilesTimer, connectionInfo))
 				exitCode = UBA_EXIT_CODE(20);
+		auto receivedFiles = std::move(process.m_receivedFiles);
 
 		u32 id = process.m_processId;
 		Vector<ProcessLogLine> emptyLines;
@@ -1695,6 +1709,9 @@ namespace uba
 		process.m_server = nullptr;
 		process.m_done.Set();
 		process.CallProcessExit(h);
+
+		ReleaseReceivedFiles(receivedFiles);
+
 		return true;
 	}
 
@@ -1846,6 +1863,7 @@ namespace uba
 		if (prevExitCode == 0)
 			if (!WriteReceivedFiles(remoteProcess, writeFilesTimer, connectionInfo))
 				prevExitCode = UBA_EXIT_CODE(20);
+		auto receivedFiles = std::move(remoteProcess.m_receivedFiles);
 		//processStats.writeFiles = writeFilesTimer
 
 		SCOPED_FUTEX(remoteProcess.m_exitedLock, exitedLock);
@@ -1856,6 +1874,10 @@ namespace uba
 		bool success = GetNextProcess(remoteProcess, newProcess, nextProcess, prevExitCode, reader);
 		remoteProcess.m_exitCode = ~0u;
 		remoteProcess.m_done.Reset();
+		exitedLock.Leave();
+
+		ReleaseReceivedFiles(receivedFiles);
+
 		if (!success)
 			return false;
 
@@ -2156,7 +2178,11 @@ namespace uba
 			m_trace.ProcessExited(process->m_processId, process->m_exitCode, writer.GetData(), writer.GetPosition(), Vector<ProcessLogLine>());
 		}
 
+		auto receivedFiles = std::move(process->m_receivedFiles);
+
 		process->m_done.Set();
+
+		ReleaseReceivedFiles(receivedFiles);
 	}
 
 	ProcessHandle SessionServer::ProcessRemoved(u32 processId)
diff --git a/Engine/Source/Programs/UnrealBuildAccelerator/Common/Private/UbaStorage.cpp b/Engine/Source/Programs/UnrealBuildAccelerator/Common/Private/UbaStorage.cpp
--- a/Engine/Source/Programs/UnrealBuildAccelerator/Common/Private/UbaStorage.cpp
+++ b/Engine/Source/Programs/UnrealBuildAccelerator/Common/Private/UbaStorage.cpp
@@ -576,7 +576,7 @@ namespace uba
 	}
 
 
-	bool StorageImpl::AddCasFile(StringKey fileNameKey, const tchar* fileName, const CasKey& casKey, bool deferCreation, u64 fileSize, u64 lastWriteTime)
+	bool StorageImpl::AddCasFile(StringKey fileNameKey, const tchar* fileName, const CasKey& casKey, bool deferCreation, u64 fileSize, u64 lastWriteTime, CasEntry** outEntry)
 	{
 		UBA_ASSERTF(IsCompressed(casKey) == m_storeCompressed, TC("CasKey compress mode must match storage compress mode. Compress mode Key: %i, Store: %i Cas: %s (%s)"), (IsCompressed(casKey) ? 1 : 0), (m_storeCompressed ? 1 : 0), CasKeyString(casKey).str, fileName);
 		CasEntry& casEntry = GetOrCreateCasEntry(casKey);
@@ -602,7 +602,7 @@ namespace uba
 		}
 
 		// File is requested by server itself to provide to someone remote. We have to un-drop the cas in case it was marked as dropped
-		if (casEntry.dropped && m_allowUnDropOfCas)
+		if (casEntry.dropped && m_allowUnDropOfCas && !outEntry)
 		{
 			m_logger.Info(TC("Un-dropped cas for %s (%s). Cas will leak in memory"), fileName, BytesToText(fileSize).str);
 			casEntry.dropped = false;
@@ -610,7 +610,14 @@ namespace uba
 
 		if (casEntry.verified)
 			if (casEntry.exists)
+			{
+				if (outEntry)
+				{
+					++casEntry.refCount;
+					*outEntry = &casEntry;
+				}
 				return true;
+			}
 
 		casEntry.disallowed = IsDisallowedPath(fileName);
 
@@ -638,6 +645,13 @@ namespace uba
 				res.first->second.names.push_back(res2.first);
 			else if (res2.first->second != casKey)
 				return m_logger.Error(TC("Same file %s is registered with different caskey (%s vs %s)"), fileName, CasKeyString(res2.first->second).str, CasKeyString(casKey).str);
+
+			if (outEntry)
+			{
+				++casEntry.refCount;
+				*outEntry = &casEntry;
+			}
+
 			return true;
 		}
 		else
@@ -678,6 +692,12 @@ namespace uba
 		casEntry.dropped = res.dropAfterUse;
 		casEntry.exists = true;
 
+		if (outEntry)
+		{
+			++casEntry.refCount;
+			*outEntry = &casEntry;
+		}
+
 		casEntryLock.Leave();
 
 		if (!res.memoryHandle.IsValid())
@@ -1642,6 +1662,34 @@ namespace uba
 		return false;
 	}
 
+	bool StorageImpl::AddRef(const CasKey& casKey)
+	{
+		SCOPED_READ_LOCK(m_casLookupLock, lookupLock);
+		auto it = m_casLookup.find(casKey);
+		if (it == m_casLookup.end())
+			return false;
+		CasEntry& casEntry = it->second;
+		lookupLock.Leave();
+		CasEntryWriteLock casEntryLock(casEntry);
+		UBA_ASSERT(casEntry.verified && casEntry.exists);
+		UBA_ASSERT(casEntry.refCount < 65534);
+		++casEntry.refCount;
+		return true;
+	}
+
+	bool StorageImpl::ReleaseRef(const CasKey& casKey, bool dropIfZero)
+	{
+		SCOPED_READ_LOCK(m_casLookupLock, lookupLock);
+		auto it = m_casLookup.find(casKey);
+		if (it == m_casLookup.end())
+			return false;
+		CasEntry& casEntry = it->second;
+		lookupLock.Leave();
+		CasEntryWriteLock casEntryLock(casEntry);
+		ReleaseCasEntryNoLock(casEntry, dropIfZero);
+		return true;
+	}
+
 	bool StorageImpl::VerifyAndGetCachedFileInfo(CachedFileInfo& out, StringKey fileNameKey, u64 verifiedLastWriteTime, u64 verifiedSize)
 	{
 		out.casKey = CasKeyZero;
@@ -1678,6 +1726,11 @@ namespace uba
 	}
 
 	bool StorageImpl::StoreCasFile(CasKey& out, const tchar* fileName, const CasKey& casKeyOverride, bool deferCreation)
+	{
+		return StoreCasFile(out, fileName, casKeyOverride, deferCreation, nullptr);
+	}
+
+	bool StorageImpl::StoreCasFile(CasKey& out, const tchar* fileName, const CasKey& casKeyOverride, bool deferCreation, CasEntry** outEntry)
 	{
 		StringBuffer<> forKey;
 		forKey.Append(fileName);
@@ -1695,7 +1748,7 @@ namespace uba
 			if (fileEntry.casKey != CasKeyZero)
 			{
 				UBA_ASSERT(casKeyOverride == CasKeyZero || casKeyOverride == fileEntry.casKey);
- 				if (!AddCasFile(fileNameKey, fileName, fileEntry.casKey, deferCreation, fileEntry.size, fileEntry.lastWritten))
+ 				if (!AddCasFile(fileNameKey, fileName, fileEntry.casKey, deferCreation, fileEntry.size, fileEntry.lastWritten, outEntry))
 					return false;
 			}
 			out = fileEntry.casKey;
@@ -1753,14 +1806,14 @@ namespace uba
 			if (casKeyOverride != CasKeyZero && casKeyOverride != fileEntry.casKey)
 			{
 				fileEntry.casKey = casKeyOverride;
-				if (!AddCasFile(fileNameKey, fileName, fileEntry.casKey, deferCreation, storedSize, lastWriteTime))
+				if (!AddCasFile(fileNameKey, fileName, fileEntry.casKey, deferCreation, storedSize, lastWriteTime, outEntry))
 					return false;
 				out = fileEntry.casKey;
 				return true;
 			}
 			if (fileSize == fileEntry.size && (lastWriteTime && lastWriteTime == fileEntry.lastWritten) && fileEntry.casKey != CasKeyInvalid)
 			{
-				if (!AddCasFile(fileNameKey, fileName, fileEntry.casKey, deferCreation, storedSize, lastWriteTime))
+				if (!AddCasFile(fileNameKey, fileName, fileEntry.casKey, deferCreation, storedSize, lastWriteTime, outEntry))
 					return false;
 				out = fileEntry.casKey;
 				return true;
@@ -1812,7 +1865,7 @@ namespace uba
 		if (fileEntry.casKey == CasKeyZero)
 			return false;
 
-		if (!AddCasFile(fileNameKey, fileName, fileEntry.casKey, deferCreation, storedSize, lastWriteTime))
+		if (!AddCasFile(fileNameKey, fileName, fileEntry.casKey, deferCreation, storedSize, lastWriteTime, outEntry))
 			return false;
 
 		out = fileEntry.casKey;
diff --git a/Engine/Source/Programs/UnrealBuildAccelerator/Common/Private/UbaStorageServer.cpp b/Engine/Source/Programs/UnrealBuildAccelerator/Common/Private/UbaStorageServer.cpp
--- a/Engine/Source/Programs/UnrealBuildAccelerator/Common/Private/UbaStorageServer.cpp
+++ b/Engine/Source/Programs/UnrealBuildAccelerator/Common/Private/UbaStorageServer.cpp
@@ -636,7 +636,7 @@ namespace uba
 						}
 
 						bool deferCreation = false;
-						if (!AddCasFile(fileNameKey, hint.data, casKey, deferCreation))
+						if (!AddCasFile(fileNameKey, hint.data, casKey, deferCreation, InvalidValue, 0, nullptr))
 						{
 							m_logger.Error(TC("FetchBegin failed for cas file %s (%s). Can't add cas file to database"), CasKeyString(casKey).str, hint.data);
 							writer.WriteU16(0);
diff --git a/Engine/Source/Programs/UnrealBuildAccelerator/Common/Public/UbaSessionServer.h b/Engine/Source/Programs/UnrealBuildAccelerator/Common/Public/UbaSessionServer.h
--- a/Engine/Source/Programs/UnrealBuildAccelerator/Common/Public/UbaSessionServer.h
+++ b/Engine/Source/Programs/UnrealBuildAccelerator/Common/Public/UbaSessionServer.h
@@ -126,6 +126,7 @@ namespace uba
 		struct ReceivedFile;
 		bool WriteReceivedFile(ReceivedFile& file, RemoteProcess* process, const ConnectionInfo& connectionInfo);
 		bool WriteReceivedFiles(RemoteProcess& process, Timer& writeFilesTimer, const ConnectionInfo& connectionInfo);
+		void ReleaseReceivedFiles(const Vector<ReceivedFile>& receivedFiles);
 
 		void TraceSessionUpdate();
 
diff --git a/Engine/Source/Programs/UnrealBuildAccelerator/Common/Public/UbaStorage.h b/Engine/Source/Programs/UnrealBuildAccelerator/Common/Public/UbaStorage.h
--- a/Engine/Source/Programs/UnrealBuildAccelerator/Common/Public/UbaStorage.h
+++ b/Engine/Source/Programs/UnrealBuildAccelerator/Common/Public/UbaStorage.h
@@ -54,6 +54,9 @@ namespace uba
 		struct RetrieveResult { CasKey casKey; u64 size = 0; MappedView2 view; };
 		virtual bool RetrieveCasFile(RetrieveResult& out, const CasKey& casKey, StringKey keyHint, const tchar* hint, FileMappingBuffer* mappingBuffer = nullptr, u64 memoryMapAlignment = 1, bool allowProxy = true, u32 clientId = 0) = 0;
 
+		virtual bool AddRef(const CasKey& casKey) = 0;
+		virtual bool ReleaseRef(const CasKey& casKey, bool dropIfZero) = 0;
+
 		struct CachedFileInfo { CasKey casKey; };
 		virtual bool VerifyAndGetCachedFileInfo(CachedFileInfo& out, StringKey fileNameKey, u64 verifiedLastWriteTime, u64 verifiedSize) = 0;
 		virtual bool InvalidateCachedFileInfo(StringKey fileNameKey) = 0;
@@ -164,6 +167,8 @@ namespace uba
 		virtual bool GetFileSize(u64& outSize, const CasKey& casKey, const tchar* hint) override;
 		virtual bool GetFileData(void* outData, const CasKey& casKey, const tchar* hint) override;
 		virtual bool RetrieveCasFile(RetrieveResult& out, const CasKey& casKey, StringKey keyHint, const tchar* hint, FileMappingBuffer* mappingBuffer = nullptr, u64 memoryMapAlignment = 1, bool allowProxy = true, u32 clientId = 0) override;
+		virtual bool AddRef(const CasKey& casKey) override;
+		virtual bool ReleaseRef(const CasKey& casKey, bool dropIfZero) override;
 		virtual bool VerifyAndGetCachedFileInfo(CachedFileInfo& out, StringKey fileNameKey, u64 verifiedLastWriteTime, u64 verifiedSize) override;
 		virtual bool InvalidateCachedFileInfo(StringKey fileNameKey) override;
 		virtual bool StoreCasFile(CasKey& out, const tchar* fileName, const CasKey& casKeyOverride, bool deferCreation) override;
@@ -187,7 +192,9 @@ namespace uba
 		struct CasEntry;
 		struct CasEntryWriteLock;
 
+		virtual bool StoreCasFile(CasKey& out, const tchar* fileName, const CasKey& casKeyOverride, bool deferCreation, CasEntry** outEntry);
 		virtual bool HasCasFile(const CasKey& casKey, CasEntry** out = nullptr, bool useMapping = false);
+
 		bool EnsureCasFile(const CasKey& casKey, const tchar* fileName);
 		CasKey CalculateCasKey(const tchar* fileName, FileHandle fileHandle, u64 fileSize, bool storeCompressed);
 		CasKey CalculateCasKey(const u8* fileMem, u64 fileSize, bool storeCompressed);
@@ -209,7 +216,7 @@ namespace uba
 		bool WriteCasFileNoCheck(WriteCasResult& out, const StringKey& fileNameKey, const tchar* fileName, const CasKey& casKey, const tchar* casFile, bool storeCompressed, bool useMapping);
 		bool WriteCasFile(WriteCasResult& out, const tchar* fileName, const CasKey& casKey, bool useMapping);
 		bool VerifyExisting(bool& outReturnValue, CasEntryWriteLock& casEntryLock, const CasKey& casKey, CasEntry& casEntry, StringBufferBase& casFile);
-		bool AddCasFile(StringKey fileNameKey, const tchar* fileName, const CasKey& casKey, bool deferCreation, u64 fileSize = InvalidValue, u64 lastWriteTime = 0);
+		bool AddCasFile(StringKey fileNameKey, const tchar* fileName, const CasKey& casKey, bool deferCreation, u64 fileSize, u64 lastWriteTime, CasEntry** outEntry);
 		void CasEntryAccessed(const CasKey& casKey);
 		virtual bool IsDisallowedPath(const tchar* fileName);
 		virtual bool DecompressMemoryToMemory(const u8* compressedData, u64 compressedSize, u8* writeData, u64 decompressedSize, const tchar* readHint, const tchar* writeHint) override;
