[UBA]
* Fixed so shared memory DuplicateHandle/CloseHandle can handle external handles
* Added alignment of ExtendView to make sure VirtualAlloc is called with proper alignment
* Optimized so trace table is only sent when changed

diff --git a/Engine/Source/Programs/UnrealBuildAccelerator/Common/Private/UbaSharedMemoryAllocator.cpp b/Engine/Source/Programs/UnrealBuildAccelerator/Common/Private/UbaSharedMemoryAllocator.cpp
--- a/Engine/Source/Programs/UnrealBuildAccelerator/Common/Private/UbaSharedMemoryAllocator.cpp
+++ b/Engine/Source/Programs/UnrealBuildAccelerator/Common/Private/UbaSharedMemoryAllocator.cpp
@@ -128,6 +128,8 @@ namespace uba
 	{
 		if (!handle.internalHandle)
 			return {};
+		if (IsExternalMapping(handle))
+			return handle;
 		auto& h = GetHandle(handle);
 		UBA_REF_COUNT(DuplicateHandle);
 		++h.refCount;
@@ -145,8 +147,11 @@ namespace uba
 
 	void SharedMemoryAllocator::CloseHandle(Logger& logger, SharedMemoryHandle handle, const tchar* hint)
 	{
-		if (handle.internalHandle)
-			ReleaseHandle(handle);
+		if (!handle.internalHandle)
+			return;
+		if (IsExternalMapping(handle))
+			return;
+		ReleaseHandle(handle);
 	}
 
 	void SharedMemoryAllocator::ExtendMemory(BinaryWriter& out, SharedMemoryHandle handle, u64 size, const tchar* hint, bool zeroMemory)
@@ -839,6 +844,7 @@ namespace uba
 			if (u64 committed = it->committed.load())
 				handles.emplace(SizeAndOffset{ committed, u64(it) }, it->GetHint());
 		Vector<u32> items;
+		CasKeyHasher hasher;
 		for (auto ri = handles.rbegin(), re = handles.rend(); ri !=re; ++ri)
 		{
 			const tchar* hint = ri->second;
@@ -847,11 +853,16 @@ namespace uba
 			str.Clear().Appendf(TC("%s  %s"), BytesToText(ri->first.size).str, hint);
 			u32 stringId = trace.StringAdd(str);
 			items.push_back(stringId);
+			hasher.Update(&stringId, sizeof(u32));
 		}
 		lock.Leave();
 
-		u32 allocsStringId = trace.StringAdd(TCV("Allocs"));
-		trace.TableUpdate(allocsStringId, 25, 50, items.data(), u32(items.size()));
+		CasKey toSend = ToCasKey(hasher, false);
+		if (m_lastSent == toSend)
+			return;
+		m_lastSent = toSend;
+
+		trace.TableUpdate(trace.StringAdd(TCV("Allocs")), 25, 50, items.data(), u32(items.size()));
 		#endif
 	}
 
diff --git a/Engine/Source/Programs/UnrealBuildAccelerator/Common/Public/UbaSharedMemoryAllocator.h b/Engine/Source/Programs/UnrealBuildAccelerator/Common/Public/UbaSharedMemoryAllocator.h
--- a/Engine/Source/Programs/UnrealBuildAccelerator/Common/Public/UbaSharedMemoryAllocator.h
+++ b/Engine/Source/Programs/UnrealBuildAccelerator/Common/Public/UbaSharedMemoryAllocator.h
@@ -3,8 +3,7 @@
 #pragma once
 
 #include "UbaDefinitions.h"
-#include "UbaMemory.h"
-#include "UbaPlatform.h"
+#include "UbaHash.h"
 #include "UbaSharedMemoryView.h"
 
 // Track shared memory allocations to be able to trace biggest allocations
@@ -103,6 +102,7 @@ namespace uba
 		#if UBA_TRACK_SHARED_MEMORY_ALLOCATIONS
 		InternalHandle* m_firstHandle = nullptr;
 		Futex m_handleLock;
+		CasKey m_lastSent;
 		#endif
 	};
 
diff --git a/Engine/Source/Programs/UnrealBuildAccelerator/Core/Private/UbaSharedMemoryView.cpp b/Engine/Source/Programs/UnrealBuildAccelerator/Core/Private/UbaSharedMemoryView.cpp
--- a/Engine/Source/Programs/UnrealBuildAccelerator/Core/Private/UbaSharedMemoryView.cpp
+++ b/Engine/Source/Programs/UnrealBuildAccelerator/Core/Private/UbaSharedMemoryView.cpp
@@ -185,6 +185,8 @@ namespace uba
 		else if (mapType == SharedMemoryMapType_CopyOnWrite)
 			protect = PAGE_WRITECOPY;
 
+		size = AlignUp(size, 4*1024); // This is needed for VirtualFree to align properly
+
 		if (u64 toRelease = m_capacity - m_mappedSize - size)
 		{
 			TimerScope ts(kernelStats.virtualFree);
