[UBA]
* Added placeholder support. Can be enabled through "WritePlaceholders" under Session. This will make certain files (exe, dll, pdb, etc) be written with zero size. Can be useful to reduce I/O in the case the build is only used for the sake of testing if things can build
* Fixed so ExternalFileMapping now keeps shared memory handle alive after receiving it from the external file mapping provider. This means that things can be released inside session while still being available in storage code
* Fixed so FileMappingEntries are kept alive until after process exit function is called. This is an optimization but also needed when placeholders are used since WriteToCache is called from inside exit callback and we need to make sure file mappings exist until after writetocache is done.

diff --git a/Engine/Source/Programs/UnrealBuildAccelerator/Common/Private/UbaCacheClient.cpp b/Engine/Source/Programs/UnrealBuildAccelerator/Common/Private/UbaCacheClient.cpp
--- a/Engine/Source/Programs/UnrealBuildAccelerator/Common/Private/UbaCacheClient.cpp
+++ b/Engine/Source/Programs/UnrealBuildAccelerator/Common/Private/UbaCacheClient.cpp
@@ -171,6 +171,7 @@ namespace uba
 		m_validateCacheWritesInput = info.validateCacheWritesInput;
 		m_validateCacheWritesOutput = info.validateCacheWritesOutput;
 		m_useCacheHit = info.useCacheHit;
+		m_useMappingForCreatedCas = info.useMappingForCreatedCas;
 		m_useRoots = info.useRoots;
 		m_connected = false;
 
@@ -302,6 +303,8 @@ namespace uba
 
 		TimerScope buildTs(stats.build);
 
+		// This is here to make sure cache client can use dropped cas entries without making them permanent
+		// .. which means that it needs to ref count the cas entries until it is done with them
 		Vector<StorageImpl::CasEntry*> refCountedCas;
 		auto refCountedCasGuard = MakeGuard([&]()
 			{
@@ -440,7 +443,8 @@ namespace uba
 					StorageImpl::CasEntry* casEntry = nullptr;
 					if (!m_storage.StoreCasFile(casKey, path.data, CasKeyZero, deferCreation, &casEntry))
 						return m_logger.Info(TC("Failed to store output file %s when writing to cache"), path.data).ToFalse();
-					refCountedCas.push_back(casEntry);
+					if (casEntry)
+						refCountedCas.push_back(casEntry);
 				}
 				else
 				{
@@ -529,7 +533,7 @@ namespace uba
 			return false;
 
 		// actual cache entry now when we know server has the needed tables
-		if (!SendCacheEntry(tws, bucket, stats, rootPaths, cmdKey, inputsStringToCasKey, outputsStringToCasKey, logLines, logLinesSize))
+		if (!SendCacheEntry(tws, bucket, stats, rootPaths, cmdKey, inputsStringToCasKey, outputsStringToCasKey, logLines, logLinesSize, info))
 			return false;
 
 
@@ -572,7 +576,7 @@ namespace uba
 			return false;
 		if (!outResult.hit)
 			return false;
-		return FetchFilesFromCache(outResult, context);
+		return FetchFilesFromCache(outResult, context, false);
 	}
 
 	bool CacheClient::FetchFromCache(CacheResult& outResult, const RootPaths& rootPaths, u32 bucketId, const ProcessStartInfo& info)
@@ -582,7 +586,7 @@ namespace uba
 			return false;
 		if (!outResult.hit)
 			return false;
-		return FetchFilesFromCache(outResult, context);
+		return FetchFilesFromCache(outResult, context, false);
 	}
 
 	bool CacheClient::FetchEntryFromCache(CacheResult& outResult, FetchContext& context, u32 bucketId)
@@ -956,7 +960,7 @@ namespace uba
 		return false;
 	}
 
-	bool CacheClient::FetchFilesFromCache(CacheResult& outResult, FetchContext& context)
+	bool CacheClient::FetchFilesFromCache(CacheResult& outResult, FetchContext& context, bool writePlaceholders)
 	{
 		UBA_ASSERT(context.bucket);
 		auto& bucket = *context.bucket;
@@ -1000,7 +1004,7 @@ namespace uba
 				KernelStatsScope _(scope.kernelStats);
 				UBA_TRACE_FETCH_HINT_SCOPE("FetchFile")
 				TimerScope fts(scope.cacheStats.fetchOutput);
-				if (!FetchFile(*context.bucket, context.rootPaths, context.info, scope.cacheStats, scope.storageStats, *it))
+				if (!FetchFile(*context.bucket, context.rootPaths, context.info, scope.cacheStats, scope.storageStats, *it, writePlaceholders))
 					fetchSuccess = false;
 
 			}, AsView(TC("CacheFetchOutput")));
@@ -1010,7 +1014,7 @@ namespace uba
 		return success;
 	}
 
-	bool CacheClient::FetchFile(Bucket& bucket, const RootPaths& rootPaths, const ProcessStartInfo& info, CacheFetchStats& cacheStats, StorageStats& storageStats, u32 casKeyOffset)
+	bool CacheClient::FetchFile(Bucket& bucket, const RootPaths& rootPaths, const ProcessStartInfo& info, CacheFetchStats& cacheStats, StorageStats& storageStats, u32 casKeyOffset, bool writePlaceholders)
 	{
 		StringBuffer<MaxPath> path;
 		CasKey casKey;
@@ -1018,6 +1022,16 @@ namespace uba
 			return false;
 		UBA_ASSERT(IsCompressed(casKey));
 
+		if (writePlaceholders)
+		{
+			UBA_ASSERT(info.rules);
+			if (info.rules->SupportsPlaceholder(path))
+			{
+				FileAccessor destinationFile(m_logger, path.data);
+				return destinationFile.CreateWrite() && destinationFile.Close();
+			}
+		}
+
 		FileFetcher fetcher { m_storage.m_bufferSlots, storageStats };
 		fetcher.m_errorOnFail = false;
 
@@ -1286,7 +1300,7 @@ namespace uba
 		return true;
 	}
 
-	bool CacheClient::SendCacheEntry(TrackWorkScope& tws, Bucket& bucket, CacheSendStats& stats, const RootPaths& rootPaths, const CasKey& cmdKey, const Map<u32, u32>& inputsStringToCasKey, const Map<u32, u32>& outputsStringToCasKey, const u8* logLines, u64 logLinesSize)
+	bool CacheClient::SendCacheEntry(TrackWorkScope& tws, Bucket& bucket, CacheSendStats& stats, const RootPaths& rootPaths, const CasKey& cmdKey, const Map<u32, u32>& inputsStringToCasKey, const Map<u32, u32>& outputsStringToCasKey, const u8* logLines, u64 logLinesSize, const ProcessStartInfo& info)
 	{
 		StackBinaryReader<1024> reader;
 
@@ -1336,7 +1350,7 @@ namespace uba
 			TimerScope createCasTs(stats.createCas);
 			StorageImpl::CasEntry* casEntry;
 
-			if (m_storage.HasCasFile(casKey, &casEntry, true))
+			if (m_storage.HasCasFile(casKey, &casEntry, m_useMappingForCreatedCas))
 			{
 				createCasTs.Leave();
 
@@ -1380,6 +1394,13 @@ namespace uba
 					fileSize = file.GetSize();
 				}
 
+				if (fileSize == 0)
+				{
+					UBA_ASSERT(info.rules);
+					if (info.rules->SupportsPlaceholder(path)) // TODO: This should also check if placeholder writing is enabled but not available
+						return m_logger.Warning(TC("File %s has size 0 and supports placeholders.. so it is likely we are sending a placeholder file. Investigate!!"), path.data);
+				}
+
 				if (!SendFile(logger, m_client, &m_storage.m_workManager, casKey, fileData, fileSize, path.data))
 					return false;
 
diff --git a/Engine/Source/Programs/UnrealBuildAccelerator/Common/Private/UbaProcess.cpp b/Engine/Source/Programs/UnrealBuildAccelerator/Common/Private/UbaProcess.cpp
--- a/Engine/Source/Programs/UnrealBuildAccelerator/Common/Private/UbaProcess.cpp
+++ b/Engine/Source/Programs/UnrealBuildAccelerator/Common/Private/UbaProcess.cpp
@@ -645,7 +645,7 @@ namespace uba
 			else
 			{
 				if (m_startInfo.writeOutputFilesOnFail || m_startInfo.rules->IsExitCodeSuccess(m_nativeProcessExitCode))
-					if (!WriteFilesToDisk(true))
+					if (!WriteFilesToDisk())
 						exitCode = UBA_EXIT_CODE(1);
 			}
 
@@ -684,9 +684,10 @@ namespace uba
 			else
 				break;
 
+			// TODO: This should probably do much more cleanup..
+			// This is the retry cleanup if the process crashes
 			m_logLines.clear();
 			m_trackedInputs.clear();
-
 			m_shared.writtenFiles.clear();
 		}
 
@@ -745,6 +746,8 @@ namespace uba
 
 		if (!m_parentProcess)
 		{
+			CleanupWrittenFiles();
+
 			for (auto& memoryHandle : m_shared.sharedMemoryHandles)
 				m_session.m_sharedMemory.CloseHandle(m_session.m_logger, memoryHandle, TC(""));
 			m_shared.sharedMemoryHandles.clear();
@@ -1455,7 +1458,7 @@ namespace uba
 
 	bool ProcessImpl::HandleFlushWrittenFiles(BinaryReader& reader, BinaryWriter& writer)
 	{
-		WriteFilesToDisk(false);
+		WriteFilesToDisk();
 		bool result = m_session.FlushWrittenFiles(*this);
 		writer.WriteBool(result);
 		return true;
@@ -1497,7 +1500,7 @@ namespace uba
 		kernelStats.Write(statsWriter);
 		BinaryReader statsReader(statsWriter.GetData(), 0, statsWriter.GetPosition());
 
-		WriteFilesToDisk(false);
+		WriteFilesToDisk();
 
 		bool newProcess = false;
 		m_exitCode = prevExitCode;
@@ -1507,6 +1510,8 @@ namespace uba
 		if (!newProcess)
 			return true;
 
+		CleanupWrittenFiles();
+
 		m_startInfo.argumentsStr = nextProcess.arguments;
 		m_startInfo.descriptionStr = nextProcess.description;
 		m_startInfo.logFileStr = nextProcess.logFile;
@@ -1629,7 +1634,7 @@ namespace uba
 		return true;
 	}
 
-	bool ProcessImpl::WriteFilesToDisk(bool isExiting)
+	bool ProcessImpl::WriteFilesToDisk()
 	{
 		Vector<WrittenFile*> files;
 		TimerScope ts(m_processStats.writeFiles);
@@ -1664,6 +1669,17 @@ namespace uba
 		return true;
 	}
 
+	void ProcessImpl::CleanupWrittenFiles()
+	{
+		for (auto& kv : m_shared.writtenFiles)
+			if (auto entry = kv.second.entryToReduceRefCount)
+			{
+				kv.second.entryToReduceRefCount = nullptr;
+				SCOPED_FUTEX(entry->lock, entryLock);
+				--entry->refCount;
+			}
+	}
+
 	const tchar* ProcessImpl::InternalGetChildLogFile(StringBufferBase& temp)
 	{
 		if (!*m_startInfo.logFile)
diff --git a/Engine/Source/Programs/UnrealBuildAccelerator/Common/Private/UbaScheduler.cpp b/Engine/Source/Programs/UnrealBuildAccelerator/Common/Private/UbaScheduler.cpp
--- a/Engine/Source/Programs/UnrealBuildAccelerator/Common/Private/UbaScheduler.cpp
+++ b/Engine/Source/Programs/UnrealBuildAccelerator/Common/Private/UbaScheduler.cpp
@@ -873,7 +873,7 @@ namespace uba
 
 					ProcessStartInfo2& si = *info;
 
-					bool success = cacheInfo->client->FetchFilesFromCache(cacheInfo->result, *cacheInfo);
+					bool success = cacheInfo->client->FetchFilesFromCache(cacheInfo->result, *cacheInfo, m_session.WritePlaceholders());
 					--m_activeCacheDownloads;
 
 					if (success)
diff --git a/Engine/Source/Programs/UnrealBuildAccelerator/Common/Private/UbaSession.cpp b/Engine/Source/Programs/UnrealBuildAccelerator/Common/Private/UbaSession.cpp
--- a/Engine/Source/Programs/UnrealBuildAccelerator/Common/Private/UbaSession.cpp
+++ b/Engine/Source/Programs/UnrealBuildAccelerator/Common/Private/UbaSession.cpp
@@ -217,6 +217,7 @@ namespace uba
 		table.GetValueAsBool(allowSpecialApplications, TC("AllowSpecialApplications"));
 		table.GetValueAsBool(suppressLogging, TC("SuppressLogging"));
 		table.GetValueAsBool(shouldWriteToDisk, TC("ShouldWriteToDisk"));
+		table.GetValueAsBool(writePlaceholders, TC("WritePlaceholders"));
 		table.GetValueAsBool(traceEnabled, TC("TraceEnabled"));
 		table.GetValueAsBool(detailedTrace, TC("DetailedTrace"));
 		table.GetValueAsBool(traceChildProcesses, TC("TraceChildProcesses"));
@@ -288,10 +289,12 @@ namespace uba
 			if (!mem)
 				return m_logger.Error(TC("Failed to map view of filehandle for read %s (%s)"), file.name.c_str(), LastErrorToText().data);
 
-			//PrefetchVirtualMemory(mem, fileSize);
+			bool shouldWritePlaceholder = m_writePlaceholders && rules.SupportsPlaceholder(file.name);
 
 			if (storeCompressed)
 			{
+				UBA_ASSERT(!shouldWritePlaceholder);
+
 				Storage::WriteResult res;
 				CompressedFileHeader header { CalculateCasKey(mem, fileSize, true, &m_workManager, file.name.c_str()) };
 
@@ -310,6 +313,8 @@ namespace uba
 
 				if (shouldDevirtualize)
 				{
+					UBA_ASSERT(!shouldWritePlaceholder);
+
 					// Need to turn paths back into local paths
 
 					if (!destinationFile.CreateWrite(false, DefaultAttributes(), 0, m_tempPath.data))
@@ -325,9 +330,16 @@ namespace uba
 				}
 				else
 				{
-					if (!WriteMemoryToDisk(destinationFile, mem, fileSize))
+					u64 writeSize = fileSize;
+					if (shouldWritePlaceholder)
+					{
+						writeSize = 0; // If writing placeholders we are ok with file size 0
+						shouldKeepInMemory = true; // Must keep in memory in case we want to use the cache
+					}
+
+					if (!WriteMemoryToDisk(destinationFile, mem, writeSize))
 						return false;
-					writtenSize = fileSize;
+					writtenSize = writeSize;
 				}
 
 				if (u64 time = file.lastWriteTime)
@@ -395,6 +407,10 @@ namespace uba
 				++m_independentMappingCreated;
 				++m_independentMappingActive;
 
+				// This is used to really make sure file stays in memory until process has called the exit callback
+				file.entryToReduceRefCount = &entry;
+				++entry.refCount;
+
 				if (shouldWriteToDisk) // If writing to disk we only keep file in memory for 3 seconds until we evict it
 				{
 					SCOPED_FUTEX(m_fileMappingsThatCanBeFreedLock, lock);
@@ -1759,9 +1775,9 @@ namespace uba
 
 		m_runningRemote = runningRemote;
 		m_allowCustomAllocator = info.allowCustomAllocator;
-		m_allowMemoryMaps = info.allowMemoryMaps;
-		m_allowKeepFilesInMemory = info.allowKeepFilesInMemory;
-		m_allowOutputFiles = info.allowOutputFiles;
+		m_allowMemoryMaps = info.allowMemoryMaps && IsWindows;
+		m_allowKeepFilesInMemory = info.allowKeepFilesInMemory && IsWindows;
+		m_allowOutputFiles = info.allowOutputFiles && IsWindows;
 		m_allowSpecialApplications = info.allowSpecialApplications;
 		m_suppressLogging = info.suppressLogging;
 		if (!info.allowMemoryMaps)
@@ -1770,6 +1786,7 @@ namespace uba
 			m_keepOutputFileMemoryMapsThreshold = info.keepOutputFileMemoryMapsThreshold;
 		m_shouldWriteToDisk = info.shouldWriteToDisk;
 		UBA_ASSERTF(m_shouldWriteToDisk || m_allowMemoryMaps, TC("Can't disable both should write to disk and allow memory maps"));
+		m_writePlaceholders = info.writePlaceholders && m_allowMemoryMaps;
 
 		m_storeIntermediateFilesCompressed = info.storeIntermediateFilesCompressed && IsWindows; // Non-windows not implemented (yet)
 		m_readIntermediateFilesCompressed = (m_storeIntermediateFilesCompressed || (info.readIntermediateFilesCompressed && IsWindows)) && !runningRemote; // with remote we decompress the files into memory
@@ -1973,14 +1990,12 @@ namespace uba
 					return false;
 				FileMappingEntry& entry = findIt->second;
 				lookupLock.Leave();
+
 				SCOPED_FUTEX_READ(entry.lock, entryLock);
 				if (!entry.handled || !entry.success || !entry.memoryHandle.IsValid())
 					return false;
 				
-				if (entry.canBeFreed) // TODO: Hand over ref counting?
-					return false;
-
-				out.handle = entry.memoryHandle;
+				out.handle = m_sharedMemory.DuplicateHandle(entry.memoryHandle, TC("ExternalFileMapping"));
 				out.offset = entry.memoryOffset;
 				out.size = entry.contentSize;
 				out.lastWriteTime = entry.lastWriteTime;
diff --git a/Engine/Source/Programs/UnrealBuildAccelerator/Common/Private/UbaSessionServer.cpp b/Engine/Source/Programs/UnrealBuildAccelerator/Common/Private/UbaSessionServer.cpp
--- a/Engine/Source/Programs/UnrealBuildAccelerator/Common/Private/UbaSessionServer.cpp
+++ b/Engine/Source/Programs/UnrealBuildAccelerator/Common/Private/UbaSessionServer.cpp
@@ -1237,6 +1237,7 @@ namespace uba
 		if (shouldWriteToDisk)
 		{
 			bool writeCompressed = false;
+			bool shouldWritePlaceholder = m_writePlaceholders && rules.SupportsPlaceholder(destination);
 
 			RootsHandle rootsHandle = startInfo.rootsHandle;
 
@@ -1244,6 +1245,7 @@ namespace uba
 			bool escapeSpaces;
 			if (HasVfs(rootsHandle) && rules.ShouldDevirtualizeFile(destination, escapeSpaces))
 			{
+				UBA_ASSERT(!shouldWritePlaceholder);
 				formattingFunc = [&](MemoryBlock& destData, const void* sourceData, u64 sourceSize, const tchar* hint)
 					{
 						return DevirtualizeDepsFile(rootsHandle, destData, sourceData, sourceSize, escapeSpaces, hint);
@@ -1254,8 +1256,16 @@ namespace uba
 				writeCompressed = g_globalRules.FileCanBeCompressed(destination);
 			}
 
-			if (!m_storage.CopyOrLink(casKey, destination.data, file.attributes, writeCompressed, clientId, formattingFunc))
-				return m_logger.Error(TC("Failed to copy cas from %s to %s (%s)"), CasKeyString(casKey).str, destination.data, startInfo.GetDescription());
+			if (shouldWritePlaceholder)
+			{
+				if (!m_storage.WritePlaceholder(casKey, destination.data, clientId))
+					return m_logger.Error(TC("Failed to write placeholder file for cas from %s to %s (%s)"), CasKeyString(casKey).str, destination.data, startInfo.GetDescription());
+			}
+			else
+			{
+				if (!m_storage.CopyOrLink(casKey, destination.data, file.attributes, writeCompressed, clientId, formattingFunc))
+					return m_logger.Error(TC("Failed to copy cas from %s to %s (%s)"), CasKeyString(casKey).str, destination.data, startInfo.GetDescription());
+			}
 			TraceWrittenFile(file.processId, destination);
 		}
 		else
diff --git a/Engine/Source/Programs/UnrealBuildAccelerator/Common/Private/UbaStorage.cpp b/Engine/Source/Programs/UnrealBuildAccelerator/Common/Private/UbaStorage.cpp
--- a/Engine/Source/Programs/UnrealBuildAccelerator/Common/Private/UbaStorage.cpp
+++ b/Engine/Source/Programs/UnrealBuildAccelerator/Common/Private/UbaStorage.cpp
@@ -363,6 +363,8 @@ namespace uba
 		if (m_externalFileMappingsProvider)
 		{
 			ExternalFileMapping externalMapping;
+			auto externalMappingGuard = MakeGuard([&]() { m_sharedMemory.CloseHandle(m_logger, externalMapping.handle, TC("ExternalFileMapping")); });
+
 			if (m_externalFileMappingsProvider(externalMapping, fileNameKey, fileName))
 			{
 				MemoryMappingScope scope(m_logger, m_sharedMemory);
@@ -370,8 +372,8 @@ namespace uba
 				if (!fileMem)
 					return false;
 				out.sourceStoredSize = externalMapping.storedSize;
-				out.isIndependent = externalMapping.createIndependentMapping;
-				out.dropAfterUse = externalMapping.dropCasAfterUse;
+				out.isIndependent = externalMapping.createIndependentMapping || useMapping;
+				out.dropAfterUse = externalMapping.dropCasAfterUse || useMapping;
 				if (externalMapping.createIndependentMapping)
 				{
 					CompressWriter w(m_logger, m_sharedMemory, m_bufferSlots, m_workManager, Stats(), m_casCompressor, m_casCompressionLevel, m_asyncUnmapViewOfFile);
@@ -1758,6 +1760,7 @@ namespace uba
 
 
 		ExternalFileMapping externalMapping;
+		auto externalMappingGuard = MakeGuard([&]() { m_sharedMemory.CloseHandle(m_logger, externalMapping.handle, TC("ExternalFileMapping")); });
 
 		u64 fileSize = 0;
 		u64 storedSize = 0;
@@ -2300,6 +2303,8 @@ namespace uba
 	bool StorageImpl::CalculateCasKey(CasKey& out, StringKey fileNameKey, const tchar* fileName)
 	{
 		ExternalFileMapping externalMapping;
+		auto externalMappingGuard = MakeGuard([&]() { m_sharedMemory.CloseHandle(m_logger, externalMapping.handle, TC("ExternalFileMapping")); });
+
 		if (m_externalFileMappingsProvider && m_externalFileMappingsProvider(externalMapping, fileNameKey, fileName))
 		{
 			MemoryMappingScope scope(m_logger, m_sharedMemory);
@@ -2593,6 +2598,12 @@ namespace uba
 
 	}
 
+	bool StorageImpl::WritePlaceholder(const CasKey& casKey, const tchar* destination, u32 clientId)
+	{
+		FileAccessor destinationFile(m_logger, destination);
+		return destinationFile.CreateWrite() && destinationFile.Close();
+	}
+
 	bool StorageImpl::FakeCopy(const CasKey& casKey, const tchar* destination, u64 size, u64 lastWritten, u32 fileAttributes, bool deleteExisting)
 	{
 		if (deleteExisting) 
diff --git a/Engine/Source/Programs/UnrealBuildAccelerator/Common/Private/UbaStorageServer.cpp b/Engine/Source/Programs/UnrealBuildAccelerator/Common/Private/UbaStorageServer.cpp
--- a/Engine/Source/Programs/UnrealBuildAccelerator/Common/Private/UbaStorageServer.cpp
+++ b/Engine/Source/Programs/UnrealBuildAccelerator/Common/Private/UbaStorageServer.cpp
@@ -155,6 +155,14 @@ namespace uba
 		return res;
 	}
 
+	bool StorageServer::WritePlaceholder(const CasKey& casKey, const tchar* destination, u32 clientId)
+	{
+		bool res = StorageImpl::WritePlaceholder(casKey, destination, clientId);
+		if (clientId != 0)
+			ReleaseActiveRefCount(casKey, clientId);
+		return res;
+	}
+
 	void StorageServer::ActiveFetch::Release(StorageServer& server, const tchar* reason)
 	{
 		if (ownsMapping)
diff --git a/Engine/Source/Programs/UnrealBuildAccelerator/Common/Public/UbaCacheClient.h b/Engine/Source/Programs/UnrealBuildAccelerator/Common/Public/UbaCacheClient.h
--- a/Engine/Source/Programs/UnrealBuildAccelerator/Common/Public/UbaCacheClient.h
+++ b/Engine/Source/Programs/UnrealBuildAccelerator/Common/Public/UbaCacheClient.h
@@ -49,6 +49,7 @@ namespace uba
 		bool validateCacheWritesOutput = false; // Set to true to validate cas of all output files before sent to cache
 		bool useRoots = true; // Set this to false to allow paths that are not under roots and to not fix them up
 		bool useCacheHit = true; // Set this to false to ignore found cache hits.. this is for debugging/testing only
+		bool useMappingForCreatedCas = true; // Set to this to false to make sure cas created when sending cas are stored on disk
 		const tchar* hint = TC(""); // Hint will show up in cache server log
 	};
 
@@ -103,7 +104,7 @@ namespace uba
 		};
 
 		bool FetchEntryFromCache(CacheResult& outResult, FetchContext& context, u32 bucketId);
-		bool FetchFilesFromCache(CacheResult& outResult, FetchContext& context);
+		bool FetchFilesFromCache(CacheResult& outResult, FetchContext& context, bool writePlaceholders);
 
 	private:
 		struct FetchScope;
@@ -113,11 +114,11 @@ namespace uba
 		// No inlining because they use lots of stack space
 		UBA_NOINLINE bool SendPathTable(Bucket& bucket, CacheSendStats& stats, u32 requiredPathTableSize);
 		UBA_NOINLINE bool SendCasTable(Bucket& bucket, CacheSendStats& stats, u32 requiredCasTableSize);
-		UBA_NOINLINE bool SendCacheEntry(TrackWorkScope& tws, Bucket& bucket, CacheSendStats& stats, const RootPaths& rootPaths, const CasKey& cmdKey, const Map<u32, u32>& inputsStringToCasKey, const Map<u32, u32>& outputsStringToCasKey, const u8* logLines, u64 logLinesSize);
+		UBA_NOINLINE bool SendCacheEntry(TrackWorkScope& tws, Bucket& bucket, CacheSendStats& stats, const RootPaths& rootPaths, const CasKey& cmdKey, const Map<u32, u32>& inputsStringToCasKey, const Map<u32, u32>& outputsStringToCasKey, const u8* logLines, u64 logLinesSize, const ProcessStartInfo& info);
 		UBA_NOINLINE bool SendCacheEntryMessage(BinaryReader& out, Bucket& bucket, CacheSendStats& stats, const CasKey& cmdKey, const Map<u32, u32>& inputsStringToCasKey, const Map<u32, u32>& outputsStringToCasKey, const u8* logLines, u64 logLinesSize);
 		UBA_NOINLINE bool FetchCasTable(TrackWorkScope& tws, Bucket& bucket, CacheFetchStats& stats, u32 requiredCasTableOffset);
 		UBA_NOINLINE bool FetchPathTable(TrackWorkScope& tws, Bucket& bucket, CacheFetchStats& stats, u32 requiredPathTableOffset);
-		UBA_NOINLINE bool FetchFile(Bucket& bucket, const RootPaths& rootPaths, const ProcessStartInfo& info, CacheFetchStats& cacheStats, StorageStats& storageStats, u32 casKeyOffset);
+		UBA_NOINLINE bool FetchFile(Bucket& bucket, const RootPaths& rootPaths, const ProcessStartInfo& info, CacheFetchStats& cacheStats, StorageStats& storageStats, u32 casKeyOffset, bool writePlaceholders);
 		template<typename TableType> bool FetchCompactTable(u32 bucketId, TableType& table, u32 requiredTableOffset, u8 messageType);
 		bool HasEnoughCasData(Bucket& bucket, u32 requiredCasTableOffset);
 		bool HasEnoughPathData(Bucket& bucket, u32 requiredPathTableOffset);
@@ -142,6 +143,7 @@ namespace uba
 		bool m_validateCacheWritesOutput;
 		bool m_useRoots;
 		bool m_useCacheHit;
+		bool m_useMappingForCreatedCas;
 
 		Atomic<bool> m_connected;
 
diff --git a/Engine/Source/Programs/UnrealBuildAccelerator/Common/Public/UbaProcess.h b/Engine/Source/Programs/UnrealBuildAccelerator/Common/Public/UbaProcess.h
--- a/Engine/Source/Programs/UnrealBuildAccelerator/Common/Public/UbaProcess.h
+++ b/Engine/Source/Programs/UnrealBuildAccelerator/Common/Public/UbaProcess.h
@@ -95,7 +95,8 @@ namespace uba
 		void LogLine(bool printInSession, TString&& line, LogEntryType logType);
 		bool CreateTempFile(BinaryReader& reader, ProcHandle nativeProcessHandle, const tchar* application);
 		bool OpenTempFile(BinaryReader& reader, BinaryWriter& writer, const tchar* application);
-		bool WriteFilesToDisk(bool isExiting);
+		bool WriteFilesToDisk();
+		void CleanupWrittenFiles();
 
 		const tchar* InternalGetChildLogFile(StringBufferBase& temp);
 		
diff --git a/Engine/Source/Programs/UnrealBuildAccelerator/Common/Public/UbaSession.h b/Engine/Source/Programs/UnrealBuildAccelerator/Common/Public/UbaSession.h
--- a/Engine/Source/Programs/UnrealBuildAccelerator/Common/Public/UbaSession.h
+++ b/Engine/Source/Programs/UnrealBuildAccelerator/Common/Public/UbaSession.h
@@ -97,6 +97,7 @@ namespace uba
 		const tchar* GetRootDir(); // Root dir for session files
 		const tchar* GetSessionDir(); // Session dir
 		bool ShouldStoreIntermediateFilesCompressed() const { return m_storeIntermediateFilesCompressed; }
+		bool WritePlaceholders() const { return m_writePlaceholders; }
 		bool HasDetailedtrace() const { return m_detailedTrace; }
 
 		// Will write current state of trace directly to disk
@@ -246,6 +247,7 @@ namespace uba
 		bool m_allowSpecialApplications;
 		bool m_suppressLogging;
 		bool m_shouldWriteToDisk;
+		bool m_writePlaceholders;
 		bool m_detailedTrace;
 		bool m_traceChildProcesses;
 		bool m_traceWrittenFiles;
@@ -394,6 +396,7 @@ namespace uba
 		DependencyCrawler m_dependencyCrawler;
 
 		friend class ProcessImpl;
+		friend struct WrittenFile;
 	};
 
 	void GenerateNameForProcess(StringBufferBase& out, const tchar* arguments, u32 counterSuffix);
@@ -608,6 +611,10 @@ namespace uba
 		u64 memoryWritten = 0;
 		u64 lastWriteTime = 0;
 		u32 attributes = 0;
+
+		// TODO: Ugly placement of this pointer.
+		// Here to make sure FileMappingEntry is not garbage collecting before process exit callback is called
+		Session::FileMappingEntry* entryToReduceRefCount = nullptr;
 	};
 
 	struct NextProcessInfo
diff --git a/Engine/Source/Programs/UnrealBuildAccelerator/Common/Public/UbaSessionCreateInfo.h b/Engine/Source/Programs/UnrealBuildAccelerator/Common/Public/UbaSessionCreateInfo.h
--- a/Engine/Source/Programs/UnrealBuildAccelerator/Common/Public/UbaSessionCreateInfo.h
+++ b/Engine/Source/Programs/UnrealBuildAccelerator/Common/Public/UbaSessionCreateInfo.h
@@ -26,12 +26,13 @@ namespace uba
 		bool useUniqueId = true;				// If true, id of session will be "yymmdd_hhmmss". Otherwise "Debug"
 		bool allowCustomAllocator = true;		// Allow detouring of allocator inside processes.
 		bool launchVisualizer = false;			// Launch a UbaVisualizer process (this automatically enable trace)
-		bool allowMemoryMaps = IsWindows;		// Use memory maps where possible. Session creates memory maps of files that processes use
-		bool allowKeepFilesInMemory = IsWindows;// Allow detoured process to keep output files in memory and send them to session through file mapping handle
-		bool allowOutputFiles = IsWindows;		// Allow (selected) output files to be written to disk after process has ended.
+		bool allowMemoryMaps = true;			// Use memory maps where possible. Session creates memory maps of files that processes use
+		bool allowKeepFilesInMemory = true;		// Allow detoured process to keep output files in memory and send them to session through file mapping handle
+		bool allowOutputFiles = true;			// Allow (selected) output files to be written to disk after process has ended.
 		bool allowSpecialApplications = true;	// Allow uba to expand cmd.exe and call known commands instead of running additional process
 		bool suppressLogging = false;			// Suppress all logging produced by detoured processes. Will be made in detoured process to improve performance when needed
 		bool shouldWriteToDisk = true;			// Set to false to skip writing output files to disk
+		bool writePlaceholders = false;			// If this is true, 1-byte placeholder files will be written where possible
 		bool traceEnabled = false;				// Set to true to always create in-memory trace data. Is not needed if traceName, traceOutputFile or launchVisualizer is set
 		bool detailedTrace = false;				// Enable detailed trace to include jobs, individual file I/O etc in trace dump
 		bool traceChildProcesses = false;		// Trace and visualize child processes so they can be seen in visualizer
diff --git a/Engine/Source/Programs/UnrealBuildAccelerator/Common/Public/UbaStorage.h b/Engine/Source/Programs/UnrealBuildAccelerator/Common/Public/UbaStorage.h
--- a/Engine/Source/Programs/UnrealBuildAccelerator/Common/Public/UbaStorage.h
+++ b/Engine/Source/Programs/UnrealBuildAccelerator/Common/Public/UbaStorage.h
@@ -69,6 +69,7 @@ namespace uba
 
 		using FormattingFunc = Function<bool(MemoryBlock& destData, const void* sourceData, u64 sourceSize, const tchar* hint)>;
 		virtual bool CopyOrLink(const CasKey& casKey, const tchar* destination, u32 fileAttributes, bool writeCompressed = false, u32 clientId = 0, const FormattingFunc& formattingFunc = {}, bool isTemp = false, bool allowHardLink = true) = 0;
+		virtual bool WritePlaceholder(const CasKey& casKey, const tchar* destination, u32 clientId) = 0;
 		virtual bool FakeCopy(const CasKey& casKey, const tchar* destination, u64 size, u64 lastWritten, u32 fileAttributes, bool deleteExisting = true) = 0;
 		virtual bool GetCasFileName(StringBufferBase& out, const CasKey& casKey) = 0;
 
@@ -98,7 +99,7 @@ namespace uba
 
 		struct ExternalFileMapping
 		{
-			SharedMemoryHandle handle;
+			SharedMemoryHandle handle; // Will close handle so provider needs to DuplicateHandle when setting it in ExternalFileMapping
 			u64 offset = 0;
 			u64 size = 0;
 			u64 lastWriteTime = 0;
@@ -177,6 +178,7 @@ namespace uba
 		virtual bool ReportBadCasFile(const CasKey& casKey) override;
 		virtual bool CalculateCasKey(CasKey& out, StringKey fileNameKey, const tchar* fileName) override;
 		virtual bool CopyOrLink(const CasKey& casKey, const tchar* destination, u32 fileAttributes, bool writeCompressed = false, u32 clientId = 0, const FormattingFunc& formattingFunc = {}, bool isTemp = false, bool allowHardLink = true) override;
+		virtual bool WritePlaceholder(const CasKey& casKey, const tchar* destination, u32 clientId) override;
 		virtual bool FakeCopy(const CasKey& casKey, const tchar* destination, u64 size, u64 lastWritten, u32 fileAttributes, bool deleteExisting = true) override;
 		virtual bool GetCasFileName(StringBufferBase& out, const CasKey& casKey) override;
 		virtual MappedView2 MapView(const CasKey& casKey, const tchar* hint) override;
diff --git a/Engine/Source/Programs/UnrealBuildAccelerator/Common/Public/UbaStorageServer.h b/Engine/Source/Programs/UnrealBuildAccelerator/Common/Public/UbaStorageServer.h
--- a/Engine/Source/Programs/UnrealBuildAccelerator/Common/Public/UbaStorageServer.h
+++ b/Engine/Source/Programs/UnrealBuildAccelerator/Common/Public/UbaStorageServer.h
@@ -40,6 +40,7 @@ namespace uba
 		bool WaitForCasFile(const CasKey& casKey, const tchar* hint, u32 clientId = 0);
 
 		virtual bool CopyOrLink(const CasKey& casKey, const tchar* destination, u32 fileAttributes, bool writeCompressed = false, u32 clientId = 0, const FormattingFunc& formattingFunc = {}, bool isTemp = false, bool allowHardLink = true) override;
+		virtual bool WritePlaceholder(const CasKey& casKey, const tchar* destination, u32 clientId) override;
 
 	protected:
 
diff --git a/Engine/Source/Programs/UnrealBuildAccelerator/Core/Private/UbaApplicationRules.cpp b/Engine/Source/Programs/UnrealBuildAccelerator/Core/Private/UbaApplicationRules.cpp
--- a/Engine/Source/Programs/UnrealBuildAccelerator/Core/Private/UbaApplicationRules.cpp
+++ b/Engine/Source/Programs/UnrealBuildAccelerator/Core/Private/UbaApplicationRules.cpp
@@ -28,6 +28,14 @@ namespace uba
 
 	GlobalRules g_globalRules;
 
+	
+	bool SupportsLinkerPlaceholder(StringView fileName)
+	{
+		return fileName.EndsWith(TCV(".pdb"))
+			|| fileName.EndsWith(TCV(".exe"))
+			|| fileName.EndsWith(TCV(".dll"));
+	}
+
 
 	class ApplicationRulesDotnet : public ApplicationRules
 	{
@@ -243,6 +251,8 @@ namespace uba
 		{
 			return true;
 		}
+
+		virtual bool SupportsPlaceholder(StringView fileName) const override { return SupportsLinkerPlaceholder(fileName); }
 	};
 
 	class ApplicationRulesLldLinkExe : public ApplicationRulesVcLink
@@ -270,6 +280,8 @@ namespace uba
 		{
 			return true;
 		}
+
+		virtual bool SupportsPlaceholder(StringView fileName) const override { return SupportsLinkerPlaceholder(fileName); }
 	};
 
 	class ApplicationRulesRadLinkExe : public ApplicationRulesVcLink
@@ -298,6 +310,8 @@ namespace uba
 		{
 			return true;
 		}
+
+		virtual bool SupportsPlaceholder(StringView fileName) const override { return SupportsLinkerPlaceholder(fileName); }
 	};
 
 	// ==== Clang tool chain ====
@@ -496,6 +510,8 @@ namespace uba
 		{
 			return true;
 		}
+
+		virtual bool SupportsPlaceholder(StringView fileName) const override { return fileName.EndsWith(TCV(".self")); }
 	};
 
 	class ApplicationRulesClangPlusPlusExePlatform2 : public ApplicationRulesClangPlusPlusExe
@@ -543,6 +559,8 @@ namespace uba
 		{
 			return true;
 		}
+
+		virtual bool SupportsPlaceholder(StringView fileName) const override { return fileName.EndsWith(TCV(".self")); }
 	};
 
 	class ApplicationRulesStageSplitMergeApplication : public ApplicationRules
diff --git a/Engine/Source/Programs/UnrealBuildAccelerator/Core/Public/UbaApplicationRules.h b/Engine/Source/Programs/UnrealBuildAccelerator/Core/Public/UbaApplicationRules.h
--- a/Engine/Source/Programs/UnrealBuildAccelerator/Core/Public/UbaApplicationRules.h
+++ b/Engine/Source/Programs/UnrealBuildAccelerator/Core/Public/UbaApplicationRules.h
@@ -199,6 +199,12 @@ namespace uba
 			return false;
 		}
 
+		// Files that can be written as placeholders if WritePlaceholders is enabled
+		virtual bool SupportsPlaceholder(StringView fileName) const
+		{
+			return false;
+		}
+
 		u32 index = ~0u;
 	};
 
diff --git a/Engine/Source/Programs/UnrealBuildAccelerator/Test/Private/UbaTestCache.cpp b/Engine/Source/Programs/UnrealBuildAccelerator/Test/Private/UbaTestCache.cpp
--- a/Engine/Source/Programs/UnrealBuildAccelerator/Test/Private/UbaTestCache.cpp
+++ b/Engine/Source/Programs/UnrealBuildAccelerator/Test/Private/UbaTestCache.cpp
@@ -580,6 +580,7 @@ namespace uba
 				NetworkClient client(ctorSuccess, { logWriter });
 				CacheClientCreateInfo ccci(logWriter, storageServer, client, session);
 				ccci.useRoots = false;
+				ccci.useMappingForCreatedCas = false; // Needed since cacheclient and server shares same cas and we don't want cas to be evicted
 				CacheClient cacheClient(ccci);
 
 				CHECK_TRUE(client.Connect(tcpBackend, TC("127.0.0.1"), port));
diff --git a/Engine/Source/Programs/UnrealBuildAccelerator/Test/Private/UbaTestStorage.cpp b/Engine/Source/Programs/UnrealBuildAccelerator/Test/Private/UbaTestStorage.cpp
--- a/Engine/Source/Programs/UnrealBuildAccelerator/Test/Private/UbaTestStorage.cpp
+++ b/Engine/Source/Programs/UnrealBuildAccelerator/Test/Private/UbaTestStorage.cpp
@@ -246,7 +246,7 @@ namespace uba
 				{
 					if (fileNameKey != testFileNameKey)
 						return false;
-					out.handle = memHandle;
+					out.handle = sharedMem.DuplicateHandle(memHandle, TC(""));
 					out.size = 1;
 					out.dropCasAfterUse = true;
 					out.createIndependentMapping = true;
