[UBA]
* Fixed bug with shader compiling using virtual files where new code did ref count decrease on file mapping entry after it had been deleted. (Shader compiling code not submitted yet)

diff --git a/Engine/Source/Programs/UnrealBuildAccelerator/Common/Private/UbaSession.cpp b/Engine/Source/Programs/UnrealBuildAccelerator/Common/Private/UbaSession.cpp
--- a/Engine/Source/Programs/UnrealBuildAccelerator/Common/Private/UbaSession.cpp
+++ b/Engine/Source/Programs/UnrealBuildAccelerator/Common/Private/UbaSession.cpp
@@ -390,6 +390,8 @@ namespace uba
 				// Free up unsude memory since we're keeping it around
 				m_sharedMemory.TrimMemory(memoryHandle, file.memoryWritten);
 
+				bool isInvisible = rules.IsInvisible(file.name);
+
 				entry.handled = true;
 				entry.memoryHandle = memoryHandle;
 				entry.memoryOffset = 0;
@@ -401,18 +403,20 @@ namespace uba
 				entry.usedCount = 0;
 				entry.usedCountBeforeFree = g_globalRules.GetUsedCountBeforeFree(file.name);
 				entry.storedSize = writtenSize;
-				entry.isInvisible = rules.IsInvisible(file.name);
+				entry.isInvisible = isInvisible;
 
 				//m_logger.Info(TC("KEEPING %s"), file.name.c_str());
 				++m_independentMappingCreated;
 				++m_independentMappingActive;
 
-				// This is used to really make sure file stays in memory until process has called the exit callback
-				file.entryToReduceRefCount = &entry;
-				++entry.refCount;
-
-				if (shouldWriteToDisk) // If writing to disk we only keep file in memory for 3 seconds until we evict it
+				if (shouldWriteToDisk)
 				{
+					// This is used to really make sure file stays in memory until process has called the exit callback
+					// if file is invisible it should never be 
+					file.entryToReduceRefCount = &entry;
+					++entry.refCount;
+
+					// If writing to disk we only keep file in memory for 3 seconds until we evict it
 					SCOPED_FUTEX(m_fileMappingsThatCanBeFreedLock, lock);
 					m_fileMappingsThatCanBeFreed[file.key] = GetTime() + MsToTime(3*1000);
 				}
@@ -422,7 +426,7 @@ namespace uba
 				m_debugLogger->Info(TC("Mapping kept 0x%llx (%s) from detoured process (UsedCountBeforeFree: %u)"), u64(memoryHandle.ToU64()), entry.name.c_str(), u32(entry.usedCountBeforeFree));
 				#endif
 
-				if (!entry.isInvisible)
+				if (!isInvisible)
 				{
 					StringBuffer<> name;
 					GetMappingString(name, memoryHandle, 0, entry.canBeFreed);
@@ -499,6 +503,7 @@ namespace uba
 			size = entry.contentSize;
 			if (deleteInternalMapping)
 			{
+				UBA_ASSERT(entry.refCount == 0);
 				if (!entry.isInvisible)
 					return m_logger.Error(TC("Trying to delete mapping that is not invisible (%s)"), filePath.data);
 				m_fileMappingTableLookup.erase(findIt);
