[UBA]
* Added detouring of RtlFreeAnsiString since it seems like RtlFreeHeap is not enough to detour when RtlUnicodeStringToAnsiString is used

diff --git a/Engine/Source/Programs/UnrealBuildAccelerator/Detours/Private/Windows/UbaDetoursFunctionsNtDll.inl b/Engine/Source/Programs/UnrealBuildAccelerator/Detours/Private/Windows/UbaDetoursFunctionsNtDll.inl
--- a/Engine/Source/Programs/UnrealBuildAccelerator/Detours/Private/Windows/UbaDetoursFunctionsNtDll.inl
+++ b/Engine/Source/Programs/UnrealBuildAccelerator/Detours/Private/Windows/UbaDetoursFunctionsNtDll.inl
@@ -680,6 +680,18 @@ NTSTATUS Detoured_RtlUnicodeStringToAnsiString(PANSI_STRING DestinationString, P
 	return True_RtlUnicodeStringToAnsiString(DestinationString, SourceString, AllocateDestinationString);
 }
 
+VOID Detoured_RtlFreeAnsiString(PANSI_STRING AnsiString)
+{
+#if UBA_USE_MIMALLOC
+	if (AnsiString && AnsiString->Buffer && g_checkRtlHeap && IsInMiMalloc(AnsiString->Buffer))
+	{
+		mi_free(AnsiString->Buffer);
+		return;
+	}
+#endif
+	return True_RtlFreeAnsiString(AnsiString);
+}
+
 NTSTATUS NTAPI Local_NtCreateFile(bool IsCreateFunc, PHANDLE hFileHandle, ACCESS_MASK DesiredAccess, POBJECT_ATTRIBUTES ObjectAttributes, PIO_STATUS_BLOCK IoStatusBlock, PLARGE_INTEGER AllocationSize, ULONG FileAttributes,
 	ULONG ShareAccess, ULONG CreateDisposition, ULONG CreateOptions, PVOID EaBuffer, ULONG EaLength)
 {
diff --git a/Engine/Source/Programs/UnrealBuildAccelerator/Detours/Private/Windows/UbaDetoursFunctionsWin.h b/Engine/Source/Programs/UnrealBuildAccelerator/Detours/Private/Windows/UbaDetoursFunctionsWin.h
--- a/Engine/Source/Programs/UnrealBuildAccelerator/Detours/Private/Windows/UbaDetoursFunctionsWin.h
+++ b/Engine/Source/Programs/UnrealBuildAccelerator/Detours/Private/Windows/UbaDetoursFunctionsWin.h
@@ -153,6 +153,7 @@
 	DETOURED_FUNCTION(RtlFreeHeap) \
 	DETOURED_FUNCTION(RtlAnsiStringToUnicodeString) \
 	DETOURED_FUNCTION(RtlUnicodeStringToAnsiString) \
+	DETOURED_FUNCTION(RtlFreeAnsiString) \
 	DETOURED_FUNCTIONS_NTDLL_DEBUG \
 
 #define DETOURED_FUNCTIONS_SHLWAPI \
