[UBA]
* Added extra safety to cache client to never fetch files that has size 0 and can be placeholder (because they likely are and it is a bug somewhere)

diff --git a/Engine/Source/Programs/UnrealBuildAccelerator/Common/Private/UbaCacheClient.cpp b/Engine/Source/Programs/UnrealBuildAccelerator/Common/Private/UbaCacheClient.cpp
--- a/Engine/Source/Programs/UnrealBuildAccelerator/Common/Private/UbaCacheClient.cpp
+++ b/Engine/Source/Programs/UnrealBuildAccelerator/Common/Private/UbaCacheClient.cpp
@@ -1126,6 +1126,12 @@ namespace uba
 			bool writeCompressed = m_session.ShouldStoreIntermediateFilesCompressed() && g_globalRules.FileCanBeCompressed(path);
 			if (!fetcher.RetrieveFile(logger, m_client, casKey, path.data, writeCompressed, nullptr, attributes))
 				return logger.Debug(TC("Failed to download cache output %s for %s"), path.data, info.GetDescription()).ToFalse();
+			
+			if (fetcher.sizeOnDisk == 0 && info.rules->SupportsPlaceholder(path))
+			{
+				DeleteFileW(path.data);
+				return logger.Info(TC("Fetched file that seems to be placeholder, this should never happen %s for %s"), path.data, info.GetDescription()).ToFalse();
+			}
 		}
 
 		cacheStats.fetchBytesRaw += fetcher.sizeOnDisk;
diff --git a/Engine/Source/Programs/UnrealBuildAccelerator/Common/Private/UbaStorageUtils.cpp b/Engine/Source/Programs/UnrealBuildAccelerator/Common/Private/UbaStorageUtils.cpp
--- a/Engine/Source/Programs/UnrealBuildAccelerator/Common/Private/UbaStorageUtils.cpp
+++ b/Engine/Source/Programs/UnrealBuildAccelerator/Common/Private/UbaStorageUtils.cpp
@@ -613,7 +613,7 @@ namespace uba
 
 		FileAccessor destinationFile(logger, destination);
 
-		constexpr bool useFileMapping = true;
+		bool useFileMapping = sizeOnDisk > 0;
 		u8* fileMappingMem = nullptr;
 
 		if (!destinationMem)
