[UBA]
* Fixed bug where cache client could use placeholder file to calculate cas key

diff --git a/Engine/Source/Programs/UnrealBuildAccelerator/Common/Private/UbaCacheClient.cpp b/Engine/Source/Programs/UnrealBuildAccelerator/Common/Private/UbaCacheClient.cpp
--- a/Engine/Source/Programs/UnrealBuildAccelerator/Common/Private/UbaCacheClient.cpp
+++ b/Engine/Source/Programs/UnrealBuildAccelerator/Common/Private/UbaCacheClient.cpp
@@ -443,6 +443,7 @@ namespace uba
 					StorageImpl::CasEntry* casEntry = nullptr;
 					if (!m_storage.StoreCasFile(casKey, path.data, CasKeyZero, deferCreation, &casEntry))
 						return m_logger.Info(TC("Failed to store output file %s when writing to cache"), path.data).ToFalse();
+					UBA_ASSERT(casKey == CasKeyZero || casEntry);
 					if (casEntry)
 						refCountedCas.push_back(casEntry);
 				}
diff --git a/Engine/Source/Programs/UnrealBuildAccelerator/Common/Private/UbaSessionServer.cpp b/Engine/Source/Programs/UnrealBuildAccelerator/Common/Private/UbaSessionServer.cpp
--- a/Engine/Source/Programs/UnrealBuildAccelerator/Common/Private/UbaSessionServer.cpp
+++ b/Engine/Source/Programs/UnrealBuildAccelerator/Common/Private/UbaSessionServer.cpp
@@ -1258,7 +1258,7 @@ namespace uba
 
 			if (shouldWritePlaceholder)
 			{
-				if (!m_storage.WritePlaceholder(casKey, destination.data, clientId))
+				if (!m_storage.WritePlaceholder(casKey, destination.data, file.attributes, clientId))
 					return m_logger.Error(TC("Failed to write placeholder file for cas from %s to %s (%s)"), CasKeyString(casKey).str, destination.data, startInfo.GetDescription());
 			}
 			else
diff --git a/Engine/Source/Programs/UnrealBuildAccelerator/Common/Private/UbaStorage.cpp b/Engine/Source/Programs/UnrealBuildAccelerator/Common/Private/UbaStorage.cpp
--- a/Engine/Source/Programs/UnrealBuildAccelerator/Common/Private/UbaStorage.cpp
+++ b/Engine/Source/Programs/UnrealBuildAccelerator/Common/Private/UbaStorage.cpp
@@ -2598,10 +2598,38 @@ namespace uba
 
 	}
 
-	bool StorageImpl::WritePlaceholder(const CasKey& casKey, const tchar* destination, u32 clientId)
+	bool StorageImpl::WritePlaceholder(const CasKey& casKey, const tchar* destination, u32 fileAttributes, u32 clientId)
 	{
+		// We need to populate FileEntry and set the correct caskey even if it is not what we are writing to disk
+		// If we don't do this, then CacheClient will re-calculate the cas from the placeholder file and send wrong keys
+
+		StringBuffer<> forKey;
+		forKey.Append(destination);
+		if (CaseInsensitiveFs)
+			forKey.MakeLower();
+		StringKey key = ToStringKey(forKey);
+		FileEntry& fileEntry = GetOrCreateFileEntry(key);
+		SCOPED_FUTEX(fileEntry.lock, lock2);
+
+		u64 lastWriteTime = 0;
 		FileAccessor destinationFile(m_logger, destination);
-		return destinationFile.CreateWrite() && destinationFile.Close();
+		if (!destinationFile.CreateWrite() || !destinationFile.Close(&lastWriteTime))
+			return false;
+
+		fileEntry.casKey = casKey;
+		fileEntry.lastWritten = lastWriteTime;
+		fileEntry.size = 0;
+		fileEntry.verified = true;
+		fileEntry.isTemp = true;
+
+		#if UBA_TRACK_IS_EXECUTABLE
+		#if PLATFORM_WINDOWS
+		fileEntry.isExecutable = ToView(destination).EndsWith(TCV(".exe"));
+		#else
+		fileEntry.isExecutable = IsExecutable(fileAttributes);
+		#endif
+		#endif
+		return true;
 	}
 
 	bool StorageImpl::FakeCopy(const CasKey& casKey, const tchar* destination, u64 size, u64 lastWritten, u32 fileAttributes, bool deleteExisting)
diff --git a/Engine/Source/Programs/UnrealBuildAccelerator/Common/Private/UbaStorageServer.cpp b/Engine/Source/Programs/UnrealBuildAccelerator/Common/Private/UbaStorageServer.cpp
--- a/Engine/Source/Programs/UnrealBuildAccelerator/Common/Private/UbaStorageServer.cpp
+++ b/Engine/Source/Programs/UnrealBuildAccelerator/Common/Private/UbaStorageServer.cpp
@@ -155,9 +155,9 @@ namespace uba
 		return res;
 	}
 
-	bool StorageServer::WritePlaceholder(const CasKey& casKey, const tchar* destination, u32 clientId)
+	bool StorageServer::WritePlaceholder(const CasKey& casKey, const tchar* destination, u32 fileAttributes, u32 clientId)
 	{
-		bool res = StorageImpl::WritePlaceholder(casKey, destination, clientId);
+		bool res = StorageImpl::WritePlaceholder(casKey, destination, fileAttributes, clientId);
 		if (clientId != 0)
 			ReleaseActiveRefCount(casKey, clientId);
 		return res;
diff --git a/Engine/Source/Programs/UnrealBuildAccelerator/Common/Public/UbaStorage.h b/Engine/Source/Programs/UnrealBuildAccelerator/Common/Public/UbaStorage.h
--- a/Engine/Source/Programs/UnrealBuildAccelerator/Common/Public/UbaStorage.h
+++ b/Engine/Source/Programs/UnrealBuildAccelerator/Common/Public/UbaStorage.h
@@ -69,7 +69,7 @@ namespace uba
 
 		using FormattingFunc = Function<bool(MemoryBlock& destData, const void* sourceData, u64 sourceSize, const tchar* hint)>;
 		virtual bool CopyOrLink(const CasKey& casKey, const tchar* destination, u32 fileAttributes, bool writeCompressed = false, u32 clientId = 0, const FormattingFunc& formattingFunc = {}, bool isTemp = false, bool allowHardLink = true) = 0;
-		virtual bool WritePlaceholder(const CasKey& casKey, const tchar* destination, u32 clientId) = 0;
+		virtual bool WritePlaceholder(const CasKey& casKey, const tchar* destination, u32 fileAttributes, u32 clientId) = 0;
 		virtual bool FakeCopy(const CasKey& casKey, const tchar* destination, u64 size, u64 lastWritten, u32 fileAttributes, bool deleteExisting = true) = 0;
 		virtual bool GetCasFileName(StringBufferBase& out, const CasKey& casKey) = 0;
 
@@ -178,7 +178,7 @@ namespace uba
 		virtual bool ReportBadCasFile(const CasKey& casKey) override;
 		virtual bool CalculateCasKey(CasKey& out, StringKey fileNameKey, const tchar* fileName) override;
 		virtual bool CopyOrLink(const CasKey& casKey, const tchar* destination, u32 fileAttributes, bool writeCompressed = false, u32 clientId = 0, const FormattingFunc& formattingFunc = {}, bool isTemp = false, bool allowHardLink = true) override;
-		virtual bool WritePlaceholder(const CasKey& casKey, const tchar* destination, u32 clientId) override;
+		virtual bool WritePlaceholder(const CasKey& casKey, const tchar* destination, u32 fileAttributes, u32 clientId) override;
 		virtual bool FakeCopy(const CasKey& casKey, const tchar* destination, u64 size, u64 lastWritten, u32 fileAttributes, bool deleteExisting = true) override;
 		virtual bool GetCasFileName(StringBufferBase& out, const CasKey& casKey) override;
 		virtual MappedView2 MapView(const CasKey& casKey, const tchar* hint) override;
diff --git a/Engine/Source/Programs/UnrealBuildAccelerator/Common/Public/UbaStorageServer.h b/Engine/Source/Programs/UnrealBuildAccelerator/Common/Public/UbaStorageServer.h
--- a/Engine/Source/Programs/UnrealBuildAccelerator/Common/Public/UbaStorageServer.h
+++ b/Engine/Source/Programs/UnrealBuildAccelerator/Common/Public/UbaStorageServer.h
@@ -40,7 +40,7 @@ namespace uba
 		bool WaitForCasFile(const CasKey& casKey, const tchar* hint, u32 clientId = 0);
 
 		virtual bool CopyOrLink(const CasKey& casKey, const tchar* destination, u32 fileAttributes, bool writeCompressed = false, u32 clientId = 0, const FormattingFunc& formattingFunc = {}, bool isTemp = false, bool allowHardLink = true) override;
-		virtual bool WritePlaceholder(const CasKey& casKey, const tchar* destination, u32 clientId) override;
+		virtual bool WritePlaceholder(const CasKey& casKey, const tchar* destination, u32 fileAttributes, u32 clientId) override;
 
 	protected:
 
