[UBA]
* Fixed edge case where cas files should be allocated independently
* Removed invalid assert

diff --git a/Engine/Source/Programs/UnrealBuildAccelerator/Common/Private/UbaCacheClient.cpp b/Engine/Source/Programs/UnrealBuildAccelerator/Common/Private/UbaCacheClient.cpp
--- a/Engine/Source/Programs/UnrealBuildAccelerator/Common/Private/UbaCacheClient.cpp
+++ b/Engine/Source/Programs/UnrealBuildAccelerator/Common/Private/UbaCacheClient.cpp
@@ -443,7 +443,6 @@ namespace uba
 					StorageImpl::CasEntry* casEntry = nullptr;
 					if (!m_storage.StoreCasFile(casKey, path.data, CasKeyZero, deferCreation, &casEntry))
 						return m_logger.Info(TC("Failed to store output file %s when writing to cache"), path.data).ToFalse();
-					UBA_ASSERT(casKey == CasKeyZero || casEntry);
 					if (casEntry)
 						refCountedCas.push_back(casEntry);
 				}
diff --git a/Engine/Source/Programs/UnrealBuildAccelerator/Common/Private/UbaStorage.cpp b/Engine/Source/Programs/UnrealBuildAccelerator/Common/Private/UbaStorage.cpp
--- a/Engine/Source/Programs/UnrealBuildAccelerator/Common/Private/UbaStorage.cpp
+++ b/Engine/Source/Programs/UnrealBuildAccelerator/Common/Private/UbaStorage.cpp
@@ -374,7 +374,7 @@ namespace uba
 				out.sourceStoredSize = externalMapping.storedSize;
 				out.isIndependent = externalMapping.createIndependentMapping || useMapping;
 				out.dropAfterUse = externalMapping.dropCasAfterUse || useMapping;
-				if (externalMapping.createIndependentMapping)
+				if (out.isIndependent)
 				{
 					CompressWriter w(m_logger, m_sharedMemory, m_bufferSlots, m_workManager, Stats(), m_casCompressor, m_casCompressionLevel, m_asyncUnmapViewOfFile);
 					out.memoryOffset = 0;
