[UBA]
* Fixed code path where "sourceStoredSize" was not properly set causing warnings in the build

diff --git a/Engine/Source/Programs/UnrealBuildAccelerator/Common/Private/UbaStorage.cpp b/Engine/Source/Programs/UnrealBuildAccelerator/Common/Private/UbaStorage.cpp
--- a/Engine/Source/Programs/UnrealBuildAccelerator/Common/Private/UbaStorage.cpp
+++ b/Engine/Source/Programs/UnrealBuildAccelerator/Common/Private/UbaStorage.cpp
@@ -465,6 +465,8 @@ namespace uba
 				return false;
 			destGuard.Cancel();
 
+			if (readIsCompressed)
+				out.sourceStoredSize = *(u64*)destData;
 			out.memoryHandle = destMapping;
 			out.storedSize = readFileSize;
 			out.memoryOffset = 0;
