[UBA]
* Commented out warning that is broken. It cause problems with compressed vs uncompressed files

diff --git a/Engine/Source/Programs/UnrealBuildAccelerator/Common/Private/UbaStorage.cpp b/Engine/Source/Programs/UnrealBuildAccelerator/Common/Private/UbaStorage.cpp
--- a/Engine/Source/Programs/UnrealBuildAccelerator/Common/Private/UbaStorage.cpp
+++ b/Engine/Source/Programs/UnrealBuildAccelerator/Common/Private/UbaStorage.cpp
@@ -2088,8 +2088,9 @@ namespace uba
 			return m_logger.Error(TC("Failed to write deferred cas %s (%s)"), CasKeyString(casKey).str, deferredCreation.data);
 		casEntry.dropped = res.dropAfterUse;
 
-		if (res.sourceStoredSize != deferredFileSize) // TODO: Should also check last written time
-			m_logger.Warning(TC("Deferred cas %s was deferred with source file size %llu but was read with source file size %llu (%s)"), CasKeyString(casKey).str, deferredFileSize, res.sourceStoredSize, deferredCreation.data);
+		// TODO: THIS MISMATCH WITH compressed vs uncompressed
+		//if (res.sourceStoredSize != deferredFileSize) // TODO: Should also check last written time
+		//	m_logger.Warning(TC("Deferred cas %s was deferred with source file size %llu but was read with source file size %llu (%s)"), CasKeyString(casKey).str, deferredFileSize, res.sourceStoredSize, deferredCreation.data);
 
 		bool hasMapping = false;
 		if (res.memoryHandle.IsValid())
diff --git a/Engine/Source/Programs/UnrealBuildAccelerator/Common/Public/UbaStorage.h b/Engine/Source/Programs/UnrealBuildAccelerator/Common/Public/UbaStorage.h
--- a/Engine/Source/Programs/UnrealBuildAccelerator/Common/Public/UbaStorage.h
+++ b/Engine/Source/Programs/UnrealBuildAccelerator/Common/Public/UbaStorage.h
@@ -103,7 +103,7 @@ namespace uba
 			u64 offset = 0;
 			u64 size = 0;
 			u64 lastWriteTime = 0;
-			u64 storedSize = InvalidValue;
+			u64 storedSize = InvalidValue; // Size stored on disk. Can be compressed or uncompressed
 			bool createIndependentMapping = false;
 			bool dropCasAfterUse = false;
 		};
