[UBA]
* Fixed arm compile warning
* Added comments

diff --git a/Engine/Source/Programs/UnrealBuildAccelerator/Common/Private/UbaStorage.cpp b/Engine/Source/Programs/UnrealBuildAccelerator/Common/Private/UbaStorage.cpp
--- a/Engine/Source/Programs/UnrealBuildAccelerator/Common/Private/UbaStorage.cpp
+++ b/Engine/Source/Programs/UnrealBuildAccelerator/Common/Private/UbaStorage.cpp
@@ -465,8 +465,6 @@ namespace uba
 				return false;
 			destGuard.Cancel();
 
-			if (readIsCompressed)
-				out.sourceStoredSize = *(u64*)destData;
 			out.memoryHandle = destMapping;
 			out.storedSize = readFileSize;
 			out.memoryOffset = 0;
@@ -2068,7 +2066,7 @@ namespace uba
 			return false;
 		StringKey fileNameKey = findIt->second.fileNameKey;
 		StringBuffer<> deferredCreation(findIt->second.fileName);
-		u64 deferredFileSize = findIt->second.fileSize;
+		//u64 deferredFileSize = findIt->second.fileSize;
 		for (auto& it2 : findIt->second.names)
 			m_deferredCasCreationLookupByName.erase(it2);
 		m_deferredCasCreationLookup.erase(findIt);
diff --git a/Engine/Source/Programs/UnrealBuildAccelerator/Common/Public/UbaStorage.h b/Engine/Source/Programs/UnrealBuildAccelerator/Common/Public/UbaStorage.h
--- a/Engine/Source/Programs/UnrealBuildAccelerator/Common/Public/UbaStorage.h
+++ b/Engine/Source/Programs/UnrealBuildAccelerator/Common/Public/UbaStorage.h
@@ -92,7 +92,7 @@ namespace uba
 			SharedMemoryHandle memoryHandle;
 			u64 memoryOffset = InvalidValue;
 			u64 storedSize = InvalidValue; // Stored size of the destination
-			u64 sourceStoredSize = InvalidValue; // Stored size of the source used to write compressed file
+			u64 sourceStoredSize = InvalidValue; // Stored size of the source used to write compressed file. Can be both compressed and uncompressed
 			bool isIndependent = false;
 		};
 		virtual bool WriteCompressedFile(WriteResult& out, const tchar* from, FileHandle readHandle, u8* readMem, u64 fileSize, const tchar* toFile, const void* header, u64 headerSize, u64 lastWriteTime = 0) = 0;
