[UBA]
* Removed open bracket in session title text in visualizer that was accidentally left after a refactoring

diff --git a/Engine/Source/Programs/UnrealBuildAccelerator/Visualizer/Private/UbaTraceReader.cpp b/Engine/Source/Programs/UnrealBuildAccelerator/Visualizer/Private/UbaTraceReader.cpp
--- a/Engine/Source/Programs/UnrealBuildAccelerator/Visualizer/Private/UbaTraceReader.cpp
+++ b/Engine/Source/Programs/UnrealBuildAccelerator/Visualizer/Private/UbaTraceReader.cpp
@@ -33,7 +33,6 @@ namespace uba
 	void TraceView::Session::GetFullName(StringBufferBase& out)
 	{
 		out.Append(name);
-		out.Append(TCV(" ("));
 		u32 infoIndex = 0;
 		for (auto& info : infos)
 		{
