[UBA]
* Removed unused function

diff --git a/Engine/Source/Programs/UnrealBuildAccelerator/Common/Private/UbaFileMapping.cpp b/Engine/Source/Programs/UnrealBuildAccelerator/Common/Private/UbaFileMapping.cpp
--- a/Engine/Source/Programs/UnrealBuildAccelerator/Common/Private/UbaFileMapping.cpp
+++ b/Engine/Source/Programs/UnrealBuildAccelerator/Common/Private/UbaFileMapping.cpp
@@ -384,32 +384,6 @@ namespace uba
 #endif
 	}
 
-	bool DuplicateFileMapping(Logger& logger, ProcHandle sourceProcessHandle, FileMappingHandle source, ProcHandle targetProcessHandle, FileMappingHandle& targetHandle, u32 dwDesiredAccess, bool bInheritHandle, u32 dwOptions, const tchar* hint)
-	{
-#if PLATFORM_WINDOWS
-		if (!DuplicateHandle((HANDLE)sourceProcessHandle, source.mh, (HANDLE)targetProcessHandle, &targetHandle.mh, dwDesiredAccess, bInheritHandle, dwOptions))
-			return false;
-		if (source.fh)
-			if (!DuplicateHandle((HANDLE)sourceProcessHandle, source.fh, (HANDLE)targetProcessHandle, &targetHandle.fh, dwDesiredAccess, bInheritHandle, dwOptions))
-				return false;
-#else
-		UBA_ASSERT(false);
-		if (true)
-			return false;
-#endif
-
-		#if UBA_DEBUG_FILE_MAPPING
-		{
-			SCOPED_FUTEX(g_fileMappingsLock, lock);
-			auto insres = g_fileMappings.try_emplace(targetHandle.mh, hint);(void)insres;
-			UBA_ASSERT(insres.second);
-			//logger.Info(TC("FILEMAPPING DUPLI: %s %llu (%llu)"), hint, targetHandle.mh, g_fileMappings.size());
-		}
-		#endif
-
-		return true;
-	}
-
 	void MapMemoryCopy(void* dest, const void* source, u64 size)
 	{
 		auto& stats = KernelStats::GetCurrent();
diff --git a/Engine/Source/Programs/UnrealBuildAccelerator/Common/Public/UbaFileMapping.h b/Engine/Source/Programs/UnrealBuildAccelerator/Common/Public/UbaFileMapping.h
--- a/Engine/Source/Programs/UnrealBuildAccelerator/Common/Public/UbaFileMapping.h
+++ b/Engine/Source/Programs/UnrealBuildAccelerator/Common/Public/UbaFileMapping.h
@@ -36,7 +36,6 @@ namespace uba
 	bool UnmapViewOfFile(Logger& logger, const void* lpBaseAddress, u64 bytesToUnmap, const tchar* hint, bool allowDiscard = false);
 	void DiscardFileMapping(Logger& logger, FileMappingHandle h, u64 size, const tchar* hint);
 	bool CloseFileMapping(Logger& logger, FileMappingHandle h, const tchar* hint);
-	bool DuplicateFileMapping(Logger& logger, ProcHandle sourceProcessHandle, FileMappingHandle sourceHandle, ProcHandle targetProcessHandle, FileMappingHandle& targetHandle, u32 dwDesiredAccess, bool bInheritHandle, u32 dwOptions, const tchar* hint);
 	void MapMemoryCopy(void* dest, const void* source, u64 size);
 	void SetMemoryWorkingSet(Logger& logger, u64 minWorkingSet);
 	void LockMemory(Logger& logger, void* dest, u64 size, const tchar* hint);
