[UBA]
* Added support for SharedMemory being backed by temp/deleteonclose file. The benefit of this is that it does not use up committed memory which can be useful when doing for example big ltcg linking steps. Downside is that the kernel might write pages to disk even when it doesn't have to.

diff --git a/Engine/Source/Programs/UnrealBuildAccelerator/Common/Private/UbaSharedMemoryAllocator.cpp b/Engine/Source/Programs/UnrealBuildAccelerator/Common/Private/UbaSharedMemoryAllocator.cpp
--- a/Engine/Source/Programs/UnrealBuildAccelerator/Common/Private/UbaSharedMemoryAllocator.cpp
+++ b/Engine/Source/Programs/UnrealBuildAccelerator/Common/Private/UbaSharedMemoryAllocator.cpp
@@ -2,6 +2,7 @@
 
 #include "UbaSharedMemoryAllocator.h"
 #include "UbaBinaryReaderWriter.h"
+#include "UbaFile.h"
 #include "UbaHash.h"
 #include "UbaTrace.h"
 
@@ -52,11 +53,24 @@ namespace uba
 
 	////////////////////////////////////////////////////////////////////////////////////////////////////
 
-	bool SharedMemoryAllocator::Init(u64 capacity, u64 commitStepSize)
+	bool SharedMemoryAllocator::Init(u64 capacity, u64 commitStepSize, const tchar* tempFileForSharedMemory)
 	{
 		m_commitStepSize = Min(commitStepSize, capacity);
 		m_capacity = AlignUp(capacity, PageSize);
-		m_fileMapping = CreateMemoryMappingW(m_logger, PAGE_READWRITE|SEC_RESERVE, m_capacity, nullptr, TC("SharedMemoryAllocator"));
+
+		auto tempFileGuard = MakeGuard([&](){ CloseFile(tempFileForSharedMemory, m_tempFile); });
+
+		if (tempFileForSharedMemory && *tempFileForSharedMemory)
+		{
+			m_tempFile = CreateFileW(tempFileForSharedMemory, GENERIC_READ|GENERIC_WRITE, 0, CREATE_ALWAYS, FILE_ATTRIBUTE_TEMPORARY|FILE_FLAG_DELETE_ON_CLOSE);
+			if (m_tempFile == InvalidFileHandle)
+				return false;
+			m_fileMapping = CreateFileMappingW(m_logger, m_tempFile, PAGE_READWRITE, m_capacity, TC("SharedMemoryAllocator"));
+		}
+		else
+		{
+			m_fileMapping = CreateMemoryMappingW(m_logger, PAGE_READWRITE|SEC_RESERVE, m_capacity, nullptr, TC("SharedMemoryAllocator"));
+		}
 		if (!m_fileMapping.IsValid())
 			return false;
 
@@ -68,6 +82,8 @@ namespace uba
 			CloseFileMapping(m_logger, m_fileMapping, TC("SharedMemoryAllocator"));
 			return false;
 		}
+
+		tempFileGuard.Cancel();
 		return true;
 	}
 
@@ -92,6 +108,11 @@ namespace uba
 		return m_capacity;
 	}
 
+	u64 SharedMemoryAllocator::GetUsed()
+	{
+		return m_used;
+	}
+
 	SharedMemoryHandle SharedMemoryAllocator::CreateHandle(Logger& logger, const tchar* hint, bool priority)
 	{
 		++m_handleCount;
@@ -262,7 +283,8 @@ namespace uba
 
 		UBA_ASSERT(h.slices.size() == 1);
 		auto& slice = h.slices.back();
-		MapViewCommit(m_logger, m_memory + slice.offset + offset, size);
+		if (m_tempFile == InvalidFileHandle)
+			MapViewCommit(m_logger, m_memory + slice.offset + offset, size);
 
 		SCOPED_FUTEX(m_lock, lock);
 		m_committed += size;
@@ -688,8 +710,11 @@ namespace uba
 			// No block is big enough, create new one
 			u64 toCommit = Min(AlignUp(desiredSize, m_commitStepSize), m_capacity);
 			UBA_SHIPPING_ASSERTF(m_reserved + toCommit <= m_capacity, TC("Not enough space in capacity for commit (Capacity: %llu, Reserved: %llu Requested: %llu"), m_capacity, m_reserved, toCommit);
-			bool res = MapViewCommit(m_logger, m_memory + m_reserved, toCommit);
-			UBA_SHIPPING_ASSERTF(res, TC("VirtualAlloc failed (%s)"), LastErrorToText().data);
+			if (m_tempFile == InvalidFileHandle)
+			{
+				bool res = MapViewCommit(m_logger, m_memory + m_reserved, toCommit);
+				UBA_SHIPPING_ASSERTF(res, TC("VirtualAlloc failed (%s)"), LastErrorToText().data);
+			}
 			u64 offset = m_reserved;
 			m_reserved += toCommit;
 			m_committed += toCommit;
diff --git a/Engine/Source/Programs/UnrealBuildAccelerator/Common/Public/UbaSharedMemoryAllocator.h b/Engine/Source/Programs/UnrealBuildAccelerator/Common/Public/UbaSharedMemoryAllocator.h
--- a/Engine/Source/Programs/UnrealBuildAccelerator/Common/Public/UbaSharedMemoryAllocator.h
+++ b/Engine/Source/Programs/UnrealBuildAccelerator/Common/Public/UbaSharedMemoryAllocator.h
@@ -20,11 +20,12 @@ namespace uba
 	class SharedMemoryAllocator
 	{
 	public:
-		bool Init(u64 capacity, u64 commitStepSize = 16ull*1024*1024);
+		bool Init(u64 capacity, u64 commitStepSize = 16ull*1024*1024, const tchar* tempFileForSharedMemory = TC(""));
 		void Deinit(bool keepFileMapping); // keepFileMapping can be used to just leak handle for shutdown
 
 		operator SharedMemoryAllocatorHandle() const;
 		u64 GetCapacity();
+		u64 GetUsed();
 
 		SharedMemoryHandle CreateHandle(Logger& logger, const tchar* hint, bool priority = false); // Priority means it will have a compressable handle id
 		SharedMemoryHandle DuplicateHandle(SharedMemoryHandle handle, const tchar* hint);
@@ -76,13 +77,14 @@ namespace uba
 		Logger& m_logger;
 
 		u64 m_commitStepSize;
+		FileHandle m_tempFile = InvalidFileHandle;
 		FileMappingHandle m_fileMapping;
 		u64 m_capacity = 0; // Capacity is total reserved address space from system
 		u64 m_reserved = 0; // Reserved is the amount of address space used
 		u64 m_committed = 0; // Committed is the amount of pages in address space that has been committed
 		u8* m_memory = nullptr;
 
-		u64 m_used = 0;
+		Atomic<u64> m_used = 0; // Amount of memory used through ExtendMemory
 
 		Futex m_lock;
 
