[UBA]
* Added option to change so shared memory is backed by temporary+deleteonclose file. Added to be able to reduce committed memory charge when doing for example big monolithic links

[Session]
TempFileForSharedMemory=<pathtofile>

diff --git a/Engine/Source/Programs/UnrealBuildAccelerator/Common/Private/UbaStorage.cpp b/Engine/Source/Programs/UnrealBuildAccelerator/Common/Private/UbaStorage.cpp
--- a/Engine/Source/Programs/UnrealBuildAccelerator/Common/Private/UbaStorage.cpp
+++ b/Engine/Source/Programs/UnrealBuildAccelerator/Common/Private/UbaStorage.cpp
@@ -48,6 +48,8 @@ namespace uba
 		table.GetValueAsU32(maxParallelCopyOrLink, TC("MaxParallelCopyOrLink"));
 		table.GetValueAsU32(deleteCasBatchSize, TC("DeleteCasBatchSize"));
 
+		table.GetValueAsString(tempFileForSharedMemory, TC("TempFileForSharedMemory"));
+
 		const tchar* compressor;
 		if (table.GetValueAsString(compressor, TC("Compressor")))
 			casCompressor = GetCompressor(compressor);
@@ -1005,7 +1007,7 @@ namespace uba
 
 		u64 reservedSpace = 1024ull * 1024 * 1024 * 1024;
 		u64 commitStepSize = 64ull * 1024 * 1024;
-		m_sharedMemory.Init(reservedSpace, commitStepSize);
+		m_sharedMemory.Init(reservedSpace, commitStepSize, info.tempFileForSharedMemory);
 	}
 
 	StorageImpl::~StorageImpl()
@@ -2087,6 +2089,10 @@ namespace uba
 		casEntry.dropped = res.dropAfterUse;
 
 		// TODO: THIS MISMATCH WITH compressed vs uncompressed
+		// What can happen is that we use an external memory mapping to calculate the caskey and then defer the creation of the cas
+		// External memory mapping gets written to disk compressed
+		// External memory mapping gets evicted after x seconds because it is not being used
+		// Another agent wants the content of the caskey and now needs to re-read the file which is now compressed.
 		//if (res.sourceStoredSize != deferredFileSize) // TODO: Should also check last written time
 		//	m_logger.Warning(TC("Deferred cas %s was deferred with source file size %llu but was read with source file size %llu (%s)"), CasKeyString(casKey).str, deferredFileSize, res.sourceStoredSize, deferredCreation.data);
 
