[UBA]
* Added option to change so shared memory is backed by temporary+deleteonclose file. Added to be able to reduce committed memory charge when doing for example big monolithic links

[Session]
TempFileForSharedMemory=<pathtofile>

diff --git a/Engine/Source/Programs/UnrealBuildAccelerator/Common/Public/UbaStorage.h b/Engine/Source/Programs/UnrealBuildAccelerator/Common/Public/UbaStorage.h
--- a/Engine/Source/Programs/UnrealBuildAccelerator/Common/Public/UbaStorage.h
+++ b/Engine/Source/Programs/UnrealBuildAccelerator/Common/Public/UbaStorage.h
@@ -132,6 +132,7 @@ namespace uba
 		bool allowDeferredDeletes = true;
 		bool skipCleanupOnShutdown = false;
 		bool allowUnDropOfCas = true;
+		const tchar* tempFileForSharedMemory = TC("");
 		MutexHandle exclusiveMutex = InvalidMutexHandle; // Set this to hand over pre-exclusive access to storage. Storage will release mutex on shutdown
 		u8 casCompressor = 0;
 		u8 casCompressionLevel = 0;
