[UBA]
* Renamed TempFileForSharedMemory to SharedMemoryTempFile
* Added SharedMemoryCapacity and SharedMemoryCommitSize storage settings

diff --git a/Engine/Source/Programs/UnrealBuildAccelerator/Common/Private/UbaSharedMemoryAllocator.cpp b/Engine/Source/Programs/UnrealBuildAccelerator/Common/Private/UbaSharedMemoryAllocator.cpp
--- a/Engine/Source/Programs/UnrealBuildAccelerator/Common/Private/UbaSharedMemoryAllocator.cpp
+++ b/Engine/Source/Programs/UnrealBuildAccelerator/Common/Private/UbaSharedMemoryAllocator.cpp
@@ -53,16 +53,16 @@ namespace uba
 
 	////////////////////////////////////////////////////////////////////////////////////////////////////
 
-	bool SharedMemoryAllocator::Init(u64 capacity, u64 commitStepSize, const tchar* tempFileForSharedMemory)
+	bool SharedMemoryAllocator::Init(u64 capacity, u64 commitStepSize, const tchar* tempFile)
 	{
 		m_commitStepSize = Min(commitStepSize, capacity);
 		m_capacity = AlignUp(capacity, PageSize);
 
-		auto tempFileGuard = MakeGuard([&](){ CloseFile(tempFileForSharedMemory, m_tempFile); });
+		auto tempFileGuard = MakeGuard([&](){ CloseFile(tempFile, m_tempFile); });
 
-		if (tempFileForSharedMemory && *tempFileForSharedMemory)
+		if (tempFile && *tempFile)
 		{
-			m_tempFile = CreateFileW(tempFileForSharedMemory, GENERIC_READ|GENERIC_WRITE, 0, CREATE_ALWAYS, FILE_ATTRIBUTE_TEMPORARY|FILE_FLAG_DELETE_ON_CLOSE);
+			m_tempFile = CreateFileW(tempFile, GENERIC_READ|GENERIC_WRITE, 0, CREATE_ALWAYS, FILE_ATTRIBUTE_TEMPORARY|FILE_FLAG_DELETE_ON_CLOSE);
 			if (m_tempFile == InvalidFileHandle)
 				return false;
 			m_fileMapping = CreateFileMappingW(m_logger, m_tempFile, PAGE_READWRITE, m_capacity, TC("SharedMemoryAllocator"));
diff --git a/Engine/Source/Programs/UnrealBuildAccelerator/Common/Private/UbaStorage.cpp b/Engine/Source/Programs/UnrealBuildAccelerator/Common/Private/UbaStorage.cpp
--- a/Engine/Source/Programs/UnrealBuildAccelerator/Common/Private/UbaStorage.cpp
+++ b/Engine/Source/Programs/UnrealBuildAccelerator/Common/Private/UbaStorage.cpp
@@ -48,7 +48,9 @@ namespace uba
 		table.GetValueAsU32(maxParallelCopyOrLink, TC("MaxParallelCopyOrLink"));
 		table.GetValueAsU32(deleteCasBatchSize, TC("DeleteCasBatchSize"));
 
-		table.GetValueAsString(tempFileForSharedMemory, TC("TempFileForSharedMemory"));
+		table.GetValueAsString(sharedMemoryTempFile, TC("SharedMemoryTempFile"));
+		table.GetValueAsU64(sharedMemoryCapacity, TC("SharedMemoryCapacity"));
+		table.GetValueAsU64(sharedMemoryCommitSize, TC("SharedMemoryCommitSize"));
 
 		const tchar* compressor;
 		if (table.GetValueAsString(compressor, TC("Compressor")))
@@ -1005,9 +1007,7 @@ namespace uba
 		m_casCompressor = info.casCompressor;
 		m_casCompressionLevel = info.casCompressionLevel;
 
-		u64 reservedSpace = 1024ull * 1024 * 1024 * 1024;
-		u64 commitStepSize = 64ull * 1024 * 1024;
-		m_sharedMemory.Init(reservedSpace, commitStepSize, info.tempFileForSharedMemory);
+		m_sharedMemory.Init(info.sharedMemoryCapacity, info.sharedMemoryCommitSize, info.sharedMemoryTempFile);
 	}
 
 	StorageImpl::~StorageImpl()
diff --git a/Engine/Source/Programs/UnrealBuildAccelerator/Common/Public/UbaSharedMemoryAllocator.h b/Engine/Source/Programs/UnrealBuildAccelerator/Common/Public/UbaSharedMemoryAllocator.h
--- a/Engine/Source/Programs/UnrealBuildAccelerator/Common/Public/UbaSharedMemoryAllocator.h
+++ b/Engine/Source/Programs/UnrealBuildAccelerator/Common/Public/UbaSharedMemoryAllocator.h
@@ -20,7 +20,7 @@ namespace uba
 	class SharedMemoryAllocator
 	{
 	public:
-		bool Init(u64 capacity, u64 commitStepSize = 16ull*1024*1024, const tchar* tempFileForSharedMemory = TC(""));
+		bool Init(u64 capacity, u64 commitStepSize = 16ull*1024*1024, const tchar* tempFile = TC(""));
 		void Deinit(bool keepFileMapping); // keepFileMapping can be used to just leak handle for shutdown
 
 		operator SharedMemoryAllocatorHandle() const;
diff --git a/Engine/Source/Programs/UnrealBuildAccelerator/Common/Public/UbaStorage.h b/Engine/Source/Programs/UnrealBuildAccelerator/Common/Public/UbaStorage.h
--- a/Engine/Source/Programs/UnrealBuildAccelerator/Common/Public/UbaStorage.h
+++ b/Engine/Source/Programs/UnrealBuildAccelerator/Common/Public/UbaStorage.h
@@ -132,7 +132,9 @@ namespace uba
 		bool allowDeferredDeletes = true;
 		bool skipCleanupOnShutdown = false;
 		bool allowUnDropOfCas = true;
-		const tchar* tempFileForSharedMemory = TC("");
+		const tchar* sharedMemoryTempFile = TC("");
+		u64 sharedMemoryCapacity = 1024ull*1024*1024*1024; // 1TB
+		u64 sharedMemoryCommitSize = 64ull*1024*1024;
 		MutexHandle exclusiveMutex = InvalidMutexHandle; // Set this to hand over pre-exclusive access to storage. Storage will release mutex on shutdown
 		u8 casCompressor = 0;
 		u8 casCompressionLevel = 0;
