[UBA]
* Fixed non-windows compile error

diff --git a/Engine/Source/Programs/UnrealBuildAccelerator/Common/Private/UbaSharedMemoryAllocator.cpp b/Engine/Source/Programs/UnrealBuildAccelerator/Common/Private/UbaSharedMemoryAllocator.cpp
--- a/Engine/Source/Programs/UnrealBuildAccelerator/Common/Private/UbaSharedMemoryAllocator.cpp
+++ b/Engine/Source/Programs/UnrealBuildAccelerator/Common/Private/UbaSharedMemoryAllocator.cpp
@@ -60,6 +60,7 @@ namespace uba
 
 		auto tempFileGuard = MakeGuard([&](){ CloseFile(tempFile, m_tempFile); });
 
+		#if PLATFORM_WINDOWS
 		if (tempFile && *tempFile)
 		{
 			m_tempFile = CreateFileW(tempFile, GENERIC_READ|GENERIC_WRITE, 0, CREATE_ALWAYS, FILE_ATTRIBUTE_TEMPORARY|FILE_FLAG_DELETE_ON_CLOSE);
@@ -68,6 +69,7 @@ namespace uba
 			m_fileMapping = CreateFileMappingW(m_logger, m_tempFile, PAGE_READWRITE, m_capacity, TC("SharedMemoryAllocator"));
 		}
 		else
+		#endif
 		{
 			m_fileMapping = CreateMemoryMappingW(m_logger, PAGE_READWRITE|SEC_RESERVE, m_capacity, nullptr, TC("SharedMemoryAllocator"));
 		}
