[UBA]
* Added error message when failing to create temp file for shared memory

diff --git a/Engine/Source/Programs/UnrealBuildAccelerator/Common/Private/UbaSharedMemoryAllocator.cpp b/Engine/Source/Programs/UnrealBuildAccelerator/Common/Private/UbaSharedMemoryAllocator.cpp
--- a/Engine/Source/Programs/UnrealBuildAccelerator/Common/Private/UbaSharedMemoryAllocator.cpp
+++ b/Engine/Source/Programs/UnrealBuildAccelerator/Common/Private/UbaSharedMemoryAllocator.cpp
@@ -65,7 +65,7 @@ namespace uba
 		{
 			m_tempFile = CreateFileW(tempFile, GENERIC_READ|GENERIC_WRITE, 0, CREATE_ALWAYS, FILE_ATTRIBUTE_TEMPORARY|FILE_FLAG_DELETE_ON_CLOSE);
 			if (m_tempFile == InvalidFileHandle)
-				return false;
+				return m_logger.Error(TC("Failed to create file %s (%s)"), tempFile, LastErrorToText().data);
 			m_fileMapping = CreateFileMappingW(m_logger, m_tempFile, PAGE_READWRITE, m_capacity, TC("SharedMemoryAllocator"));
 		}
 		else
