[UBA]
* Added support for making shared memory temp file sparse

diff --git a/Engine/Source/Programs/UnrealBuildAccelerator/Common/Private/UbaFile.cpp b/Engine/Source/Programs/UnrealBuildAccelerator/Common/Private/UbaFile.cpp
--- a/Engine/Source/Programs/UnrealBuildAccelerator/Common/Private/UbaFile.cpp
+++ b/Engine/Source/Programs/UnrealBuildAccelerator/Common/Private/UbaFile.cpp
@@ -378,6 +378,19 @@ namespace uba
 #endif
 	}
 
+	bool MakeSparse(Logger& logger, const tchar* fileName, FileHandle h, u64 size)
+	{
+#if PLATFORM_WINDOWS
+		HANDLE handle = asHANDLE(h);
+		DWORD dwTemp;
+		if (!::DeviceIoControl(handle, FSCTL_SET_SPARSE, NULL, 0, NULL, 0, &dwTemp, NULL))
+			return logger.Error(TC("Failed to make file %s sparse (%s)"), fileName, LastErrorToText().data);
+		if (!SetEndOfFile(logger, fileName, h, size))
+			return false;
+#endif
+		return true;
+	}
+
 	bool GetDirectoryOfCurrentModule(Logger& logger, StringBufferBase& out)
 	{
 		#if PLATFORM_WINDOWS
diff --git a/Engine/Source/Programs/UnrealBuildAccelerator/Common/Private/UbaSharedMemoryAllocator.cpp b/Engine/Source/Programs/UnrealBuildAccelerator/Common/Private/UbaSharedMemoryAllocator.cpp
--- a/Engine/Source/Programs/UnrealBuildAccelerator/Common/Private/UbaSharedMemoryAllocator.cpp
+++ b/Engine/Source/Programs/UnrealBuildAccelerator/Common/Private/UbaSharedMemoryAllocator.cpp
@@ -53,7 +53,7 @@ namespace uba
 
 	////////////////////////////////////////////////////////////////////////////////////////////////////
 
-	bool SharedMemoryAllocator::Init(u64 capacity, u64 commitStepSize, const tchar* tempFile)
+	bool SharedMemoryAllocator::Init(u64 capacity, u64 commitStepSize, const tchar* tempFile, bool tempFileSparse)
 	{
 		m_commitStepSize = Min(commitStepSize, capacity);
 		m_capacity = AlignUp(capacity, PageSize);
@@ -66,6 +66,9 @@ namespace uba
 			m_tempFile = CreateFileW(tempFile, GENERIC_READ|GENERIC_WRITE, 0, CREATE_ALWAYS, FILE_ATTRIBUTE_TEMPORARY|FILE_FLAG_DELETE_ON_CLOSE);
 			if (m_tempFile == InvalidFileHandle)
 				return m_logger.Error(TC("Failed to create file %s (%s)"), tempFile, LastErrorToText().data);
+			if (tempFileSparse)
+				if (!MakeSparse(m_logger, tempFile, m_tempFile, m_capacity))
+					return false;
 			m_fileMapping = CreateFileMappingW(m_logger, m_tempFile, PAGE_READWRITE, m_capacity, TC("SharedMemoryAllocator"));
 		}
 		else
diff --git a/Engine/Source/Programs/UnrealBuildAccelerator/Common/Private/UbaStorage.cpp b/Engine/Source/Programs/UnrealBuildAccelerator/Common/Private/UbaStorage.cpp
--- a/Engine/Source/Programs/UnrealBuildAccelerator/Common/Private/UbaStorage.cpp
+++ b/Engine/Source/Programs/UnrealBuildAccelerator/Common/Private/UbaStorage.cpp
@@ -48,9 +48,10 @@ namespace uba
 		table.GetValueAsU32(maxParallelCopyOrLink, TC("MaxParallelCopyOrLink"));
 		table.GetValueAsU32(deleteCasBatchSize, TC("DeleteCasBatchSize"));
 
-		table.GetValueAsString(sharedMemoryTempFile, TC("SharedMemoryTempFile"));
 		table.GetValueAsU64(sharedMemoryCapacity, TC("SharedMemoryCapacity"));
 		table.GetValueAsU64(sharedMemoryCommitSize, TC("SharedMemoryCommitSize"));
+		table.GetValueAsString(sharedMemoryTempFile, TC("SharedMemoryTempFile"));
+		table.GetValueAsBool(sharedMemoryTempFileSparse, TC("SharedMemoryTempFileSparse"));
 
 		const tchar* compressor;
 		if (table.GetValueAsString(compressor, TC("Compressor")))
@@ -1007,7 +1008,7 @@ namespace uba
 		m_casCompressor = info.casCompressor;
 		m_casCompressionLevel = info.casCompressionLevel;
 
-		m_sharedMemory.Init(info.sharedMemoryCapacity, info.sharedMemoryCommitSize, info.sharedMemoryTempFile);
+		m_sharedMemory.Init(info.sharedMemoryCapacity, info.sharedMemoryCommitSize, info.sharedMemoryTempFile, info.sharedMemoryTempFileSparse);
 	}
 
 	StorageImpl::~StorageImpl()
diff --git a/Engine/Source/Programs/UnrealBuildAccelerator/Common/Public/UbaFile.h b/Engine/Source/Programs/UnrealBuildAccelerator/Common/Public/UbaFile.h
--- a/Engine/Source/Programs/UnrealBuildAccelerator/Common/Public/UbaFile.h
+++ b/Engine/Source/Programs/UnrealBuildAccelerator/Common/Public/UbaFile.h
@@ -67,6 +67,7 @@ namespace uba
 	bool FileExists(Logger& logger, const tchar* fileName, u64* outSize = nullptr, u32* outAttributes = nullptr, u64* lastWriteTime = nullptr);
 	bool SetFilePointer(Logger& logger, const tchar* fileName, FileHandle handle, u64 position);
 	bool SetEndOfFile(Logger& logger, const tchar* fileName, FileHandle handle, u64 size);
+	bool MakeSparse(Logger& logger, const tchar* fileName, FileHandle h, u64 size);
 	bool GetDirectoryOfCurrentModule(Logger& logger, StringBufferBase& out);
 	bool DeleteAllFiles(Logger& logger, const tchar* dir, bool deleteDir = true, u32* count = nullptr);
 	bool SearchPathForFile(Logger& logger, StringBufferBase& out, const tchar* file, StringView workingDir, StringView applicationDir);
diff --git a/Engine/Source/Programs/UnrealBuildAccelerator/Common/Public/UbaSharedMemoryAllocator.h b/Engine/Source/Programs/UnrealBuildAccelerator/Common/Public/UbaSharedMemoryAllocator.h
--- a/Engine/Source/Programs/UnrealBuildAccelerator/Common/Public/UbaSharedMemoryAllocator.h
+++ b/Engine/Source/Programs/UnrealBuildAccelerator/Common/Public/UbaSharedMemoryAllocator.h
@@ -20,7 +20,7 @@ namespace uba
 	class SharedMemoryAllocator
 	{
 	public:
-		bool Init(u64 capacity, u64 commitStepSize = 16ull*1024*1024, const tchar* tempFile = TC(""));
+		bool Init(u64 capacity, u64 commitStepSize = 16ull*1024*1024, const tchar* tempFile = TC(""), bool tempFileSparse = true);
 		void Deinit(bool keepFileMapping); // keepFileMapping can be used to just leak handle for shutdown
 
 		operator SharedMemoryAllocatorHandle() const;
diff --git a/Engine/Source/Programs/UnrealBuildAccelerator/Common/Public/UbaStorage.h b/Engine/Source/Programs/UnrealBuildAccelerator/Common/Public/UbaStorage.h
--- a/Engine/Source/Programs/UnrealBuildAccelerator/Common/Public/UbaStorage.h
+++ b/Engine/Source/Programs/UnrealBuildAccelerator/Common/Public/UbaStorage.h
@@ -132,9 +132,10 @@ namespace uba
 		bool allowDeferredDeletes = true;
 		bool skipCleanupOnShutdown = false;
 		bool allowUnDropOfCas = true;
-		const tchar* sharedMemoryTempFile = TC("");
 		u64 sharedMemoryCapacity = 1024ull*1024*1024*1024; // 1TB
 		u64 sharedMemoryCommitSize = 64ull*1024*1024;
+		const tchar* sharedMemoryTempFile = TC("");
+		bool sharedMemoryTempFileSparse = true;
 		MutexHandle exclusiveMutex = InvalidMutexHandle; // Set this to hand over pre-exclusive access to storage. Storage will release mutex on shutdown
 		u8 casCompressor = 0;
 		u8 casCompressionLevel = 0;
