[UBA]
* Changed so memory from shared allocator is not zeroed in cases where entire memory will be overwritten anyway

diff --git a/Engine/Source/Programs/UnrealBuildAccelerator/Common/Private/UbaFileMappingBuffer.cpp b/Engine/Source/Programs/UnrealBuildAccelerator/Common/Private/UbaFileMappingBuffer.cpp
--- a/Engine/Source/Programs/UnrealBuildAccelerator/Common/Private/UbaFileMappingBuffer.cpp
+++ b/Engine/Source/Programs/UnrealBuildAccelerator/Common/Private/UbaFileMappingBuffer.cpp
@@ -57,7 +57,7 @@ namespace uba
 			// TODO: This can come from network with tiny allocations.. kind of sucks to take a 64k causing fragmentation.
 			// Should categorize in small and big blocks.
 			SharedMemoryHandle handle = m_allocator.CreateHandle(m_logger, hint);
-			m_allocator.ExtendMemory(handle, size);
+			m_allocator.ExtendMemory(handle, size, hint, false);
 
 			m_independentSize += size;
 
diff --git a/Engine/Source/Programs/UnrealBuildAccelerator/Common/Private/UbaSession.cpp b/Engine/Source/Programs/UnrealBuildAccelerator/Common/Private/UbaSession.cpp
--- a/Engine/Source/Programs/UnrealBuildAccelerator/Common/Private/UbaSession.cpp
+++ b/Engine/Source/Programs/UnrealBuildAccelerator/Common/Private/UbaSession.cpp
@@ -942,6 +942,39 @@ namespace uba
 			else
 				out.name.Append(TCV("$f"));
 		}
+#if 0
+		else if (isCompressed && fileNameView.EndsWith(TCV(".obj")) && m_sharedMemory.GetUsed() > 16ull*1024*1024*1024)
+		{
+			m_logger.Warning(TC("SharedMemory over threshold, off-loading memory mappings to disk (%s - %llu)"), fileName, contentSize);
+			Guid g;
+			CreateGuid(g);
+			StringBuffer<> tempFileName(m_tempPath);
+			tempFileName.Append(GuidToString(g));
+			FileAccessor tempFile(m_logger, tempFileName.data);
+			if (!tempFile.CreateMemoryWrite(false, FILE_ATTRIBUTE_TEMPORARY|FILE_FLAG_DELETE_ON_CLOSE, contentSize))
+				return false;
+			u8* tempFileMem = tempFile.GetData();
+			if (!m_storage.DecompressFileToMemory(fileName, fileHandle, tempFileMem, contentSize, TC("GetOrCreateMemoryMapFromFile"), fileStartOffset))
+				return false;
+			
+			FileHandle tempFileHandle;
+			if (!DuplicateHandle(GetCurrentProcess(), (HANDLE)tempFile.GetHandle(), GetCurrentProcess(), (HANDLE*)&tempFileHandle, 0, false, DUPLICATE_SAME_ACCESS))
+				return false;
+			if (tempFileHandle == InvalidFileHandle)
+				return false;
+			if (!tempFile.Close())
+				return false;
+			//FileHandle tempFileHandle;
+			//if (!OpenFileSequentialRead(m_logger, tempFileName.data, tempFileHandle))
+			//	return false;
+			FileMappingHandle mappingHandle = CreateFileMappingW(m_logger, tempFileHandle, PAGE_READONLY, contentSize, tempFileName.data);
+			if (!mappingHandle.IsValid())
+				return false;
+			//uba::CloseFile(tempFileName.data, tempFileHandle); // it seems like this can't be closed because then the mapping might go away
+			entry.memoryHandle = m_sharedMemory.RegisterExternalMapping(mappingHandle);
+			GetMappingString(out.name, entry.memoryHandle, 0, false);
+		}
+#endif
 		else
 		{
 			SharedMemoryHandle memoryHandle;
@@ -961,7 +994,7 @@ namespace uba
 				memoryHandle = m_sharedMemory.CreateHandle(m_logger, fileName);
 				if (!memoryHandle.IsValid())
 					return false;
-				m_sharedMemory.ExtendMemory(memoryHandle, contentSize);
+				m_sharedMemory.ExtendMemory(memoryHandle, contentSize, fileName, false);
 				memory = m_sharedMemory.MapView(memoryHandle, fileName, SharedMemoryMapType_ReadWrite);
 				if (!memory)
 					return false;
@@ -1797,7 +1830,7 @@ namespace uba
 		m_readIntermediateFilesCompressed = (m_storeIntermediateFilesCompressed || (info.readIntermediateFilesCompressed && IsWindows)) && !runningRemote; // with remote we decompress the files into memory
 		m_allowLocalDetour = info.allowLocalDetour;
 		m_extractObjFilesSymbols = info.extractObjFilesSymbols;
-		m_allowLinkDependencyCrawler = info.allowLinkDependencyCrawler;
+		m_allowLinkDependencyCrawler = info.allowLinkDependencyCrawler && m_allowMemoryMaps;
 		m_allowDiscardVirtualMemory = info.allowDiscardVirtualMemory;
 
 		m_detailedTrace = info.detailedTrace;
@@ -1944,7 +1977,7 @@ namespace uba
 				{
 					Guid guid;
 					CreateGuid(guid);
-					traceName.Append(GuidToString(guid).str);
+					traceName.Append(GuidToString(guid));
 				}
 			}
 
@@ -4038,7 +4071,7 @@ namespace uba
 		{
 			symHandle = m_sharedMemory.CreateHandle(m_logger, TC("SymHandle"));
 			auto mg = MakeGuard([&]() { m_sharedMemory.CloseHandle(m_logger, symHandle, TC("SymHandle")); });
-			m_sharedMemory.ExtendMemory(symHandle, memoryBlock.writtenSize);
+			m_sharedMemory.ExtendMemory(symHandle, memoryBlock.writtenSize, TC("SymHandle"), false);
 			u8* mem2 = m_sharedMemory.MapView(symHandle, TC("SymHandle"), SharedMemoryMapType_ReadWrite);
 			if (!mem2)
 				return false;
diff --git a/Engine/Source/Programs/UnrealBuildAccelerator/Common/Private/UbaStorage.cpp b/Engine/Source/Programs/UnrealBuildAccelerator/Common/Private/UbaStorage.cpp
--- a/Engine/Source/Programs/UnrealBuildAccelerator/Common/Private/UbaStorage.cpp
+++ b/Engine/Source/Programs/UnrealBuildAccelerator/Common/Private/UbaStorage.cpp
@@ -461,7 +461,7 @@ namespace uba
 		{
 			SharedMemoryHandle destMapping = m_sharedMemory.CreateHandle(m_logger, fileName);
 			auto destGuard = MakeGuard([&]() { m_sharedMemory.CloseHandle(m_logger, destMapping, fileName); });
-			m_sharedMemory.ExtendMemory(destMapping, readFileSize);
+			m_sharedMemory.ExtendMemory(destMapping, readFileSize, fileName, false);
 			u8* destData = m_sharedMemory.MapView(destMapping, fileName, SharedMemoryMapType_ReadWrite);
 			if (!destData)
 				return false;
diff --git a/Engine/Source/Programs/UnrealBuildAccelerator/Common/Private/UbaStorageUtils.cpp b/Engine/Source/Programs/UnrealBuildAccelerator/Common/Private/UbaStorageUtils.cpp
--- a/Engine/Source/Programs/UnrealBuildAccelerator/Common/Private/UbaStorageUtils.cpp
+++ b/Engine/Source/Programs/UnrealBuildAccelerator/Common/Private/UbaStorageUtils.cpp
@@ -1027,7 +1027,7 @@ namespace uba
 		u64 mappingSize = fileSize + 1024;
 
 		SharedMemoryHandle destMapping = m_allocator.CreateHandle(m_logger, hint);
-		m_allocator.ExtendMemory(destMapping, mappingSize);
+		m_allocator.ExtendMemory(destMapping, mappingSize, hint, false);
 		auto destGuard = MakeGuard([&]() { m_allocator.CloseHandle(m_logger, destMapping, hint); });
 		u8* compressedData = m_allocator.MapView(destMapping, hint, SharedMemoryMapType_ReadWrite);
 		if (!compressedData)
@@ -1054,7 +1054,7 @@ namespace uba
 		u64 mappingSize = fileSize + 1024;
 
 		SharedMemoryHandle destMapping = m_allocator.CreateHandle(m_logger, hint);
-		m_allocator.ExtendMemory(destMapping, mappingSize);
+		m_allocator.ExtendMemory(destMapping, mappingSize, hint, false);
 		auto destGuard = MakeGuard([&]() { m_allocator.CloseHandle(m_logger, destMapping, hint); });
 		u8* compressedData = m_allocator.MapView(destMapping, hint, SharedMemoryMapType_ReadWrite);
 		if (!compressedData)
