[UBA]
* Add sequential scan to createfile call when opening for decompression to memory

diff --git a/Engine/Source/Programs/UnrealBuildAccelerator/Common/Private/UbaSession.cpp b/Engine/Source/Programs/UnrealBuildAccelerator/Common/Private/UbaSession.cpp
--- a/Engine/Source/Programs/UnrealBuildAccelerator/Common/Private/UbaSession.cpp
+++ b/Engine/Source/Programs/UnrealBuildAccelerator/Common/Private/UbaSession.cpp
@@ -868,7 +868,7 @@ namespace uba
 		entry.handled = true;
 
 		bool isDir = false;
-		u32 attributes = DefaultAttributes();
+		u32 attributes = DefaultAttributes() | FILE_FLAG_SEQUENTIAL_SCAN;
 		FileHandle fileHandle = uba::CreateFileW(fileName, GENERIC_READ, FILE_SHARE_READ, OPEN_EXISTING, attributes);
 		if (fileHandle == InvalidFileHandle)
 		{
