[UBA]
* Changed so not finding knowninput is not an error

diff --git a/Engine/Source/Programs/UnrealBuildAccelerator/Common/Private/UbaStorageServer.cpp b/Engine/Source/Programs/UnrealBuildAccelerator/Common/Private/UbaStorageServer.cpp
--- a/Engine/Source/Programs/UnrealBuildAccelerator/Common/Private/UbaStorageServer.cpp
+++ b/Engine/Source/Programs/UnrealBuildAccelerator/Common/Private/UbaStorageServer.cpp
@@ -575,10 +575,11 @@ namespace uba
 					keyHint = reader.ReadStringKey();
 
 				reader.ReadString(hint);
+				bool isKnownInput = hint.Equals(TCV("KnownInput"));
 
 				StringBuffer<16> temp;
 				StringView logHint = hint;
-				if (hint.Equals(TCV("KnownInput")))
+				if (isKnownInput)
 				{
 					temp.Append(CasKeyString(casKey).SubStr(0, 10));
 					logHint = temp;
@@ -610,7 +611,18 @@ namespace uba
 
 						StringKey fileNameKey = keyHint;
 						if (keyHint == StringKeyZero)
+						{
+							// If known input without keyHint we have no way of finding cas.
+							// This should be fine though because there will be a later request with the name later on
+							if (isKnownInput)
+							{
+								m_logger.Info(TC("Server did not find cas %s for %s in file table lookup"), CasKeyString(casKey).str, hint.data);
+								writer.WriteU16(0);
+								return true;
+							}
+
 							fileNameKey = CaseInsensitiveFs ? ToStringKeyLower(hint) : ToStringKey(hint);
+						}
 
 						CasKey checkedCasKey;
 						{
