[UBA]
* Improved assert description

diff --git a/Engine/Source/Programs/UnrealBuildAccelerator/Detours/Private/Windows/UbaDetoursFunctionsKernelBase.inl b/Engine/Source/Programs/UnrealBuildAccelerator/Detours/Private/Windows/UbaDetoursFunctionsKernelBase.inl
--- a/Engine/Source/Programs/UnrealBuildAccelerator/Detours/Private/Windows/UbaDetoursFunctionsKernelBase.inl
+++ b/Engine/Source/Programs/UnrealBuildAccelerator/Detours/Private/Windows/UbaDetoursFunctionsKernelBase.inl
@@ -727,8 +727,8 @@ BOOL Detoured_WriteFile(HANDLE hFile, LPCVOID lpBuffer, DWORD nNumberOfBytesToWr
 			DEBUG_LOG_DETOURED(L"WriteFile", L"(MEMORY)%s %llu (%ls) ToWrite: %u -> Success", (lpOverlapped ? L" OVERLAPPED" : L""), uintptr_t(hFile), HandleToName(hFile), nNumberOfBytesToWrite);
 			return TRUE;
 		}
-		UBA_ASSERTF(!fi.isFileMap, L"Trying to write to file %ls which is a filemap. This is not supported\n", HandleToName(hFile));
-		UBA_ASSERTF(dh.trueHandle != INVALID_HANDLE_VALUE, L"Trying to write to file %ls which does not have a valid handle\n", HandleToName(hFile));
+		UBA_ASSERTF(!fi.isFileMap, L"Trying to write to file %ls which is a filemap. This is not supported (%s)", HandleToName(hFile), fi.originalName);
+		UBA_ASSERTF(dh.trueHandle != INVALID_HANDLE_VALUE, L"Trying to write to file %ls which does not have a valid handle (%s)", HandleToName(hFile), fi.originalName);
 		trueHandle = dh.trueHandle;
 	}
 	else if (hFile == PseudoHandle)
