[UBA]
* Added ProcessRace trace type to trace and bumped version
* Added WIP visualization..

diff --git a/Engine/Source/Programs/UnrealBuildAccelerator/Common/Public/UbaTrace.h b/Engine/Source/Programs/UnrealBuildAccelerator/Common/Public/UbaTrace.h
--- a/Engine/Source/Programs/UnrealBuildAccelerator/Common/Public/UbaTrace.h
+++ b/Engine/Source/Programs/UnrealBuildAccelerator/Common/Public/UbaTrace.h
@@ -76,6 +76,7 @@ namespace uba
 		UBA_TRACE_TYPE(StringAdd) \
 		UBA_TRACE_TYPE(StatusUpdate2) \
 		UBA_TRACE_TYPE(TableUpdate) \
+		UBA_TRACE_TYPE(ProcessRace) \
 
 	enum TraceType : u8
 	{
@@ -84,7 +85,7 @@ namespace uba
 		#undef UBA_TRACE_TYPE
 	};
 
-	static constexpr u32 TraceVersion = 50;
+	static constexpr u32 TraceVersion = 51;
 	static constexpr u32 TraceReadCompatibilityVersion = 6;
 
 	////////////////////////////////////////////////////////////////////////////////////////////////////
@@ -128,6 +129,7 @@ namespace uba
 		void ProcessExited(u32 processId, u32 exitCode);
 		void ProcessReturned(u32 processId, StringView reason);
 		void ProcessAddBreadcrumbs(u32 processId, StringView breadcrumbs, bool deleteOld);
+		void ProcessRace(u32 processId, u32 remoteProcess);
 		void ProxyCreated(u32 clientId, const tchar* proxyName);
 		void ProxyUsed(u32 clientId, const tchar* proxyName);
 		void FileFetchLight(u32 clientId, const CasKey& key, u64 fileSize);
diff --git a/Engine/Source/Programs/UnrealBuildAccelerator/Visualizer/Private/UbaTraceReader.cpp b/Engine/Source/Programs/UnrealBuildAccelerator/Visualizer/Private/UbaTraceReader.cpp
--- a/Engine/Source/Programs/UnrealBuildAccelerator/Visualizer/Private/UbaTraceReader.cpp
+++ b/Engine/Source/Programs/UnrealBuildAccelerator/Visualizer/Private/UbaTraceReader.cpp
@@ -882,6 +882,23 @@ namespace uba
 			process.bitmapDirty = true;
 			break;
 		}
+		case TraceType_ProcessRace:
+		{
+			u32 id = reader.ReadU32();
+			u32 remoteId = reader.ReadU32();(void)remoteId;
+
+			auto findIt = m_activeProcesses.find(id);
+			if (findIt != m_activeProcesses.end())
+			{
+				TraceView::ProcessLocation& active = findIt->second;
+				auto& session = GetSession(out, active.sessionIndex);
+
+				auto& processes = session.processors[active.processorIndex].processes;
+				TraceView::Process& process = processes[active.processIndex];
+				process.isRacing = true;
+			}
+			break;
+		}
 		case TraceType_FileFetchBegin:
 		{
 			Guid clientUid = ReadClientId(out, reader);
diff --git a/Engine/Source/Programs/UnrealBuildAccelerator/Visualizer/Private/UbaTraceReader.h b/Engine/Source/Programs/UnrealBuildAccelerator/Visualizer/Private/UbaTraceReader.h
--- a/Engine/Source/Programs/UnrealBuildAccelerator/Visualizer/Private/UbaTraceReader.h
+++ b/Engine/Source/Programs/UnrealBuildAccelerator/Visualizer/Private/UbaTraceReader.h
@@ -30,6 +30,7 @@ namespace uba
 			bool bitmapDirty = true;
 			bool isRemote = false;
 			bool isReuse = false;
+			bool isRacing = false;
 			u64 createFilesTime = 0;
 			u64 writeFilesTime = 0;
 			Vector<u8> stats;
diff --git a/Engine/Source/Programs/UnrealBuildAccelerator/Visualizer/Private/UbaVisualizer.cpp b/Engine/Source/Programs/UnrealBuildAccelerator/Visualizer/Private/UbaVisualizer.cpp
--- a/Engine/Source/Programs/UnrealBuildAccelerator/Visualizer/Private/UbaVisualizer.cpp
+++ b/Engine/Source/Programs/UnrealBuildAccelerator/Visualizer/Private/UbaVisualizer.cpp
@@ -480,6 +480,8 @@ namespace uba
 
 			m_processBrushes[0].inProgress = CreateSolidBrush(RGB(70, 70, 70));
 			m_processBrushes[1].inProgress = CreateSolidBrush(RGB(130, 130, 130));
+			m_processBrushes[0].inProgressRace = CreateSolidBrush(RGB(50, 50, 50));
+			m_processBrushes[1].inProgressRace = CreateSolidBrush(RGB(120, 120, 120));
 
 			m_processBrushes[0].error = CreateSolidBrush(RGB(140, 0, 0));
 			m_processBrushes[1].error = CreateSolidBrush(RGB(190, 0, 0));
@@ -520,6 +522,8 @@ namespace uba
 
 			m_processBrushes[0].inProgress = CreateSolidBrush(RGB(150, 150, 150));
 			m_processBrushes[1].inProgress = CreateSolidBrush(RGB(180, 180, 180));
+			m_processBrushes[0].inProgressRace = CreateSolidBrush(RGB(130, 130, 130));
+			m_processBrushes[1].inProgressRace = CreateSolidBrush(RGB(170, 170, 170));
 
 			m_processBrushes[0].error = CreateSolidBrush(RGB(255, 70, 70));
 			m_processBrushes[1].error = CreateSolidBrush(RGB(255, 100, 70));
@@ -2606,7 +2610,7 @@ namespace uba
 		if (!process.returnedReason.empty())
 			brush = m_processBrushes[selected].returned;
 		else if (!done)
-			brush = m_processBrushes[selected].inProgress;
+			brush = process.isRacing ? m_processBrushes[selected].inProgressRace : m_processBrushes[selected].inProgress;
 		else if (process.type == TraceView::ProcessType_CacheFetchTest)
 			brush = m_processBrushes[selected].cacheFetchTest;
 		else if (process.type == TraceView::ProcessType_CacheFetchDownload)
diff --git a/Engine/Source/Programs/UnrealBuildAccelerator/Visualizer/Private/UbaVisualizer.h b/Engine/Source/Programs/UnrealBuildAccelerator/Visualizer/Private/UbaVisualizer.h
--- a/Engine/Source/Programs/UnrealBuildAccelerator/Visualizer/Private/UbaVisualizer.h
+++ b/Engine/Source/Programs/UnrealBuildAccelerator/Visualizer/Private/UbaVisualizer.h
@@ -183,6 +183,7 @@ namespace uba
 		struct ProcessBrushes
 		{
 			HBRUSH inProgress = 0;
+			HBRUSH inProgressRace = 0;
 			HBRUSH success = 0;
 			HBRUSH error = 0;
 			HBRUSH returned = 0;
