[UBA]
* Added ref count assert check in Storage ReleaseCasEntryNoLock to catch negative ref counts

diff --git a/Engine/Source/Programs/UnrealBuildAccelerator/Common/Private/UbaStorage.cpp b/Engine/Source/Programs/UnrealBuildAccelerator/Common/Private/UbaStorage.cpp
--- a/Engine/Source/Programs/UnrealBuildAccelerator/Common/Private/UbaStorage.cpp
+++ b/Engine/Source/Programs/UnrealBuildAccelerator/Common/Private/UbaStorage.cpp
@@ -1919,6 +1919,7 @@ namespace uba
 
 	void StorageImpl::ReleaseCasEntryNoLock(CasEntry& casEntry, bool dropIfZero)
 	{
+		UBA_ASSERT(casEntry.refCount > 0);
 		if (--casEntry.refCount)
 			return;
 		if (!casEntry.dropped || !dropIfZero)
