[UBA]
* Added commented out debug lines to UbaMimalloc that enables double deletes, stomps, etc

diff --git a/Engine/Source/Programs/UnrealBuildAccelerator/Mimalloc/UbaMimalloc.Build.cs b/Engine/Source/Programs/UnrealBuildAccelerator/Mimalloc/UbaMimalloc.Build.cs
--- a/Engine/Source/Programs/UnrealBuildAccelerator/Mimalloc/UbaMimalloc.Build.cs
+++ b/Engine/Source/Programs/UnrealBuildAccelerator/Mimalloc/UbaMimalloc.Build.cs
@@ -19,9 +19,8 @@ public class UbaMimalloc : ModuleRules
 			});
 		}
 
-		//PublicDefinitions.AddRange(new string[] {
-		//	"UBA_USE_MIMALLOC=1", // compile in mimalloc
-		//});
+		// Enable this to enable allocator tracking... double deletes etc.
+		// PublicDefinitions.AddRange(["MI_DEBUG=3", "MI_SECURE=4"]);
 
 		if (Target.LinkType == TargetLinkType.Modular)
 		{
