[UBA]
* Added ProcessRace trace type to trace and bumped version
* Added WIP visualization..

diff --git a/Engine/Source/Programs/UnrealBuildAccelerator/Common/Private/UbaTrace.cpp b/Engine/Source/Programs/UnrealBuildAccelerator/Common/Private/UbaTrace.cpp
--- a/Engine/Source/Programs/UnrealBuildAccelerator/Common/Private/UbaTrace.cpp
+++ b/Engine/Source/Programs/UnrealBuildAccelerator/Common/Private/UbaTrace.cpp
@@ -486,6 +486,13 @@ namespace uba
 		writer.WriteBool(deleteOld);
 	}
 
+	void Trace::ProcessRace(u32 processId, u32 remoteProcess)
+	{
+		BEGIN_TRACE_ENTRY(ProcessRace);
+		writer.WriteU32(processId);
+		writer.WriteU32(remoteProcess);
+	}
+
 	void Trace::ProxyCreated(u32 clientId, const tchar* proxyName)
 	{
 		BEGIN_TRACE_ENTRY(ProxyCreated);
