[UBA]
* Added support for userdata overrides in ProcessImpl

diff --git a/Engine/Source/Programs/UnrealBuildAccelerator/Common/Private/UbaProcess.cpp b/Engine/Source/Programs/UnrealBuildAccelerator/Common/Private/UbaProcess.cpp
--- a/Engine/Source/Programs/UnrealBuildAccelerator/Common/Private/UbaProcess.cpp
+++ b/Engine/Source/Programs/UnrealBuildAccelerator/Common/Private/UbaProcess.cpp
@@ -210,11 +210,12 @@ namespace uba
 		}
 	}
 
-	bool ProcessImpl::Start(const ProcessStartInfo& startInfo, bool runningRemote, void* environment, bool async)
+	bool ProcessImpl::Start(const ProcessStartInfo& startInfo, bool runningRemote, void* environment, bool async, void* userDataOverride)
 	{
 		m_startTime = GetTime();
 
 		m_startInfo = startInfo;
+		m_userDataOverride = userDataOverride;
 		m_runningRemote = runningRemote;
 
 		FixPathSeparators(m_startInfo.workingDirStr.data());
@@ -269,6 +270,8 @@ namespace uba
 			if (!res)
 			{
 				sharedMemory.CloseHandle(m_session.m_logger, memHandle, TC("OverlayFiles"));
+				m_exitCode = 44325;
+				CallExitedFunc();
 				return false;
 			}
 
@@ -514,7 +517,7 @@ namespace uba
 	{
 		SCOPED_FUTEX(m_logLinesLock, l); // Abuse exiting mutex to not waste memory
 		if (!m_cancelAllowed)
-			return false;
+			return m_cancelled;
 		m_cancelled = true;
 		#if PLATFORM_WINDOWS
 		m_cancelEvent.Set();
@@ -774,9 +777,9 @@ namespace uba
 		if (!m_startInfo.exitedFunc)
 			return;
 		auto exitedFunc = m_startInfo.exitedFunc;
-		auto userData = m_startInfo.userData;
+		auto userData = m_userDataOverride ? m_userDataOverride : m_startInfo.userData;
+		m_userDataOverride = nullptr;
 		m_startInfo.exitedFunc = nullptr;
-		m_startInfo.userData = nullptr;
 		ProcessHandle h;
 		h.m_process = this;
 		ProcessExitedResponse exitedResponse = ProcessExitedResponse_None;
@@ -1271,7 +1274,7 @@ namespace uba
 		info.rootsHandle = m_startInfo.rootsHandle;
 
 		ProcessImpl* parent = isChild ? this : nullptr;
-		ProcessHandle h = m_session.InternalRunProcess(info, true, parent, true);
+		ProcessHandle h = m_session.InternalRunProcess(info, true, parent, true, nullptr);
 		if (!h.m_process)
 		{
 			// TODO: Do we need to log something here? This failure is most likely a cause of something else so error logging might show somewhere else
diff --git a/Engine/Source/Programs/UnrealBuildAccelerator/Common/Private/UbaSession.cpp b/Engine/Source/Programs/UnrealBuildAccelerator/Common/Private/UbaSession.cpp
--- a/Engine/Source/Programs/UnrealBuildAccelerator/Common/Private/UbaSession.cpp
+++ b/Engine/Source/Programs/UnrealBuildAccelerator/Common/Private/UbaSession.cpp
@@ -180,6 +180,7 @@ namespace uba
 		UBA_ASSERT(m_process);
 		return m_process->IsRemote();
 	}
+
 	ProcessExecutionType ProcessHandle::GetExecutionType() const
 	{
 		UBA_ASSERT(m_process);
@@ -2168,7 +2169,7 @@ namespace uba
 		FlushDeadProcesses();
 		ValidateStartInfo(startInfo);
 		enableDetour &= m_allowLocalDetour;
-		return InternalRunProcess(startInfo, async, nullptr, enableDetour);
+		return InternalRunProcess(startInfo, async, nullptr, enableDetour, startInfo.userData);
 	}
 
 	void Session::ValidateStartInfo(const ProcessStartInfo& startInfo)
@@ -2177,7 +2178,7 @@ namespace uba
 		UBA_ASSERTF(!TStrchr(startInfo.workingDir, '~'), TC("WorkingDir path must use long name (%s)"), startInfo.workingDir);
 	}
 
-	ProcessHandle Session::InternalRunProcess(const ProcessStartInfo& startInfo, bool async, ProcessImpl* parent, bool enableDetour)
+	ProcessHandle Session::InternalRunProcess(const ProcessStartInfo& startInfo, bool async, ProcessImpl* parent, bool enableDetour, void* userData)
 	{
 		auto& si = const_cast<ProcessStartInfo&>(startInfo);
 		const tchar* originalLogFile = si.logFile;
@@ -2207,7 +2208,7 @@ namespace uba
 		void* env = GetProcessEnvironmentVariables();
 		auto process = new ProcessImpl(*this, processId, parent, enableDetour);
 		ProcessHandle h(process);
-		if (!process->Start(startInfo, m_runningRemote, env, async))
+		if (!process->Start(startInfo, m_runningRemote, env, async, userData))
 			h = {};
 
 		si.logFile = originalLogFile;
@@ -2647,7 +2648,7 @@ namespace uba
 
 		SCOPED_FUTEX(m_processesLock, lock);
 		bool success = m_processes.try_emplace(processId, ProcessHandle(&process)).second;
-		UBA_ASSERT(success);(void)success;
+		UBA_ASSERTF(success, TC("Failed to add process with id %u. It already exists (%s)"), processId, startInfo.GetDescription());(void)success;
 	}
 
 	void Session::ProcessExited(ProcessImpl& process, u64 executionTime)
diff --git a/Engine/Source/Programs/UnrealBuildAccelerator/Common/Private/UbaSessionClient.cpp b/Engine/Source/Programs/UnrealBuildAccelerator/Common/Private/UbaSessionClient.cpp
--- a/Engine/Source/Programs/UnrealBuildAccelerator/Common/Private/UbaSessionClient.cpp
+++ b/Engine/Source/Programs/UnrealBuildAccelerator/Common/Private/UbaSessionClient.cpp
@@ -2205,7 +2205,7 @@ namespace uba
 						process->m_shared.overlayFilesSize = u32(startInfo.overlayFilesData.size());
 					}
 
-					if (!process->Start(startInfo, true, env, true)) // If false, exitFunction is not called
+					if (!process->Start(startInfo, true, env, true, exitedRec)) // If false, exitFunction is not called
 					{
 						SendReturnProcess(rec->handle.GetId(), TC("Failed to find executable"));
 						activeProcesses.pop_back();
diff --git a/Engine/Source/Programs/UnrealBuildAccelerator/Common/Public/UbaProcess.h b/Engine/Source/Programs/UnrealBuildAccelerator/Common/Public/UbaProcess.h
--- a/Engine/Source/Programs/UnrealBuildAccelerator/Common/Public/UbaProcess.h
+++ b/Engine/Source/Programs/UnrealBuildAccelerator/Common/Public/UbaProcess.h
@@ -71,7 +71,7 @@ namespace uba
 
 		struct PipeReader;
 
-		bool Start(const ProcessStartInfo& startInfo, bool runningRemote, void* environment, bool async);
+		bool Start(const ProcessStartInfo& startInfo, bool runningRemote, void* environment, bool async, void* userDataOverride);
 
 		bool IsActive();
 		bool IsCancelled();
@@ -119,6 +119,7 @@ namespace uba
 		Futex m_initLock;
 		u32 m_id;
 		FileMappingAllocation m_comMemory;
+		void* m_userDataOverride = nullptr;
 
 	#if PLATFORM_WINDOWS
 		SharedEvent m_cancelEvent;
diff --git a/Engine/Source/Programs/UnrealBuildAccelerator/Common/Public/UbaSession.h b/Engine/Source/Programs/UnrealBuildAccelerator/Common/Public/UbaSession.h
--- a/Engine/Source/Programs/UnrealBuildAccelerator/Common/Public/UbaSession.h
+++ b/Engine/Source/Programs/UnrealBuildAccelerator/Common/Public/UbaSession.h
@@ -125,7 +125,7 @@ namespace uba
 		bool Create(const SessionCreateInfo& info);
 
 		void ValidateStartInfo(const ProcessStartInfo& startInfo);
-		ProcessHandle InternalRunProcess(const ProcessStartInfo& startInfo, bool async, ProcessImpl* parent, bool enableDetour);
+		ProcessHandle InternalRunProcess(const ProcessStartInfo& startInfo, bool async, ProcessImpl* parent, bool enableDetour, void* userData);
 		void ProcessAdded(Process& process, u32 sessionId);
 		void ProcessExited(ProcessImpl& process, u64 executionTime);
 		void PrintProcessStats(ProcessStats& stats, const tchar* logName);
