[UBA]
* Fixed double delete in UbaSessionClient when process failed to start (because of disconnect for example)... this must have created a bunch of shutdown crashes in UbaAgent

diff --git a/Engine/Source/Programs/UnrealBuildAccelerator/Common/Private/UbaSessionClient.cpp b/Engine/Source/Programs/UnrealBuildAccelerator/Common/Private/UbaSessionClient.cpp
--- a/Engine/Source/Programs/UnrealBuildAccelerator/Common/Private/UbaSessionClient.cpp
+++ b/Engine/Source/Programs/UnrealBuildAccelerator/Common/Private/UbaSessionClient.cpp
@@ -1508,6 +1508,7 @@ namespace uba
 			}
 		}
 
+
 		while (true)
 		{
 			u32 processId = reader.ReadU32();
@@ -1957,6 +1958,7 @@ namespace uba
 
 		Vector<InternalProcessStartInfo> startInfos;
 
+
 		while (m_loop)
 		{
 			float maxWeight = float(m_maxProcessCount);
@@ -2044,6 +2046,10 @@ namespace uba
 				if (!m_remoteExecutionEnabled)
 				{
 					m_logger.Info(TC("Got remote execution disabled response from host (will finish %llu active processes)"), startInfos.size() + activeProcesses.size());
+
+
+						// We must always wait for cancelled processes to really finish since we might get new processes from host
+						// which depend on this process and then outputs might collide.
 				}
 
 				if (startInfos.empty())
@@ -2209,7 +2215,6 @@ namespace uba
 					{
 						SendReturnProcess(rec->handle.GetId(), TC("Failed to find executable"));
 						activeProcesses.pop_back();
-						delete exitedRec;
 						m_remoteExecutionEnabled = false;
 					}
 				}
