[UBA]
* Moved session client killing of process because of memory pressure to its own function

diff --git a/Engine/Source/Programs/UnrealBuildAccelerator/Common/Private/UbaSessionClient.cpp b/Engine/Source/Programs/UnrealBuildAccelerator/Common/Private/UbaSessionClient.cpp
--- a/Engine/Source/Programs/UnrealBuildAccelerator/Common/Private/UbaSessionClient.cpp
+++ b/Engine/Source/Programs/UnrealBuildAccelerator/Common/Private/UbaSessionClient.cpp
@@ -1958,6 +1958,20 @@ namespace uba
 
 		Vector<InternalProcessStartInfo> startInfos;
 
+		auto TryKill = [&](ProcessRec& rec, const tchar* reason)
+			{
+				if (rec.isKilled || rec.isDone)
+					return false;
+				SCOPED_FUTEX(rec.lock, lock);
+				if (rec.isDone)
+					return false;
+				if (!rec.handle.Cancel())
+					return false;
+				rec.isKilled = true;
+				SendReturnProcess(rec.handle.GetId(), reason);
+				++m_stats.killCount;
+				return true;
+			};
 
 		while (m_loop)
 		{
@@ -1975,16 +1989,8 @@ namespace uba
 				for (auto it = activeProcesses.rbegin(); it != activeProcesses.rend(); ++it)
 				{
 					ProcessRec& rec = *it;
-					if (rec.isKilled || rec.isDone)
+					if (!TryKill(rec, TC("Running out of memory")))
 						continue;
-					SCOPED_FUTEX(rec.lock, lock);
-					if (rec.isDone)
-						continue;
-					if (!rec.handle.Cancel())
-						continue;
-					rec.isKilled = true;
-					SendReturnProcess(rec.handle.GetId(), TC("Running out of memory"));
-					++m_stats.killCount;
 					m_logger.Warning(TC("Killed process due to memory pressure (Available: %s Total: %s)"), BytesToText(memAvail).str, BytesToText(memTotal).str);
 					break;
 				}
@@ -2048,6 +2054,10 @@ namespace uba
 					m_logger.Info(TC("Got remote execution disabled response from host (will finish %llu active processes)"), startInfos.size() + activeProcesses.size());
 
 
+						// We must always wait for cancelled processes to really finish since we might get new processes from host
+						// which depend on this process and then outputs might collide.
+
+
 						// We must always wait for cancelled processes to really finish since we might get new processes from host
 						// which depend on this process and then outputs might collide.
 				}
