[UBA]
* Fixed so UbaCli can loop .yaml files and shut everything down in between

diff --git a/Engine/Source/Programs/UnrealBuildAccelerator/Cli/Private/UbaCli.cpp b/Engine/Source/Programs/UnrealBuildAccelerator/Cli/Private/UbaCli.cpp
--- a/Engine/Source/Programs/UnrealBuildAccelerator/Cli/Private/UbaCli.cpp
+++ b/Engine/Source/Programs/UnrealBuildAccelerator/Cli/Private/UbaCli.cpp
@@ -1104,149 +1104,199 @@ namespace uba
 
 		NetworkBackend& networkBackend = useMemsock ? (NetworkBackend&)networkBackendMemory : (NetworkBackend&)networkBackendTcp;
 
-		NetworkServerCreateInfo nsci(logWriter);
-		nsci.useThreadedSend = true;
-		nsci.Apply(config);
+		RootsHandle rootsHandle = 0;
+		CacheClient* cacheClient = nullptr;
+		Function<bool(Scheduler*)> createCoordinator;
 
-		//nsci.workerCount = 4;
-		bool ctorSuccess = true;
-		NetworkServer& networkServer = *new NetworkServer(ctorSuccess, nsci);
-		auto destroyServer = MakeGuard([&]() { delete &networkServer; });
-		if (!ctorSuccess)
-			return false;
+		bool isRemote = commandType == CommandType_Remote || commandType == CommandType_Agent;
 
-		if (useCrypto)
+		auto RunSession = [&](const Function<bool(SessionServer& sessionServer)>& func)
 		{
-			networkServer.RegisterCryptoKey(crypto128Data);
-			logger.Info(TC("Using crypto key %s for connections"), crypto.c_str());
-		}
+			NetworkServerCreateInfo nsci(logWriter);
+			nsci.useThreadedSend = true;
+			nsci.Apply(config);
 
+			//nsci.workerCount = 4;
+			bool ctorSuccess = true;
+			NetworkServer& networkServer = *new NetworkServer(ctorSuccess, nsci);
+			auto destroyServer = MakeGuard([&]() { delete &networkServer; });
+			if (!ctorSuccess)
+				return false;
 
-		bool isRemote = commandType == CommandType_Remote || commandType == CommandType_Agent;
+			if (useCrypto)
+			{
+				networkServer.RegisterCryptoKey(crypto128Data);
+				logger.Info(TC("Using crypto key %s for connections"), crypto.c_str());
+			}
 
-		if (StringView(application).EndsWith(TCV(".yaml")))
-			useScheduler = true;
-
-		StorageServerCreateInfo storageInfo(networkServer, g_rootDir.data, logWriter);
-		storageInfo.casCapacityBytes = storageCapacity;
-		storageInfo.storeCompressed = storeCompressed;
-		storageInfo.Apply(config);
-		StorageServer& storageServer = *new StorageServer(storageInfo);
-		auto destroyStorage = MakeGuard([&]() { delete &storageServer; });
-
-		SessionServerCreateInfo info(storageServer, networkServer, logWriter);
-		info.useUniqueId = useUniqueId || useScheduler;
-		info.traceEnabled = cacheCommand.empty() || writeCacheSummary;
-		info.detailedTrace = detailedTrace;
-		info.traceChildProcesses = traceChildProcesses;
-		info.launchVisualizer = launchVisualizer;
-		info.allowCustomAllocator = !disableCustomAllocator;
-		//info.shouldWriteToDisk = shouldWriteToDisk;
-		info.rootDir = g_rootDir.data;
-		//info.traceName.Append(TCV("TESTTRACE"));
-		//info.storeIntermediateFilesCompressed = false;
-		info.readIntermediateFilesCompressed = true;
-		//info.extractObjFilesSymbols = true;
-		#if UBA_DEBUG_LOG_ENABLED
-		info.remoteLogEnabled = true;
-		#endif
-		//info.remoteTraceEnabled = true;
+			if (StringView(application).EndsWith(TCV(".yaml")))
+				useScheduler = true;
 
-		info.deleteSessionsOlderThanSeconds = info.useUniqueId ? 0 : 1;
-		info.Apply(config);
+			StorageServerCreateInfo storageInfo(networkServer, g_rootDir.data, logWriter);
+			storageInfo.casCapacityBytes = storageCapacity;
+			storageInfo.storeCompressed = storeCompressed;
+			storageInfo.Apply(config);
+			StorageServer& storageServer = *new StorageServer(storageInfo);
+			auto destroyStorage = MakeGuard([&]() { delete &storageServer; });
+
+			SessionServerCreateInfo info(storageServer, networkServer, logWriter);
+			info.useUniqueId = useUniqueId || useScheduler;
+			info.traceEnabled = cacheCommand.empty() || writeCacheSummary;
+			info.detailedTrace = detailedTrace;
+			info.traceChildProcesses = traceChildProcesses;
+			info.launchVisualizer = launchVisualizer;
+			info.allowCustomAllocator = !disableCustomAllocator;
+			//info.shouldWriteToDisk = shouldWriteToDisk;
+			info.rootDir = g_rootDir.data;
+			//info.traceName.Append(TCV("TESTTRACE"));
+			//info.storeIntermediateFilesCompressed = false;
+			info.readIntermediateFilesCompressed = true;
+			//info.extractObjFilesSymbols = true;
+			#if UBA_DEBUG_LOG_ENABLED
+			info.remoteLogEnabled = true;
+			#endif
+			//info.remoteTraceEnabled = true;
 
-		SessionServer& sessionServer = *new SessionServer(info);
-		auto destroySession = MakeGuard([&]()
-			{
-				#if !UBA_DEBUG
-				if (info.useUniqueId)
-					DeleteAllFiles(logger, sessionServer.GetSessionDir());
-				#endif
-				delete &sessionServer;
-			});
+			info.deleteSessionsOlderThanSeconds = info.useUniqueId ? 0 : 1;
+			info.Apply(config);
 
-		CacheClient* cacheClient = nullptr;
-		auto ccg = MakeGuard([&]() { if (!cacheClient) return; auto& nc = cacheClient->GetClient(); nc.Disconnect(); delete cacheClient; delete &nc; });
+			SessionServer& sessionServer = *new SessionServer(info);
+			auto destroySession = MakeGuard([&]()
+				{
+					#if !UBA_DEBUG
+					if (info.useUniqueId)
+						DeleteAllFiles(logger, sessionServer.GetSessionDir());
+					#endif
+					delete &sessionServer;
+				});
 
-		auto CreateCacheClient = [&]()
+			cacheClient = nullptr;
+			auto ccg = MakeGuard([&]() { if (!cacheClient) return; auto& nc = cacheClient->GetClient(); nc.Disconnect(); delete cacheClient; delete &nc; });
+
+			auto CreateCacheClient = [&]()
+				{
+					NetworkClientCreateInfo nci(logWriter);
+					if (!crypto.empty())
+						nci.cryptoKey128 = crypto128Data;
+					auto nc = new NetworkClient(ctorSuccess, nci);
+					cacheClient = new CacheClient({logWriter, storageServer, *nc, sessionServer});
+				};
+
+			if (cacheHost.count)
 			{
-				NetworkClientCreateInfo nci(logWriter);
-				if (!crypto.empty())
-					nci.cryptoKey128 = crypto128Data;
-				auto nc = new NetworkClient(ctorSuccess, nci);
-				cacheClient = new CacheClient({logWriter, storageServer, *nc, sessionServer});
-			};
+				CreateCacheClient();
+				if (!cacheClient->GetClient().Connect(networkBackend, cacheHost.data, cachePort))
+					return logger.Error(TC("Failed to connect to cache server"));
 
-		if (cacheHost.count)
-		{
-			CreateCacheClient();
-			if (!cacheClient->GetClient().Connect(networkBackend, cacheHost.data, cachePort))
-				return logger.Error(TC("Failed to connect to cache server"));
+				if (!storageServer.LoadCasTable(true))
+					return false;
 
-			if (!storageServer.LoadCasTable(true))
-				return false;
+				if (!cacheCommand.empty())
+				{
+					LoggerWithWriter consoleLogger(g_consoleLogWriter);
+					const tchar* additionalInfo = nullptr;
+					return cacheClient->ExecuteCommand(consoleLogger, cacheCommand.data(), nullptr, additionalInfo);
+				}
 
-			if (!cacheCommand.empty())
-			{
-				LoggerWithWriter consoleLogger(g_consoleLogWriter);
-				const tchar* additionalInfo = nullptr;
-				return cacheClient->ExecuteCommand(consoleLogger, cacheCommand.data(), nullptr, additionalInfo);
+				if (writeCacheSummary)
+				{
+					StringBuffer<> tempFile(sessionServer.GetTempPath());
+					Guid guid;
+					CreateGuid(guid);
+					tempFile.Append(GuidToString(guid).str).Append(TCV(".txt"));
+					if (!cacheClient->ExecuteCommand(logger, TC("content"), tempFile.data, cacheFilterString.data()))
+						return false;
+					logger.Info(TC("Cache status summary written to %s"), tempFile.data);
+
+					#if PLATFORM_WINDOWS
+					ShellExecuteW(NULL, L"open", tempFile.data, NULL, NULL, SW_SHOW);
+					#endif
+					return true;
+				}
 			}
 
-			if (writeCacheSummary)
+			// Remove empty spaces and line feeds etc at the end.. just to solve annoying copy paste command lines and accidentally getting line feed
+			while (!arguments.empty())
 			{
-				StringBuffer<> tempFile(sessionServer.GetTempPath());
-				Guid guid;
-				CreateGuid(guid);
-				tempFile.Append(GuidToString(guid).str).Append(TCV(".txt"));
-				if (!cacheClient->ExecuteCommand(logger, TC("content"), tempFile.data, cacheFilterString.data()))
-					return false;
-				logger.Info(TC("Cache status summary written to %s"), tempFile.data);
+				tchar lastChar = arguments[arguments.size()-1];
+				if (lastChar != '\n' && lastChar != '\r' && lastChar != '\t' && lastChar != ' ')
+					break;
+				arguments.resize(arguments.size() - 1);
+			}
 
-				#if PLATFORM_WINDOWS
-				ShellExecuteW(NULL, L"open", tempFile.data, NULL, NULL, SW_SHOW);
-				#endif
-				return true;
+			// Vfs testing
+			rootsHandle = 0;
+
+			if (!vfsEntries.empty())
+			{
+				StackBinaryWriter<8*1024> writer;
+				for (VfsEntry& entry : vfsEntries)
+				{
+					writer.WriteByte(0);
+					writer.WriteString(entry.virtualPath);
+					writer.WriteString(entry.localPath);
+				}
+				rootsHandle = sessionServer.RegisterRoots(writer.GetData(), writer.GetPosition());
+				sessionServer.DevirtualizeString(application, rootsHandle, true);
 			}
-		}
 
-		// Remove empty spaces and line feeds etc at the end.. just to solve annoying copy paste command lines and accidentally getting line feed
-		while (!arguments.empty())
-		{
-			tchar lastChar = arguments[arguments.size()-1];
-			if (lastChar != '\n' && lastChar != '\r' && lastChar != '\t' && lastChar != ' ')
-				break;
-			arguments.resize(arguments.size() - 1);
-		}
 
-		// Vfs testing
-		RootsHandle rootsHandle = 0;
 
-		if (!vfsEntries.empty())
-		{
-			StackBinaryWriter<8*1024> writer;
-			for (VfsEntry& entry : vfsEntries)
+			if (isRemote || useScheduler)
 			{
-				writer.WriteByte(0);
-				writer.WriteString(entry.virtualPath);
-				writer.WriteString(entry.localPath);
+				if (!storageServer.m_casTableLoaded)
+					if (!storageServer.LoadCasTable(true))
+						return false;
+				if (!networkServer.StartListen(networkBackend, port, listenIp.data))
+					return false;
 			}
-			rootsHandle = sessionServer.RegisterRoots(writer.GetData(), writer.GetPosition());
-			sessionServer.DevirtualizeString(application, rootsHandle, true);
-		}
 
+			CoordinatorWrapper coordinator;
+			auto cg = MakeGuard([&]() { coordinator.Destroy(); });
 
+			createCoordinator = [&](Scheduler* scheduler)
+				{
+					if (coordinatorName.empty())
+						return true;
+					StringBuffer<512> coordinatorWorkDir(g_rootDir);
+					coordinatorWorkDir.EnsureEndsWithSlash().Append(coordinatorName);
+					StringBuffer<512> binariesDir;
+					if (!GetDirectoryOfCurrentModule(logger, binariesDir))
+						return false;
 
-		if (isRemote || useScheduler)
-		{
-			if (!storageServer.m_casTableLoaded)
-				if (!storageServer.LoadCasTable(true))
-					return false;
-			if (!networkServer.StartListen(networkBackend, port, listenIp.data))
-				return false;
-		}
+					CoordinatorCreateInfo cinfo;
+					cinfo.workDir = coordinatorWorkDir.data;
+					cinfo.binariesDir = binariesDir.data;
+
+					// TODO: This is very horde specific.. maybe all these parameters should be a string or something
+					cinfo.pool = coordinatorPool.c_str();
+					cinfo.maxCoreCount = 1;
+					cinfo.logging = false;
+					if (!coordinator.Create(logger, coordinatorName.c_str(), cinfo, networkBackend, networkServer, scheduler))
+						return false;
+					return true;
+				};
 
-		auto RunLocal = [&](const TString& app, const TString& arg, bool enableDetour)
+			auto stopServer = MakeGuard([&]() { networkServer.DisconnectClients(); });
+			auto stopListen = MakeGuard([&]() { networkBackend.StopListen(); });
+
+			bool res = func(sessionServer);
+
+			logger.BeginScope();
+			if (printSummary)
+			{
+				sessionServer.PrintSummary(logger);
+				storageServer.PrintSummary(logger);
+				networkServer.PrintSummary(logger);
+				KernelStats::GetGlobal().Print(logger, true);
+				PrintContentionSummary(logger);
+			}
+			logger.EndScope();
+
+			return res;
+		};
+
+		auto RunLocal = [&](SessionServer& sessionServer, const TString& app, const TString& arg, bool enableDetour)
 		{
 			u64 start = GetTime();
 			ProcessStartInfo pinfo;
@@ -1293,7 +1343,7 @@ namespace uba
 			return true;
 		};
 
-		auto RunRemote = [&](const TString& app, const TString& arg)
+		auto RunRemote = [&](SessionServer& sessionServer, const TString& app, const TString& arg)
 		{
 			u64 start = GetTime();
 			ProcessStartInfo pinfo;
@@ -1318,7 +1368,7 @@ namespace uba
 
 		const tchar* clientZone = TC("DummyZone");
 
-		auto RunWithClient = [&](const Function<bool()>& func, u32 clientCount)
+		auto RunWithClient = [&](const Function<bool()>& func, u32 clientCount, Scheduler* scheduler)
 			{
 				List<Client> clients;
 				for (u32 i=0; i!=clientCount; ++i)
@@ -1328,46 +1378,29 @@ namespace uba
 					if (!clients.emplace_back().Init(cii))
 						return false;
 				}
+
+				#if 0
+				Thread disconnectThread([&]()
+					{
+						Sleep((20 + (rand() % 5)) * 1000);
+						clients.front().networkClient->Disconnect();
+						if (scheduler)
+							scheduler->SetMaxLocalProcessors(100);
+						return true;
+					});
+				#endif
 				return func();
 			};
 
-		auto RunAgent = [&](const TString& app, const TString& arg)
+		auto RunAgent = [&](SessionServer& sessionServer, const TString& app, const TString& arg)
 		{
-			return RunWithClient([&]() { return RunRemote(app, arg); }, 1);
+			return RunWithClient([&]() { return RunRemote(sessionServer, app, arg); }, 1, nullptr);
 		};
 
-		CoordinatorWrapper coordinator;
-		auto cg = MakeGuard([&]() { coordinator.Destroy(); });
-
-		auto CreateCoordinator = [&](Scheduler* scheduler)
-			{
-				if (coordinatorName.empty())
-					return true;
-				StringBuffer<512> coordinatorWorkDir(g_rootDir);
-				coordinatorWorkDir.EnsureEndsWithSlash().Append(coordinatorName);
-				StringBuffer<512> binariesDir;
-				if (!GetDirectoryOfCurrentModule(logger, binariesDir))
-					return false;
-
-				CoordinatorCreateInfo cinfo;
-				cinfo.workDir = coordinatorWorkDir.data;
-				cinfo.binariesDir = binariesDir.data;
 
-				// TODO: This is very horde specific.. maybe all these parameters should be a string or something
-				cinfo.pool = coordinatorPool.c_str();
-				cinfo.maxCoreCount = 1;
-				cinfo.logging = false;
-				if (!coordinator.Create(logger, coordinatorName.c_str(), cinfo, networkBackend, networkServer, scheduler))
-					return false;
-				return true;
-			};
-
-		auto stopServer = MakeGuard([&]() { networkServer.DisconnectClients(); });
-		auto stopListen = MakeGuard([&]() { networkBackend.StopListen(); });
-
-		auto RunScheduler = [&](const TString& app, const TString& arg)
+		auto RunScheduler = [&](SessionServer& sessionServer, const TString& app, const TString& arg)
 		{
-			/*
+			#if 0
 			bool ctorSuccess;
 			NetworkBackendMemory nbm(logWriter);
 			
@@ -1391,22 +1424,32 @@ namespace uba
 				cacheNetworkClient.Connect(nbm, TC("127.0.0.1"));
 			}
 			auto g = MakeGuard([&]() { cacheNetworkClient.Disconnect(); cacheNetworkServer.DisconnectClients(); });
-			*/
+			#endif
+
+
 			auto g = MakeGuard([&]() { if (cacheClient) cacheClient->GetClient().Disconnect(); });
 
+			bool success = true;
+			Atomic<u32> counter = 0;
+			Event finished(EventResetType_Manual);
+
 			CacheClient* cacheClients[] = { cacheClient };
 			SchedulerCreateInfo info(sessionServer);
 			info.Apply(config);
 			info.forceRemote = isRemote;
 			info.forceNative = commandType == CommandType_Native;
 			info.maxLocalProcessors = maxProcessCount;
+			//info.maxRacingPercent = 50u;
+			//info.memWatchdog = true;
+			//info.memTracingLevel = 1;
 			info.cacheClients = cacheClients;
 			info.cacheClientCount = cacheClient ? 1 : 0;
 			info.writeToCache = populateCache;
 			Scheduler scheduler(info);
 
+			//agentCount = 2;
 
-			if (!CreateCoordinator(&scheduler))
+			if (!createCoordinator(&scheduler))
 				return false;
 
 			if (StringView(app).EndsWith(TCV(".yaml")))
@@ -1432,10 +1475,6 @@ namespace uba
 			u32 queued, activeLocal, activeRemote, outFinished;
 			scheduler.GetStats(queued, activeLocal, activeRemote, outFinished);
 
-			bool success = true;
-			Atomic<u32> counter = 0;
-			static Event finished(EventResetType_Manual);
-
 			scheduler.SetProcessFinishedCallback([&](const ProcessHandle& ph)
 				{
 					auto& si = ph.GetStartInfo();
@@ -1474,14 +1513,14 @@ namespace uba
 					u64 time = GetTime() - start;
 					logger.Info(TC("Scheduler run took %s"), TimeToText(time).str);
 					logger.Info(TC(""));
-					stopServer.Execute();
+					sessionServer.GetServer().DisconnectClients();
 					return success;
 				};
 
 			if (commandType == CommandType_Agent)
 			{
 				u32 clientCount = maxProcessCount == 1 ? 1 : agentCount;
-				return RunWithClient([&]() { return RunQueue(); }, clientCount);
+				return RunWithClient([&]() { return RunQueue(); }, clientCount, &scheduler);
 			}
 			else
 				return RunQueue();
@@ -1503,54 +1542,44 @@ namespace uba
 		}
 #endif
 
+		bool success = true;
+
 		for (u32 i=0; i!=loopCount; ++i)
 		{
-			bool success = false;
-
-			if (useScheduler)
+			if (i != 0)
 			{
-				success = RunScheduler(application, arguments);
+				logger.Info(TC(""));
+				logger.Info(TC("RUN #%u"), i+1);
 			}
-			else
-			{
-				if (!CreateCoordinator(nullptr))
-					return false;
 
-				switch (commandType)
+			success = RunSession([&](SessionServer& sessionServer)
 				{
-				case CommandType_Native:
-					success = RunLocal(application, arguments, false);
-					break;
-				case CommandType_Local:
-					success = RunLocal(application, arguments, true);
-					break;
-				case CommandType_Remote:
-					success = RunRemote(application, arguments);
-					break;
-				case CommandType_Agent:
-					success = RunAgent(application, arguments);
-				}
-			}
+					if (useScheduler)
+						return RunScheduler(sessionServer, application, arguments);
+					if (!createCoordinator(nullptr))
+						return false;
+					switch (commandType)
+					{
+					case CommandType_Native:
+						return RunLocal(sessionServer, application, arguments, false);
+						break;
+					case CommandType_Local:
+						return RunLocal(sessionServer, application, arguments, true);
+						break;
+					case CommandType_Remote:
+						return RunRemote(sessionServer, application, arguments);
+						break;
+					case CommandType_Agent:
+						return RunAgent(sessionServer, application, arguments);
+					}
+					return false;
+				});
 			if (!success)
-				return false;
-
-			if (false)
-				networkServer.DisconnectClients();
-
-		}
-
-		logger.BeginScope();
-		if (printSummary)
-		{
-			sessionServer.PrintSummary(logger);
-			storageServer.PrintSummary(logger);
-			networkServer.PrintSummary(logger);
-			KernelStats::GetGlobal().Print(logger, true);
-			PrintContentionSummary(logger);
+				break;
+			//if (false)
+			//	networkServer.DisconnectClients();
 		}
-		logger.EndScope();
-
-		return true;
+		return success;
 	}
 }
 
