[UBA]
* Added SkipCopy to Storage to be able to release ref counted cas entries received from clients even though file was never written

diff --git a/Engine/Source/Programs/UnrealBuildAccelerator/Common/Private/UbaStorageServer.cpp b/Engine/Source/Programs/UnrealBuildAccelerator/Common/Private/UbaStorageServer.cpp
--- a/Engine/Source/Programs/UnrealBuildAccelerator/Common/Private/UbaStorageServer.cpp
+++ b/Engine/Source/Programs/UnrealBuildAccelerator/Common/Private/UbaStorageServer.cpp
@@ -163,6 +163,13 @@ namespace uba
 		return res;
 	}
 
+	bool StorageServer::SkipWrite(const CasKey& casKey, u32 clientId)
+	{
+		if (clientId != 0)
+			ReleaseActiveRefCount(casKey, clientId);
+		return true;
+	}
+
 	void StorageServer::ActiveFetch::Release(StorageServer& server, const tchar* reason)
 	{
 		if (ownsMapping)
@@ -992,7 +999,7 @@ namespace uba
 				{
 					++casEntry.refCount;
 					if (!casEntry.refCount)
-						m_logger.Error(TC("CasEntry refcount wrap-around during store. This is bad, report to uba devs"));
+						m_logger.Error(TC("CasEntry %s refcount wrap-around during store. This is bad, report to uba devs (%s)"), CasKeyString(casKey).str, hint.data);
 					SCOPED_WRITE_LOCK(m_activeRefCountsLock, refCountLock);
 					m_activeRefCounts[clientId].push_back(&casEntry);
 				}
diff --git a/Engine/Source/Programs/UnrealBuildAccelerator/Common/Public/UbaStorage.h b/Engine/Source/Programs/UnrealBuildAccelerator/Common/Public/UbaStorage.h
--- a/Engine/Source/Programs/UnrealBuildAccelerator/Common/Public/UbaStorage.h
+++ b/Engine/Source/Programs/UnrealBuildAccelerator/Common/Public/UbaStorage.h
@@ -71,6 +71,7 @@ namespace uba
 		virtual bool CopyOrLink(const CasKey& casKey, const tchar* destination, u32 fileAttributes, bool writeCompressed = false, u32 clientId = 0, const FormattingFunc& formattingFunc = {}, bool isTemp = false, bool allowHardLink = true) = 0;
 		virtual bool WritePlaceholder(const CasKey& casKey, const tchar* destination, u32 fileAttributes, u32 clientId) = 0;
 		virtual bool FakeCopy(const CasKey& casKey, const tchar* destination, u64 size, u64 lastWritten, u32 fileAttributes, bool deleteExisting = true) = 0;
+		virtual bool SkipWrite(const CasKey& casKey, u32 clientId) { return true; } 
 		virtual bool GetCasFileName(StringBufferBase& out, const CasKey& casKey) = 0;
 
 		virtual MappedView2 MapView(const CasKey& casKey, const tchar* hint) = 0;
diff --git a/Engine/Source/Programs/UnrealBuildAccelerator/Common/Public/UbaStorageServer.h b/Engine/Source/Programs/UnrealBuildAccelerator/Common/Public/UbaStorageServer.h
--- a/Engine/Source/Programs/UnrealBuildAccelerator/Common/Public/UbaStorageServer.h
+++ b/Engine/Source/Programs/UnrealBuildAccelerator/Common/Public/UbaStorageServer.h
@@ -41,6 +41,7 @@ namespace uba
 
 		virtual bool CopyOrLink(const CasKey& casKey, const tchar* destination, u32 fileAttributes, bool writeCompressed = false, u32 clientId = 0, const FormattingFunc& formattingFunc = {}, bool isTemp = false, bool allowHardLink = true) override;
 		virtual bool WritePlaceholder(const CasKey& casKey, const tchar* destination, u32 fileAttributes, u32 clientId) override;
+		virtual bool SkipWrite(const CasKey& casKey, u32 clientId) override;
 
 	protected:
 
