[UBA]
* Removed unsupported (and never implemented) SessionServer_RunProcessRacing. Racing will be implemented inside uba scheduler instead

diff --git a/Engine/Source/Programs/UnrealBuildAccelerator/Host/Private/UbaExports.cpp b/Engine/Source/Programs/UnrealBuildAccelerator/Host/Private/UbaExports.cpp
--- a/Engine/Source/Programs/UnrealBuildAccelerator/Host/Private/UbaExports.cpp
+++ b/Engine/Source/Programs/UnrealBuildAccelerator/Host/Private/UbaExports.cpp
@@ -576,10 +576,6 @@ extern "C"
 	{
 		return new uba::ProcessHandle(server->RunProcessRemote(info, weight, knownInputs, knownInputsCount, allowCrossArchitecture));
 	}
-	uba::ProcessHandle* SessionServer_RunProcessRacing(uba::SessionServer* server, uba::u32 raceAgainstRemoteProcessId)
-	{
-		return new uba::ProcessHandle(server->RunProcessRacing(raceAgainstRemoteProcessId));
-	}
 	uba::u64 SessionServer_RegisterRoots(uba::SessionServer* server, const void* rootsData, uba::u64 rootsDataSize)
 	{
 		return server->RegisterRoots(rootsData, rootsDataSize);
diff --git a/Engine/Source/Programs/UnrealBuildAccelerator/Host/Public/UbaExports.h b/Engine/Source/Programs/UnrealBuildAccelerator/Host/Public/UbaExports.h
--- a/Engine/Source/Programs/UnrealBuildAccelerator/Host/Public/UbaExports.h
+++ b/Engine/Source/Programs/UnrealBuildAccelerator/Host/Public/UbaExports.h
@@ -133,7 +133,6 @@ extern "C"
 
 	UBA_API uba::ProcessHandle* SessionServer_RunProcess(uba::SessionServer* server, uba::ProcessStartInfo& info, bool async, bool enableDetour);
 	UBA_API uba::ProcessHandle* SessionServer_RunProcessRemote(uba::SessionServer* server, uba::ProcessStartInfo& info, float weight, const void* knownInputs = nullptr, uba::u32 knownInputsCount = 0, bool allowCrossArchitecture = false);
-	UBA_API uba::ProcessHandle* SessionServer_RunProcessRacing(uba::SessionServer* server, uba::u32 raceAgainstRemoteProcessId);
 	UBA_API uba::u64 SessionServer_RegisterRoots(uba::SessionServer* server, const void* rootsData, uba::u64 rootsDataSize);
 
 	UBA_API void SessionServer_SetMaxRemoteProcessCount(uba::SessionServer* server, uba::u32 count);
