[UBA]
* Added support for process racing. This means that when local machine goes under x% usage it will start processes that can race against same processes that are running remotely. This is especially useful if remote agents are slower than host or if network connection is not great.

Feature can be enabled through:

[Scheduler]
MaxRacingPercent = 50

Above mean that when utilization is under 50% scheduler will start racing processes to get back up to 50%

diff --git a/Engine/Source/Programs/UnrealBuildAccelerator/Common/Private/UbaScheduler.cpp b/Engine/Source/Programs/UnrealBuildAccelerator/Common/Private/UbaScheduler.cpp
--- a/Engine/Source/Programs/UnrealBuildAccelerator/Common/Private/UbaScheduler.cpp
+++ b/Engine/Source/Programs/UnrealBuildAccelerator/Common/Private/UbaScheduler.cpp
@@ -83,7 +83,9 @@ namespace uba
 	{
 		Scheduler* scheduler = nullptr;
 		ProcessStartInfo2* startInfo;
+		Atomic<u32> refCount = 1;
 		u32 processIndex = ~0u;
+		u32 processCount = 1;
 		bool wasReturned = false;
 		bool hasCheckedWasReturned = false;
 		bool isLocal = true;
@@ -155,7 +157,8 @@ namespace uba
 			table->GetValueAsU32(maxLocalProcessors, TC("MaxLocalProcessors"));
 			table->GetValueAsU32(maxParallelCacheQueries, TC("MaxParallelCacheQueries"));
 			table->GetValueAsBool(useThreadToExitRemoteProcess, TC("UseThreadToExitRemoteProcess"));
-			
+			table->GetValueAsU32(maxRacingPercent, TC("MaxRacingPercent"));
+
 			table->GetValueAsBool(memWatchdog, TC("MemTrack")); // Legacy
 
 			table->GetValueAsBool(memWatchdog, TC("MemWatchdog"));
@@ -176,6 +179,7 @@ namespace uba
 	Scheduler::Scheduler(const SchedulerCreateInfo& info)
 	:	m_session(info.session)
 	,	m_maxLocalProcessors(info.maxLocalProcessors != ~0u ? info.maxLocalProcessors : GetLogicalProcessorCount())
+	,	m_maxRacingWeight(info.maxRacingPercent*float(m_maxLocalProcessors)/100.0f)
 	,	m_maxParallelCacheQueries(info.maxParallelCacheQueries)
 	,	m_updateThreadWakeup(EventResetType_Auto)
 	,	m_enableProcessReuse(info.enableProcessReuse)
@@ -229,6 +233,11 @@ namespace uba
 			while (!m_memThread.Wait(1) && !m_memEstimatedUsage)
 				;
 		}
+		else if (m_maxRacingWeight)
+		{
+			m_memThreadEvent.Create(EventResetType_Manual);
+			m_memThread.Start([this]() { ThreadMemoryCheckLoop();  return 0; }, TC("UbaMemTrackLoop"));
+		}
 
 		m_loop = true;
 		m_updateThread.Start([this]() { ThreadLoop(); return 0; }, TC("UbaSchedLoop"));
@@ -266,6 +275,11 @@ namespace uba
 			m_memThread.Wait();
 			SaveMemTrackTable();
 		}
+		else if (m_maxRacingWeight)
+		{
+			m_memThreadEvent.Set();
+			m_memThread.Wait();
+		}
 
 		{
 			SCOPED_FUTEX(m_processEntriesLock, lock);
@@ -493,15 +507,33 @@ namespace uba
 	bool Scheduler::ProcessReturned(Process& process, bool isLocal)
 	{
 		auto& si = process.GetStartInfo();
+
+		if (!si.userData)
+		{
+			if (!isLocal)
+				process.Cancel(); // This can happen when racing process won.. then remote process callback already called so userData is nullptr
+			return false;
+		}
+		
 		auto& ei = *(ExitProcessInfo*)si.userData;
 		float processWeight = 0;
 		u32 processIndex;
 
+		++ei.refCount;
+		auto eig = MakeGuard([&]() { UBA_ASSERT(ei.refCount); if (!--ei.refCount) delete &ei; });
+
+		// We can end up here at the same time for both remote and racing process pointing to the same ExitProcessInfo
+		// The last one of them needs to update the process entries lock, the first needs to early exit
+		// Note that local process does not affect processCount until ProcessExit which is called in parallel between Cancel and WaitForExit
+
+		// Remote processes can only end up here from disconnect or network message. so a remote process can't finish and return at the same time
+
+		// Local processes can only end up here from race/memtrack killing them. They can also not finish at the same time since the Cancel call
+		// will fail if process is about to finish.
+
+
 		if (isLocal)
 		{
-			if (!si.userData)
-				return false;
-
 			processWeight = ei.startInfo->weight;
 			processIndex = ei.processIndex; // Must be fetched before cancel
 
@@ -516,10 +548,11 @@ namespace uba
 		{
 			processIndex = ei.processIndex; // Must be fetched before cancel
 			ei.wasReturned = true;
-			process.Cancel(); // Cancel will call ProcessExited
+
+			bool res = process.Cancel(); // Cancel will call ProcessExited
+			UBA_ASSERTF(res, TC("Failed to cancel remote process?"));(void)res;
 		}
 		
-
 		if (processIndex == ~0u)
 		{
 			UBA_ASSERT(!isLocal);
@@ -535,25 +568,34 @@ namespace uba
 			}
 		}
 
+		// We let the process that matches ei.isLocal do the return of process
+		// .. if remote process cancels first it will give over ownership to local race process
+		// .. if race process cancels first ei.isLocal will still be false and remote will clean up
+		if (ei.isLocal != isLocal)
+			return true;
+
+		UBA_ASSERT(ei.processCount == 0);
+
 		SCOPED_FUTEX(m_processEntriesLock, lock);
 
 		if (isLocal)
 			m_activeLocalProcessWeight -= processWeight;
 
-		if (m_processEntries[processIndex].runStatus != RunStatus_Running)
-			return m_session.GetLogger().Error(TC("This should not happen (%u)"), m_processEntries[processIndex].runStatus);
+		ProcessEntry& entry = m_processEntries[processIndex];
+		if (entry.runStatus != RunStatus_Running)
+			return m_session.GetLogger().Error(TC("This should not happen (%u)"), entry.runStatus);
 
 		bool success = !m_cancelled;
 
 		if (success)
 		{
-			m_processEntries[processIndex].runStatus = RunStatus_QueuedForRun;
+			entry.runStatus = RunStatus_QueuedForRun;
 			m_processEntriesRunStart = Min(m_processEntriesRunStart, processIndex);
 		}
 		else
 		{
 			m_session.GetLogger().Error(TC("Here to check this path works"));
-			m_processEntries[processIndex].runStatus = RunStatus_Skipped;
+			entry.runStatus = RunStatus_Skipped;
 			ProcessHandle ph(&process);
 			FinishProcess(ph);
 		}
@@ -614,15 +656,51 @@ namespace uba
 
 	void Scheduler::ProcessExited(ExitProcessInfo* info, const ProcessHandle& handle)
 	{
-		auto ig = MakeGuard([info]() { delete info; });
-
-		bool isLocal = !handle.IsRemote();
+		auto ig = MakeGuard([info]() { UBA_ASSERT(info->refCount);if (!--info->refCount) delete info; });
 
 		auto si = info->startInfo;
 
-		if (isLocal && si)
+		if (!handle.IsRemote() && si)
 			LocalProcessExit(*si, handle);
 
+		u32 processCount = info->processCount--;
+
+		if (processCount == 2)
+		{
+			UBA_ASSERT(handle.GetExitCode() == ProcessCancelExitCode);
+
+			float weight = info->startInfo->weight;
+
+			if (handle.IsRemote()) // Racing process won. Let's make it a "real" local process
+			{
+				// handle.HasExited() returns true before we get in here.. so might not found racing process
+				SCOPED_FUTEX(m_racingProcessesLock, lock);
+				for (auto i=m_racingProcesses.begin(),e=m_racingProcesses.end(); i!=e; ++i)
+				{
+					auto& rp = *(ProcessImpl*)i->m_process;
+					if (rp.m_startInfo.userData != info)
+						continue;
+					m_racingProcesses.erase(i);
+					break;
+				}
+
+				// Turn it into real process adding weight
+				info->isLocal = true;
+				
+				m_activeLocalProcessWeight += weight;
+
+				UpdateActiveProcessCounter(true, 1);
+				UpdateActiveProcessCounter(false, -1);
+			}
+
+			m_activeRacingProcessWeight -= weight;
+
+			SCOPED_CRITICAL_SECTION(m_wasReturnedLock, lock);
+			info->wasReturned = false;
+			info->hasCheckedWasReturned = false;
+			return;
+		}
+
 		{
 			SCOPED_CRITICAL_SECTION(m_wasReturnedLock, lock);
 			if (info->wasReturned)
@@ -636,7 +714,7 @@ namespace uba
 			return;
 		}
 
-		if (!m_useThreadToExitRemoteProcess || isLocal)
+		if (!m_useThreadToExitRemoteProcess || info->isLocal)
 		{
 			ExitProcess(*info, *handle.m_process, handle.m_process->GetExitCode(), false);
 			return;
@@ -644,10 +722,13 @@ namespace uba
 
 		// If process is remote we want to return as soon as possible to return worker thread
 		ig.Cancel();
+
 		m_finishedRemoteProcessThreads.emplace_back([this, handle = handle, info]()
 			{
 				ExitProcess(*info, *handle.m_process, handle.m_process->GetExitCode(), false);
-				delete info;
+				UBA_ASSERT(info->refCount);
+				if (!--info->refCount)
+					delete info;
 				return 0;
 
 			}, TC("UbaScheRemFin"));
@@ -715,7 +796,9 @@ namespace uba
 
 	u32 Scheduler::PopProcess(bool isLocal, RunStatus& outPrevStatus)
 	{
-		bool atMaxLocalWeight = m_activeLocalProcessWeight >= float(m_maxLocalProcessors);
+		float activeWeight = m_activeLocalProcessWeight;
+		float maxWeight = float(m_maxLocalProcessors);
+		bool atMaxLocalWeight = activeWeight >= maxWeight;
 		auto processEntries = m_processEntries.data();
 		bool allFinished = true;
 
@@ -751,6 +834,8 @@ namespace uba
 					continue;
 				if (runStatus == RunStatus_QueuedForRun && atMaxLocalWeight)
 					continue;
+				if (activeWeight > 1.0f && activeWeight + entry.info->weight > maxWeight + 1.0f)
+					continue;
 			}
 			else
 			{
@@ -920,8 +1005,12 @@ namespace uba
 			SCOPED_FUTEX(m_processEntriesLock, lock);
 			u32 indexToRun = PopProcess(isLocal, prevStatus);
 			if (indexToRun == ~0u)
-				return false;
-
+			{
+				lock.Leave();
+				if (!isLocal)
+					return false;
+				return RaceRemoteProcess();
+			}
 			auto& processEntry = m_processEntries[indexToRun];
 			auto info = processEntry.info;
 			bool canDetour = processEntry.canDetour && !m_forceNative;
@@ -966,6 +1055,61 @@ namespace uba
 		}
 	}
 
+	bool Scheduler::RaceRemoteProcess()
+	{
+		if (!m_activeRemoteProcesses || !m_maxRacingWeight)
+			return false;
+
+		float totalWeight = m_activeLocalProcessWeight + m_activeRacingProcessWeight;
+		if (totalWeight >= m_maxRacingWeight)
+			return false;
+
+		if (!LocalProcessOkToSpawn())
+			return false;
+
+		SCOPED_FUTEX(m_racingProcessesLock, lock);
+		lock.Leave(); // Have to leave to prevent entrance inversion with m_remoteProcessAndSessionLock
+
+		ProcessHandle handle = m_session.RunProcessRacing([&](ProcessHandle& remoteProcess, ProcessStartInfo& psi)
+			{
+				auto ei = (ExitProcessInfo*)remoteProcess.GetStartInfo().userData;
+				float weight = ei->startInfo->weight;
+				if (totalWeight + weight > m_maxRacingWeight)
+					return false;
+				m_activeRacingProcessWeight += weight;
+				psi.userData = ei;
+				psi.exitedFunc = [](void* userData, const ProcessHandle& handle, ProcessExitedResponse& response)
+					{
+						auto ei = (ExitProcessInfo*)userData;
+						ei->scheduler->ProcessExited(ei, handle);
+					};
+
+				LocalProcessStart(*ei->startInfo);
+
+				lock.Enter(); // Enter here instead
+
+				++ei->refCount;
+				++ei->processCount;
+				return true;
+			});
+
+		lock.Enter(); // In case we didn't create a process
+
+		for (auto i=m_racingProcesses.begin();i!=m_racingProcesses.end();)
+		{
+			if (i->HasExited())
+				i = m_racingProcesses.erase(i);
+			else
+				++i;
+		}
+
+		if (!handle.IsValid())
+			return false;
+
+		m_racingProcesses.emplace_back(handle);
+		return true;
+	}
+
 	bool Scheduler::HandleReuseMessage(Process& process, NextProcessInfo& outNextProcess, u32 prevExitCode)
 	{
 		if (!m_enableProcessReuse)
@@ -1072,6 +1216,7 @@ namespace uba
 
 		if (isDone)
 		{
+			//UBA_ASSERT(exitCode == 0);
 			entry.runStatus = exitCode == 0 ? RunStatus_Success : RunStatus_Failed;
 			entry.info = nullptr;
 			entry.dependencies = nullptr;
@@ -1166,9 +1311,25 @@ namespace uba
 
 	void Scheduler::KillNewestLocalProcess(StringView reason)
 	{
-		ProcessHandle ph = m_session.GetNewestLocalProcess();
+		ProcessHandle ph;
+
+		// Kill racing processes first and then newest
+		
+		SCOPED_FUTEX(m_racingProcessesLock, lock);
+		for (auto i=m_racingProcesses.rbegin();i!=m_racingProcesses.rend(); ++i)
+		{
+			if (i->HasExited())
+				continue;
+			ph = *i;
+			break;
+		}
+		
 		if (!ph.IsValid())
-			return;
+		{
+			ph = m_session.GetNewestLocalProcess();
+			if (!ph.IsValid())
+				return;
+		}
 
 		m_session.GetTrace().SchedulerKillProcess(ph.GetId());
 
@@ -1199,6 +1360,8 @@ namespace uba
 					UBA_ASSERT(res);(void)res;
 				};
 		}
+		else
+			++ei->refCount;
 
 		StringKey key(impl.m_nativeProcessId, impl.m_nativeProcessCreationTime);
 		SCOPED_FUTEX(m_processesLock, l);
@@ -1301,6 +1464,7 @@ namespace uba
 				}
 			}			
 
+			UBA_ASSERT(stats.activeProcessCount);
 			--stats.activeProcessCount;
 
 			if (processPeakMemory >= stats.baseline) // Don't know if this test is too aggressive.. but 
@@ -1327,8 +1491,18 @@ namespace uba
 		{
 			StringKey key(impl.m_nativeProcessId, impl.m_nativeProcessCreationTime);
 			SCOPED_FUTEX(m_processesLock, l);
-			bool res = m_processes.erase(key) == 1;
-			UBA_ASSERTF(res, TC("Tried to erase local process from process list but can't find it (%s)"), info.GetDescription());(void)res;
+			auto findIt = m_processes.find(key);
+			UBA_ASSERTF(findIt != m_processes.end(), TC("Tried to erase local process from process list but can't find it (%s)"), info.GetDescription());
+			if (findIt != m_processes.end())
+			{
+				if (auto ep = findIt->second.ep)
+				{
+					UBA_ASSERT(ep->refCount);
+					if (!--ep->refCount)
+						delete ep;
+				}
+				m_processes.erase(findIt);
+			}
 		}
 	}
 
@@ -1473,6 +1647,29 @@ namespace uba
 			m_memWaitThreshold = u64(double(memTotal) * double(memStartWaitPercent) / 100.0);
 			memKillThreshold = u64(double(memTotal) * double(m_memStartKillPercent) / 100.0);
 
+			// Kill racing processes if we have run over racing process weight
+			float totalWeight = m_activeLocalProcessWeight + m_activeRacingProcessWeight;
+			if (totalWeight > m_maxRacingWeight)
+			{
+				SCOPED_FUTEX(m_racingProcessesLock, lock);
+				for (auto i=m_racingProcesses.rbegin();i!=m_racingProcesses.rend(); ++i)
+				{
+					if (i->HasExited())
+						continue;
+					ProcessHandle ph = *i;
+					m_session.GetTrace().SchedulerKillProcess(ph.GetId());
+					u64 startTime = GetTime();
+					if (!ProcessReturned(*ph.m_process, true))
+						continue;
+					m_session.GetLogger().Info(TC("Killed process %s (%s) - Too many racing processes"), ph.GetStartInfo().GetDescription(), TimeToText(GetTime() - startTime).str);
+					break;
+				}
+			}
+
+
+			if (!m_memWatchdog)
+				continue;
+
 			u64 untrackedMemory = 0;
 			u64 processMemTotalPageCount = 0;
 
@@ -1582,6 +1779,9 @@ namespace uba
 				while (memUsed > memKillThreshold)
 				{
 					str.Clear().Appendf(TC("Low on memory (%s/%s). Kill threshold is %s"), BytesToText(memUsed).str, BytesToText(memTotal).str, BytesToText(memKillThreshold).str);
+
+					// kill racing processes first
+
 					KillNewestLocalProcess(str);
 					if (!GetMemoryInfo(logger, memAvail, memTotal, &maxPageFile))
 						break;
diff --git a/Engine/Source/Programs/UnrealBuildAccelerator/Common/Private/UbaSessionClient.cpp b/Engine/Source/Programs/UnrealBuildAccelerator/Common/Private/UbaSessionClient.cpp
--- a/Engine/Source/Programs/UnrealBuildAccelerator/Common/Private/UbaSessionClient.cpp
+++ b/Engine/Source/Programs/UnrealBuildAccelerator/Common/Private/UbaSessionClient.cpp
@@ -1489,7 +1489,7 @@ namespace uba
 	};
 
 
-	bool SessionClient::SendProcessAvailable(Vector<InternalProcessStartInfo>& out, float availableWeight)
+	bool SessionClient::SendProcessAvailable(Vector<InternalProcessStartInfo>& outStartInfos, Vector<u32>& outCancelled, float availableWeight)
 	{
 		StackBinaryReader<SendMaxSize> reader;
 
@@ -1508,6 +1508,9 @@ namespace uba
 			}
 		}
 
+		u64 processesLostRaceCount = reader.Read7BitEncoded();
+		while (processesLostRaceCount--)
+			outCancelled.push_back(reader.ReadU32());
 
 		while (true)
 		{
@@ -1524,8 +1527,7 @@ namespace uba
 				m_remoteExecutionEnabled = false;
 				break;
 			}
-			out.push_back({});
-			InternalProcessStartInfo& info = out.back();
+			InternalProcessStartInfo& info = outStartInfos.emplace_back();
 			info.processId = processId;
 			info.Read(reader);
 
@@ -1558,7 +1560,7 @@ namespace uba
 			}
 		}
 
-		if (!out.empty())
+		if (!outStartInfos.empty())
 			if (neededDirectoryTableSize > GetDirectoryTableSize())
 				if (!SendUpdateDirectoryTable(reader.Reset()))
 					return false;
@@ -1957,6 +1959,7 @@ namespace uba
 		};
 
 		Vector<InternalProcessStartInfo> startInfos;
+		Vector<u32> cancelledProcesses;
 
 		auto TryKill = [&](ProcessRec& rec, const tchar* reason)
 			{
@@ -2041,8 +2044,9 @@ namespace uba
 				}
 
 				startInfos.clear();
+				cancelledProcesses.clear();
 
-				if (!SendProcessAvailable(startInfos, availableWeight))
+				if (!SendProcessAvailable(startInfos, cancelledProcesses, availableWeight))
 				{
 					m_loop = false;
 					break;
@@ -2062,6 +2066,21 @@ namespace uba
 						// which depend on this process and then outputs might collide.
 				}
 
+				for (u32 cancelled : cancelledProcesses)
+				{
+					for (auto& rec : activeProcesses)
+					{
+						if (rec.handle.GetId() != cancelled)
+							continue;
+						TryKill(rec, TC("Killed by server"));
+
+						// We must always wait for cancelled processes to really finish since we might get new processes from host
+						// which depend on this process and then outputs might collide.
+						if (!rec.handle.WaitForExit(60*1000))
+							m_logger.Error(TC("Failed to wait for killed process to exit"));
+					}
+				}
+
 				if (startInfos.empty())
 				{
 					canSpawn = false;
diff --git a/Engine/Source/Programs/UnrealBuildAccelerator/Common/Private/UbaSessionServer.cpp b/Engine/Source/Programs/UnrealBuildAccelerator/Common/Private/UbaSessionServer.cpp
--- a/Engine/Source/Programs/UnrealBuildAccelerator/Common/Private/UbaSessionServer.cpp
+++ b/Engine/Source/Programs/UnrealBuildAccelerator/Common/Private/UbaSessionServer.cpp
@@ -22,6 +22,14 @@ namespace uba
 		bool casRefcounted;
 	};
 
+	struct SessionServer::RacingProcessExitInfo
+	{
+		SessionServer* server;
+		EventSlim allowedToExit;
+		ProcessHandle remoteProcess;
+		ProcessStartInfo::ExitedCallback* exitedFunc;
+	};
+
 	class SessionServer::RemoteProcess final : public Process
 	{
 	public:
@@ -95,6 +103,9 @@ namespace uba
 			exitedFunc(userData, h, response);
 		}
 
+		using Process::AddRef;
+		using Process::Release;
+
 		SessionServer* m_server;
 		ProcessStartInfoHolder m_startInfo;
 		Futex m_exitedLock;
@@ -112,8 +123,11 @@ namespace uba
 		bool m_allowCrossArchitecture = false;
 		u32 m_clientId = ~0u;
 		u32 m_sessionId = 0;
+		u64 m_startTime = 0;
 		TString m_executingHost;
 
+		ProcessHandle m_racingProcess;
+
 		struct KnownInput { StringKey file; CasKey key; u32 mappingAlignment = 0; bool allowProxy = true; };
 		KnownInput* m_knownInputs = nullptr;
 		u32 m_knownInputsCount = 0;
@@ -312,10 +326,85 @@ namespace uba
 		#endif
 	}
 
-	ProcessHandle SessionServer::RunProcessRacing(u32 raceAgainstRemoteProcessId)
+	ProcessHandle SessionServer::RunProcessRacing(const ProcessToRaceCallback& callback)
 	{
-		// TODO: Implement
-		return {};
+		u64 newestStartTime = 0;
+		RemoteProcess* newest = nullptr;
+
+		SCOPED_CRITICAL_SECTION(m_remoteProcessAndSessionLock, lock);
+		for (const ProcessHandle& h : m_activeRemoteProcesses)
+		{
+			RemoteProcess& rp = *(RemoteProcess*)h.m_process;
+			if (rp.m_startTime <= newestStartTime)
+				continue;
+			if (!rp.m_startInfo.rules->SupportsRacing())
+				continue;
+			if (rp.m_racingProcess.IsValid()) // && !rc.m_racingProcess.IsCancelled()
+				continue;
+			newest = &rp;
+			newestStartTime = rp.m_startTime;
+		}
+
+		if (!newest)
+			return {};
+
+		ProcessStartInfo psi = newest->m_startInfo;
+		ProcessHandle remoteHandle(newest);
+		if (!callback(remoteHandle, psi))
+			return {};
+
+		StringBuffer<> logFile;
+		logFile.Append(psi.logFile);
+		if (const tchar* dot = logFile.Last('.'))
+		{
+			u64 offset = dot - logFile.data;
+			logFile.Resize(offset).Append(TCV("_racing")).Append(psi.logFile + offset);
+		}
+		else
+			logFile.Append(TCV("_racing"));
+		psi.logFile = logFile.data;
+
+		auto rpi = new RacingProcessExitInfo;
+		rpi->server = this;
+		rpi->remoteProcess = std::move(remoteHandle);
+		rpi->exitedFunc = psi.exitedFunc;
+		rpi->allowedToExit.Create(EventResetType_Manual);
+
+		psi.exitedFunc = [](void* userData, const ProcessHandle& h, ProcessExitedResponse& r)
+			{
+				auto& rpi = *(RacingProcessExitInfo*)userData;
+				rpi.allowedToExit.IsSet();
+
+				auto& remoteProcess = *(RemoteProcess*)rpi.remoteProcess.m_process;
+				auto& racingProcess = *(ProcessImpl*)h.m_process;
+
+				if (racingProcess.IsCancelled())
+				{
+					remoteProcess.m_racingProcess = {};
+				}
+				else
+				{
+					// We got here before remote process.. need to make sure process exit callback is called first
+					rpi.server->OnRaceLost(remoteProcess);
+					remoteProcess.m_exitCode = ProcessCancelExitCode;
+					remoteProcess.m_done.Set();
+					remoteProcess.CallProcessExit(rpi.remoteProcess);
+				}
+
+				rpi.exitedFunc(racingProcess.m_startInfo.userData, h, r);
+
+				delete &rpi;
+			};
+
+		ProcessHandle racingProcess = InternalRunProcess(psi, true, nullptr, m_allowLocalDetour, rpi);
+
+		m_trace.ProcessRace(racingProcess.GetId(), newest->GetId());
+
+		newest->m_racingProcess = racingProcess;
+
+		rpi->allowedToExit.Set();
+
+		return racingProcess;
 	}
 
 	ProcessHandle SessionServer::RunProcessRemote(const ProcessStartInfo& startInfo, float weight, const void* knownInputs, u32 knownInputsCount, bool allowCrossArchitecture)
@@ -350,7 +439,7 @@ namespace uba
 				memcpy(knownInputsCopy, knownInputs, knownInputsBytes);
 
 			#if defined(__clang_analyzer__)
-			free(knownInputsCopy); // doesn't seem to understand it is handed over
+			free(knownInputsCopy); // analyzer doesn't seem to understand it is handed over
 			#endif
 
 			m_server.AddWork([remoteProcess, this, knownInputsCopy, knownInputsCount, rules](const WorkContext& context)
@@ -1294,11 +1383,11 @@ namespace uba
 		return true;
 	}
 
-	bool SessionServer::WriteReceivedFiles(RemoteProcess& process, Timer& writeFilesTimer, const ConnectionInfo& connectionInfo)
+	bool SessionServer::WriteReceivedFiles(Vector<ReceivedFile>& receivedFiles, RemoteProcess& process, Timer& writeFilesTimer, const ConnectionInfo& connectionInfo)
 	{
 		TimerScope ts(writeFilesTimer);
 		Atomic<bool> writeSuccess = true;
-		m_server.ParallelFor(u32(process.m_receivedFiles.size()), process.m_receivedFiles, [&](const WorkContext&, auto& it)
+		m_server.ParallelFor(u32(receivedFiles.size()), receivedFiles, [&](const WorkContext&, auto& it)
 			{
 				if (!WriteReceivedFile(*it, &process, connectionInfo))
 					writeSuccess = false;
@@ -1308,7 +1397,14 @@ namespace uba
 		return writeSuccess;
 	}
 
-	void SessionServer::ReleaseReceivedFiles(const Vector<ReceivedFile>& receivedFiles)
+	void SessionServer::SkipReceivedFiles(Vector<ReceivedFile>& receivedFiles, const ConnectionInfo& connectionInfo)
+	{
+		u32 clientId = connectionInfo.GetId();
+		for (auto& file : receivedFiles)
+			m_storage.SkipWrite(file.casKey, clientId);
+	}
+
+	void SessionServer::ReleaseReceivedFiles(Vector<ReceivedFile>& receivedFiles)
 	{
 		for (auto& file : receivedFiles)
 			if (file.casRefcounted)
@@ -1474,6 +1570,7 @@ namespace uba
 		if (sessionIndex >= m_clientSessions.size())
 			return m_logger.Error(TC("Got ProcessAvailable message from connection using bad sessionid (%u/%llu)"), sessionIndex, m_clientSessions.size());
 		ClientSession& session = *m_clientSessions[sessionIndex];
+		auto processesLostRace = std::move(session.processesLostRace);
 		sessionsLock.Leave();
 
 		bool isCrossArchitecture = IsArmBinary != session.isArm;
@@ -1481,6 +1578,10 @@ namespace uba
 		u32 weight32 = reader.ReadU32();
 		float availableWeight = *(float*)&weight32;
 
+		writer.Write7BitEncoded(processesLostRace.size());
+		for (auto& lostProcess : processesLostRace)
+			writer.WriteU32(lostProcess);
+
 		Vector<RemoteProcess*> processesWithKnownInputsToSend;
 
 		float weightLeft = availableWeight;
@@ -1663,6 +1764,13 @@ namespace uba
 		auto& process = *(RemoteProcess*)h.m_process;
 
 		SCOPED_CRITICAL_SECTION(m_remoteProcessAndSessionLock, cs2);
+
+		ProcessHandle racingProcess = process.m_racingProcess;
+		bool hasRacingProcess = racingProcess.IsValid();
+		bool useRacingProcess = false;
+		if (hasRacingProcess)
+			useRacingProcess = !racingProcess.Cancel(); // If cancel fails it means that racing process has started writing files and we have to let it finish
+
 		if (!m_activeRemoteProcesses.erase(&process))
 		{
 			cs2.Leave();
@@ -1678,36 +1786,51 @@ namespace uba
 		if (session.enabled)
 			++m_availableRemoteSlotCount;
 		process.m_clientId = ~0u;
+		auto receivedFiles = std::move(process.m_receivedFiles);
 		cs2.Leave();
 
+		Timer writeFilesTimer;
+
 		u32 exitCode = reader.ReadU32();
 		u32 logLineCount = reader.ReadU32();
-
-		process.m_exitCode = exitCode;
-		process.m_logLines.reserve(logLineCount);
+		Vector<ProcessLogLine> logLines;
+		logLines.reserve(logLineCount);
 		while (logLineCount-- != 0)
 		{
 			TString text = reader.ReadString();
 			LogEntryType type = LogEntryType(reader.ReadByte());
-			process.m_logLines.push_back({ std::move(text), type });
+			logLines.push_back({ std::move(text), type });
 		}
 
-		auto& startInfo = process.m_startInfo;
+		process.m_logLines = std::move(logLines);
 
-		if (auto func = startInfo.logLineFunc)
-			for (auto& line : process.m_logLines)
-				func(startInfo.logLineUserData, line.text.c_str(), u32(line.text.size()), line.type);
+		bool filesWritten = false;
+		if (!useRacingProcess)
+		{
+			auto& startInfo = process.m_startInfo;
 
-		Timer writeFilesTimer;
-		if (startInfo.writeOutputFilesOnFail || startInfo.rules->IsExitCodeSuccess(exitCode))
-			if (!WriteReceivedFiles(process, writeFilesTimer, connectionInfo))
-				exitCode = UBA_EXIT_CODE(20);
-		auto receivedFiles = std::move(process.m_receivedFiles);
+			if (auto func = startInfo.logLineFunc)
+				for (auto& line : process.m_logLines)
+					func(startInfo.logLineUserData, line.text.c_str(), u32(line.text.size()), line.type);
+			if (startInfo.writeOutputFilesOnFail || startInfo.rules->IsExitCodeSuccess(exitCode))
+			{
+				filesWritten = true;
+				if (!WriteReceivedFiles(receivedFiles, process, writeFilesTimer, connectionInfo)) // Reduce refcount
+					exitCode = UBA_EXIT_CODE(20);
+			}
+		}
+		else
+		{
+			exitCode = ProcessCancelExitCode;
+		}
+
+		if (!filesWritten)
+			SkipReceivedFiles(receivedFiles, connectionInfo); // We need to inform storage server that we don't need these anymore (to reduce ref count for active casentries)
 
 		u32 id = process.m_processId;
 		Vector<ProcessLogLine> emptyLines;
-		auto& logLines = (exitCode != 0 || m_detailedTrace) ? process.m_logLines : emptyLines;
-		m_trace.ProcessExited(id, exitCode, reader.GetPositionData(), reader.GetLeft(), logLines);
+		auto& tracedLogLines = (exitCode != 0 || m_detailedTrace) ? process.m_logLines : emptyLines;
+		m_trace.ProcessExited(id, exitCode, reader.GetPositionData(), reader.GetLeft(), tracedLogLines);
 
 		ProcessStats processStats;
 		processStats.Read(reader, ~0u);
@@ -1717,6 +1840,13 @@ namespace uba
 		process.m_wallTime = processStats.wallTime;
 		process.m_peakMemory = processStats.peakJobMemory > 0 ? processStats.peakJobMemory : processStats.peakMemory;
 		process.m_server = nullptr;
+
+		// If not using the racing process we want it to call exit function with cancel first before we call the remote process exit function
+		if (hasRacingProcess)
+			if (!racingProcess.WaitForExit(60*60*1000))
+				m_logger.Error(TC("Timed out waiting for cancelled racing process!"));
+
+		process.m_exitCode = exitCode;
 		process.m_done.Set();
 		process.CallProcessExit(h);
 
@@ -1868,12 +1998,16 @@ namespace uba
 		}
 
 		auto& remoteProcess = *(RemoteProcess*)h.m_process;
+		auto receivedFiles = std::move(remoteProcess.m_receivedFiles);
 
 		Timer writeFilesTimer; // TODO: This should be added somewhere
 		if (prevExitCode == 0)
-			if (!WriteReceivedFiles(remoteProcess, writeFilesTimer, connectionInfo))
+		{
+			if (!WriteReceivedFiles(receivedFiles, remoteProcess, writeFilesTimer, connectionInfo))
 				prevExitCode = UBA_EXIT_CODE(20);
-		auto receivedFiles = std::move(remoteProcess.m_receivedFiles);
+		}
+		else
+			SkipReceivedFiles(receivedFiles, connectionInfo);
 		//processStats.writeFiles = writeFilesTimer
 
 		SCOPED_FUTEX(remoteProcess.m_exitedLock, exitedLock);
@@ -2132,6 +2266,7 @@ namespace uba
 				process->m_clientId = clientId;
 				process->m_sessionId = sessionId;
 				process->m_executingHost = session.name;
+				process->m_startTime = GetTime();
 				UBA_ASSERT(!process->m_cancelled);
 				m_activeRemoteProcesses.insert(process);
 				return process;
@@ -2195,6 +2330,18 @@ namespace uba
 		ReleaseReceivedFiles(receivedFiles);
 	}
 
+	void SessionServer::OnRaceLost(RemoteProcess& process)
+	{
+		// We want to tell client to remove this process
+		SCOPED_CRITICAL_SECTION(m_remoteProcessAndSessionLock, queueLock);
+		if (process.m_sessionId == 0) // Was disconnected
+			return;
+		u32 sessionIndex = process.m_sessionId - 1;
+		UBA_ASSERT(sessionIndex < m_clientSessions.size());
+		ClientSession& session = *m_clientSessions[sessionIndex];
+		session.processesLostRace.push_back(process.GetId());
+	}
+
 	ProcessHandle SessionServer::ProcessRemoved(u32 processId)
 	{
 		SCOPED_FUTEX(m_processesLock, lock);
diff --git a/Engine/Source/Programs/UnrealBuildAccelerator/Common/Public/UbaScheduler.h b/Engine/Source/Programs/UnrealBuildAccelerator/Common/Public/UbaScheduler.h
--- a/Engine/Source/Programs/UnrealBuildAccelerator/Common/Public/UbaScheduler.h
+++ b/Engine/Source/Programs/UnrealBuildAccelerator/Common/Public/UbaScheduler.h
@@ -38,6 +38,7 @@ namespace uba
 		bool writeToCache = false; // Set to true in combination with setting cacheClient to populate cache
 		bool useThreadToExitRemoteProcess = false; // Set this to true to move exit process call from job to thread. 
 		bool memWatchdog = false; // Will watch memory and throttle/kill processes if needed
+		u32 maxRacingPercent = 0u; // If processor weights go below this percent, then we will start filling up with racing processes
 		u8 memTracingLevel = 0;
 		u8 memStartWaitPercent = 85; // When memory usage goes above this percent, no new processes will be spawned until back below
 		u8 memStartKillPercent = 95; // When memory usage goes above this percent, newest processes will be killed to bring it back below
@@ -136,6 +137,7 @@ namespace uba
 		u32 PopProcess(bool isLocal, RunStatus& outPrevStatus);
 		bool RunCacheQuery();
 		bool RunQueuedProcess(bool isLocal);
+		bool RaceRemoteProcess();
 		bool HandleReuseMessage(Process& process, NextProcessInfo& outNextProcess, u32 prevExitCode);
 		void ExitProcess(ExitProcessInfo& info, Process& process, u32 exitCode, bool fromCache);
 		void SkipProcess(ProcessStartInfo2& info);
@@ -156,6 +158,7 @@ namespace uba
 
 		SessionServer& m_session;
 		u32 m_maxLocalProcessors;
+		float m_maxRacingWeight;
 		u32 m_maxParallelCacheQueries;
 		
 		struct ProcessEntry
@@ -189,7 +192,8 @@ namespace uba
 		bool m_cancelled = false;
 		ConfigTable* m_processConfigs = nullptr;
 
-		float m_activeLocalProcessWeight = 0.0f;
+		Atomic<float> m_activeLocalProcessWeight = 0.0f;
+		Atomic<float> m_activeRacingProcessWeight = 0.0f;
 		u32 m_activeCacheQueries = 0;
 		Atomic<u32> m_activeCacheDownloads = 0;
 
@@ -232,6 +236,9 @@ namespace uba
 		struct ProcessRecord { ProcessImpl& process; ExitProcessInfo* ep; };
 		UnorderedMap<StringKey, ProcessRecord> m_processes;
 
+		Futex m_racingProcessesLock;
+		List<ProcessHandle> m_racingProcesses;
+
 		Scheduler(const Scheduler&) = delete;
 		void operator=(const Scheduler&) = delete;
 	};
diff --git a/Engine/Source/Programs/UnrealBuildAccelerator/Common/Public/UbaSessionClient.h b/Engine/Source/Programs/UnrealBuildAccelerator/Common/Public/UbaSessionClient.h
--- a/Engine/Source/Programs/UnrealBuildAccelerator/Common/Public/UbaSessionClient.h
+++ b/Engine/Source/Programs/UnrealBuildAccelerator/Common/Public/UbaSessionClient.h
@@ -99,7 +99,7 @@ namespace uba
 		bool UpdateNameToHashTableFromServer(StackBinaryReader<SendMaxSize>& reader);
 		void Connect();
 		void BuildEnvironmentVariables(BinaryReader& reader);
-		bool SendProcessAvailable(Vector<InternalProcessStartInfo>& out, float availableWeight);
+		bool SendProcessAvailable(Vector<InternalProcessStartInfo>& outStartInfos, Vector<u32>& outCancelled, float availableWeight);
 		void SendReturnProcess(u32 processId, const tchar* reason);
 		bool SendProcessInputs(ProcessImpl& process);
 		bool SendProcessFinished(ProcessImpl& process, u32 exitCode);
diff --git a/Engine/Source/Programs/UnrealBuildAccelerator/Common/Public/UbaSessionServer.h b/Engine/Source/Programs/UnrealBuildAccelerator/Common/Public/UbaSessionServer.h
--- a/Engine/Source/Programs/UnrealBuildAccelerator/Common/Public/UbaSessionServer.h
+++ b/Engine/Source/Programs/UnrealBuildAccelerator/Common/Public/UbaSessionServer.h
@@ -29,7 +29,8 @@ namespace uba
 
 		// Will kick off a local process with the same startInfo as the one provided to start the process with id matching raceAgainstRemoteProcessId
 		// This can be useful if there are free local cores and we know local machine is faster or network connection to remote is slow
-		ProcessHandle RunProcessRacing(u32 raceAgainstRemoteProcessId);
+		using ProcessToRaceCallback = Function<bool(ProcessHandle&, ProcessStartInfo&)>;
+		ProcessHandle RunProcessRacing(const ProcessToRaceCallback& callback);
 
 		// Disable remote execution. This will tell all clients to stop taking on new processes and disconnect as soon as their current processes are finished
 		void DisableRemoteExecution();
@@ -105,10 +106,12 @@ namespace uba
 		bool WriteDirectoryTable(ClientSession& session, BinaryReader& reader, BinaryWriter& writer);
 		bool WriteNameToHashTable(BinaryReader& reader, BinaryWriter& writer, u32 requestedSize);
 
+		struct RacingProcessExitInfo;
 		class RemoteProcess;
 
 		RemoteProcess* DequeueProcess(ClientSession& session, u32 sessionId, u32 clientId);
 		void OnCancelled(RemoteProcess* process);
+		void OnRaceLost(RemoteProcess& process);
 		ProcessHandle ProcessRemoved(u32 processId);
 
 		ProcessHandle GetProcess(u32 processId);
@@ -125,8 +128,9 @@ namespace uba
 		
 		struct ReceivedFile;
 		bool WriteReceivedFile(ReceivedFile& file, RemoteProcess* process, const ConnectionInfo& connectionInfo);
-		bool WriteReceivedFiles(RemoteProcess& process, Timer& writeFilesTimer, const ConnectionInfo& connectionInfo);
-		void ReleaseReceivedFiles(const Vector<ReceivedFile>& receivedFiles);
+		bool WriteReceivedFiles(Vector<ReceivedFile>& receivedFiles, RemoteProcess& process, Timer& writeFilesTimer, const ConnectionInfo& connectionInfo);
+		void SkipReceivedFiles(Vector<ReceivedFile>& receivedFiles, const ConnectionInfo& connectionInfo);
+		void ReleaseReceivedFiles(Vector<ReceivedFile>& receivedFiles);
 
 		void TraceSessionUpdate();
 
@@ -169,6 +173,7 @@ namespace uba
 		{
 			TString name;
 			UnorderedSet<CasKey> sentKeys;
+			Vector<u32> processesLostRace;
 			Futex dirTablePosLock;
 			u32 dirTablePos = 0;
 			u32 clientId = ~0u;
diff --git a/Engine/Source/Programs/UnrealBuildAccelerator/Core/Private/UbaApplicationRules.cpp b/Engine/Source/Programs/UnrealBuildAccelerator/Core/Private/UbaApplicationRules.cpp
--- a/Engine/Source/Programs/UnrealBuildAccelerator/Core/Private/UbaApplicationRules.cpp
+++ b/Engine/Source/Programs/UnrealBuildAccelerator/Core/Private/UbaApplicationRules.cpp
@@ -153,6 +153,11 @@ namespace uba
 		{
 			return !file.EndsWith(TCV("vctip.exe")); // This is hacky but we don't want to start vctip.exe
 		}
+
+		virtual bool SupportsRacing() const override
+		{
+			return true;
+		}
 	};
 
 	class ApplicationRulesVcLink : public ApplicationRulesVC
@@ -391,6 +396,11 @@ namespace uba
 		{
 			return true; // Yes, .pch
 		}
+
+		virtual bool SupportsRacing() const override
+		{
+			return true;
+		}
 	};
 
 	class ApplicationRulesClangClExe : public ApplicationRulesClangPlusPlusExe
diff --git a/Engine/Source/Programs/UnrealBuildAccelerator/Core/Public/UbaApplicationRules.h b/Engine/Source/Programs/UnrealBuildAccelerator/Core/Public/UbaApplicationRules.h
--- a/Engine/Source/Programs/UnrealBuildAccelerator/Core/Public/UbaApplicationRules.h
+++ b/Engine/Source/Programs/UnrealBuildAccelerator/Core/Public/UbaApplicationRules.h
@@ -205,6 +205,11 @@ namespace uba
 			return false;
 		}
 
+		virtual bool SupportsRacing() const
+		{
+			return false;
+		}
+
 		u32 index = ~0u;
 	};
 
