[UBA]
* Added boxheight and horizontal scroll value to visualizer config

diff --git a/Engine/Source/Programs/UnrealBuildAccelerator/Visualizer/Private/UbaVisualizer.cpp b/Engine/Source/Programs/UnrealBuildAccelerator/Visualizer/Private/UbaVisualizer.cpp
--- a/Engine/Source/Programs/UnrealBuildAccelerator/Visualizer/Private/UbaVisualizer.cpp
+++ b/Engine/Source/Programs/UnrealBuildAccelerator/Visualizer/Private/UbaVisualizer.cpp
@@ -70,7 +70,13 @@ namespace uba
 		config.GetValueAsString(fontName, L"FontName");
 		config.GetValueAsU32(maxActiveVisible, L"MaxActiveVisible");
 		config.GetValueAsU32(maxActiveProcessHeight, L"MaxActiveProcessHeight");
-		#define UBA_VISUALIZER_FLAG(name, defaultValue, desc) config.GetValueAsBool(show##name, L"Show" #name);
+
+		config.GetValueAsU32(boxHeight, L"BoxHeight");
+		u64 temp;
+		if (config.GetValueAsU64(temp, L"HorizontalScrollValue"))
+			horizontalScaleValue = float(double(temp) / 10000000.0);
+
+#define UBA_VISUALIZER_FLAG(name, defaultValue, desc) config.GetValueAsBool(show##name, L"Show" #name);
 		UBA_VISUALIZER_FLAGS1
 		#undef UBA_VISUALIZER_FLAG
 		#define UBA_VISUALIZER_FLAG(name, defaultValue, desc) config.GetValueAsBool(name, TC(#name));
@@ -92,6 +98,10 @@ namespace uba
 		config.AddValue(L"FontName", fontName.c_str());
 		config.AddValue(L"MaxActiveVisible", maxActiveVisible);
 		config.AddValue(L"MaxActiveProcessHeight", maxActiveProcessHeight);
+
+		config.AddValue(L"BoxHeight", boxHeight);
+		config.AddValue(L"HorizontalScrollValue", u64(double(horizontalScaleValue)*10000000.0));
+
 		#define UBA_VISUALIZER_FLAG(name, defaultValue, desc) config.AddValue(L"Show" #name, show##name);
 		UBA_VISUALIZER_FLAGS1
 		#undef UBA_VISUALIZER_FLAG
@@ -895,8 +905,8 @@ namespace uba
 
 	void Visualizer::UpdateProcessFont()
 	{
-		m_zoomValue = 1.0f + float(m_boxHeight) / 30.0f;
-		int fontHeight = Max(m_boxHeight - 2, 1);
+		m_zoomValue = 1.0f + float(m_config.boxHeight) / 30.0f;
+		int fontHeight = Max(int(m_config.boxHeight) - 2, 1);
 		UpdateFont(m_processFont, fontHeight, false);
 		m_progressRectLeft = int(13 + float(m_processFont.height) * 1.5f);
 		DirtyBitmaps(true);
@@ -1047,7 +1057,7 @@ namespace uba
 
 
 		int posY = int(m_scrollPosY);
-		float scaleX = 50.0f*m_zoomValue*m_horizontalScaleValue;
+		float scaleX = 50.0f*m_zoomValue*m_config.horizontalScaleValue;
 
 		RECT progressRect = clientRect;
 		progressRect.left += m_progressRectLeft;
@@ -1212,7 +1222,7 @@ namespace uba
 				int prevX = 0;
 				int prevY = 0;
 
-				float timeScale = (m_horizontalScaleValue*m_zoomValue)*50.0f;
+				float timeScale = (m_config.horizontalScaleValue*m_zoomValue)*50.0f;
 				float startOffset = ((m_scrollPosX/timeScale) - float(int(m_scrollPosX/timeScale))) * timeScale;
 				float time = -startOffset/timeScale;
 				float playTimeS = TimeToS(playTime);
@@ -1351,7 +1361,7 @@ namespace uba
 				});
 		}
 
-		int boxHeight = m_boxHeight;
+		int boxHeight = m_config.boxHeight;
 		int stepY = int(boxHeight) + 2;
 		int processStepY = boxHeight + 1;
 
@@ -2024,7 +2034,7 @@ namespace uba
 
 		if (m_config.showCursorLine && m_mouseOverWindow)
 		{
-			float timeScale = (m_horizontalScaleValue * m_zoomValue)*50.0f;
+			float timeScale = (m_config.horizontalScaleValue * m_zoomValue)*50.0f;
 			float startOffset = -(m_scrollPosX / timeScale);
 			POINT pos;
 			GetCursorPos(&pos);
@@ -2698,7 +2708,7 @@ namespace uba
 	{
 		SetActiveFont(m_timelineFont);
 		int top = GetTimelineTop(clientRect);
-		float timeScale = (m_horizontalScaleValue*m_zoomValue)*50.0f;
+		float timeScale = (m_config.horizontalScaleValue*m_zoomValue)*50.0f;
 		float startOffset = ((m_scrollPosX/timeScale) - float(int(m_scrollPosX/timeScale))) * timeScale;
 		int index = -int(startOffset/timeScale);
 			
@@ -3003,9 +3013,9 @@ namespace uba
 		GetClientRect(m_hwnd, &clientRect);
 
 		int posY = int(m_scrollPosY);
-		int boxHeight = m_boxHeight;
+		int boxHeight = m_config.boxHeight;
 		int processStepY = boxHeight + 1;
-		float scaleX = 50.0f*m_zoomValue*m_horizontalScaleValue;
+		float scaleX = 50.0f*m_zoomValue*m_config.horizontalScaleValue;
 
 		RECT progressRect = clientRect;
 		progressRect.left += m_progressRectLeft;
@@ -3038,7 +3048,7 @@ namespace uba
 			if (pos.y >= timelineTop)// && pos.y < timelineTop + 40)
 			{
 				outResult.section = 3;
-				float timeScale = (m_horizontalScaleValue * m_zoomValue)*50.0f;
+				float timeScale = (m_config.horizontalScaleValue * m_zoomValue)*50.0f;
 				float startOffset = -(m_scrollPosX / timeScale);
 				outResult.timelineSelected = startOffset + float(pos.x - m_progressRectLeft) / timeScale;
 				return;
@@ -3087,7 +3097,7 @@ namespace uba
 		{
 			if (pos.y > posY && pos.y < posY + GraphHeight)
 			{
-				float timeScale = (m_horizontalScaleValue * m_zoomValue)*50.0f;
+				float timeScale = (m_config.horizontalScaleValue * m_zoomValue)*50.0f;
 				float startOffset = -(m_scrollPosX / timeScale);
 				u64 selectedTimeMs = u64(1000.0f*(startOffset + float(pos.x - m_progressRectLeft) / timeScale));
 
@@ -3523,13 +3533,13 @@ namespace uba
 		{
 			m_scrollPosX = 0;
 			timeS = Max(timeS, 20.0f/m_zoomValue);
-			m_horizontalScaleValue = Max(float(rect.right - m_progressRectLeft - 2)/(m_zoomValue*timeS*50.0f), 0.001f);
+			m_config.horizontalScaleValue = Max(float(rect.right - m_progressRectLeft - 2)/(m_zoomValue*timeS*50.0f), 0.001f);
 			return true;
 		}
 		else
 		{
 			float oldScrollPosX = m_scrollPosX;
-			m_scrollPosX = Min(0.0f, (float)rect.right - timeS*50.0f*m_horizontalScaleValue*m_zoomValue - float(m_progressRectLeft));
+			m_scrollPosX = Min(0.0f, (float)rect.right - timeS*50.0f*m_config.horizontalScaleValue*m_zoomValue - float(m_progressRectLeft));
 			return oldScrollPosX != m_scrollPosX;
 		}
 	}
@@ -3857,8 +3867,9 @@ namespace uba
 					GetCursorPos(&cursorPos);
 					ScreenToClient(m_hwnd, &cursorPos);
 
-					float newScaleValue = m_horizontalScaleValue;
-					int newBoxHeight = m_boxHeight;
+					float horizontalScaleValue = m_config.horizontalScaleValue;
+					float newScaleValue = horizontalScaleValue;
+					int newBoxHeight = m_config.boxHeight;
 
 					if (controlDown)
 					{
@@ -3871,24 +3882,24 @@ namespace uba
 							++newBoxHeight;
 												}
 					else
-						newScaleValue = Max(m_horizontalScaleValue + m_horizontalScaleValue*float(delta)*0.0006f, 0.001f);
+						newScaleValue = Max(horizontalScaleValue + horizontalScaleValue*float(delta)*0.0006f, 0.001f);
 
 					// TODO: m_progressRectLeft changes with zoom so anchor logic is wrong
 					const float scrollAnchorOffsetX = float(cursorPos.x) - float(m_progressRectLeft);
 					const float scrollAnchorOffsetY = 0;//float(cursorPos.y)*m_zoomValue;// - m_progressRectLeft;
 
 					float oldZoomValue = m_zoomValue;
-					if (newBoxHeight != m_boxHeight)
+					if (newBoxHeight != int(m_config.boxHeight))
 					{
-						m_boxHeight = newBoxHeight;
+						m_config.boxHeight = newBoxHeight;
 						UpdateProcessFont();
 					}
 
 					m_scrollPosY = Min(0.0f, float(m_scrollPosY - scrollAnchorOffsetY)*(m_zoomValue/oldZoomValue) + scrollAnchorOffsetY);
-					m_scrollPosX = Min(0.0f, float(m_scrollPosX - scrollAnchorOffsetX)*(m_zoomValue/oldZoomValue)*(newScaleValue/m_horizontalScaleValue) + scrollAnchorOffsetX);//LOWORD(lParam);
+					m_scrollPosX = Min(0.0f, float(m_scrollPosX - scrollAnchorOffsetX)*(m_zoomValue/oldZoomValue)*(newScaleValue/horizontalScaleValue) + scrollAnchorOffsetX);//LOWORD(lParam);
 
-					if (m_horizontalScaleValue != newScaleValue)
-						m_horizontalScaleValue = newScaleValue;
+					if (horizontalScaleValue != newScaleValue)
+						m_config.horizontalScaleValue = newScaleValue;
 
 
 					UpdateAutoscroll();
diff --git a/Engine/Source/Programs/UnrealBuildAccelerator/Visualizer/Private/UbaVisualizer.h b/Engine/Source/Programs/UnrealBuildAccelerator/Visualizer/Private/UbaVisualizer.h
--- a/Engine/Source/Programs/UnrealBuildAccelerator/Visualizer/Private/UbaVisualizer.h
+++ b/Engine/Source/Programs/UnrealBuildAccelerator/Visualizer/Private/UbaVisualizer.h
@@ -54,6 +54,8 @@ namespace uba
 		TString fontName;
 		u32 maxActiveVisible = 5;
 		u32 maxActiveProcessHeight = 16;
+		u32 boxHeight = 12;
+		float horizontalScaleValue = 0.5f;
 
 		#define UBA_VISUALIZER_FLAG(name, defaultValue, desc) bool show##name = defaultValue;
 		UBA_VISUALIZER_FLAGS1
@@ -220,7 +222,6 @@ namespace uba
 		HPEN m_activeProcPen = 0;
 		HPEN m_processUpdatePen = 0;
 		HPEN m_checkboxPen = 0;
-		int m_boxHeight = 12;
 		int m_sessionStepY = 0;
 
 		Font m_defaultFont;
