[Uba]
* Moved visualizer config to own file

diff --git a/Engine/Source/Programs/UnrealBuildAccelerator/Visualizer/Private/UbaVisualizer.cpp b/Engine/Source/Programs/UnrealBuildAccelerator/Visualizer/Private/UbaVisualizer.cpp
--- a/Engine/Source/Programs/UnrealBuildAccelerator/Visualizer/Private/UbaVisualizer.cpp
+++ b/Engine/Source/Programs/UnrealBuildAccelerator/Visualizer/Private/UbaVisualizer.cpp
@@ -45,72 +45,6 @@ namespace uba
 		Popup_Quit,
 	};
 
-	VisualizerConfig::VisualizerConfig(const tchar* fn) : filename(fn)
-	{
-		fontName = TC("Arial");
-	}
-
-	bool VisualizerConfig::Load(Logger& logger)
-	{
-		Config config;
-		if (!config.LoadFromFile(logger, filename.c_str()))
-		{
-			DWORD value = 1;
-			DWORD valueSize = sizeof(value);
-			if (RegGetValueW(HKEY_CURRENT_USER, L"Software\\Microsoft\\Windows\\CurrentVersion\\Themes\\Personalize", L"AppsUseLightTheme", RRF_RT_REG_DWORD, NULL, &value, &valueSize) == ERROR_SUCCESS)
-				DarkMode = value == 0;
-
-			return false;
-		}
-		config.GetValueAsInt(x, L"X");
-		config.GetValueAsInt(y, L"Y");
-		config.GetValueAsU32(width, L"Width");
-		config.GetValueAsU32(height, L"Height");
-		config.GetValueAsU32(fontSize, L"FontSize");
-		config.GetValueAsString(fontName, L"FontName");
-		config.GetValueAsU32(maxActiveVisible, L"MaxActiveVisible");
-		config.GetValueAsU32(maxActiveProcessHeight, L"MaxActiveProcessHeight");
-
-		config.GetValueAsU32(boxHeight, L"BoxHeight");
-		u64 temp;
-		if (config.GetValueAsU64(temp, L"HorizontalScrollValue"))
-			horizontalScaleValue = float(double(temp) / 10000000.0);
-
-#define UBA_VISUALIZER_FLAG(name, defaultValue, desc) config.GetValueAsBool(show##name, L"Show" #name);
-		UBA_VISUALIZER_FLAGS1
-		#undef UBA_VISUALIZER_FLAG
-		#define UBA_VISUALIZER_FLAG(name, defaultValue, desc) config.GetValueAsBool(name, TC(#name));
-		UBA_VISUALIZER_FLAGS2
-		#undef UBA_VISUALIZER_FLAG
-
-		fontSize = Min(fontSize, 30u);
-		return true;
-	}
-
-	bool VisualizerConfig::Save(Logger& logger)
-	{
-		Config config;
-		config.AddValue(L"X", x);
-		config.AddValue(L"Y", y);
-		config.AddValue(L"Width", width);
-		config.AddValue(L"Height", height);
-		config.AddValue(L"FontSize", fontSize);
-		config.AddValue(L"FontName", fontName.c_str());
-		config.AddValue(L"MaxActiveVisible", maxActiveVisible);
-		config.AddValue(L"MaxActiveProcessHeight", maxActiveProcessHeight);
-
-		config.AddValue(L"BoxHeight", boxHeight);
-		config.AddValue(L"HorizontalScrollValue", u64(double(horizontalScaleValue)*10000000.0));
-
-		#define UBA_VISUALIZER_FLAG(name, defaultValue, desc) config.AddValue(L"Show" #name, show##name);
-		UBA_VISUALIZER_FLAGS1
-		#undef UBA_VISUALIZER_FLAG
-		#define UBA_VISUALIZER_FLAG(name, defaultValue, desc) config.AddValue(TC(#name), name);
-		UBA_VISUALIZER_FLAGS2
-		#undef UBA_VISUALIZER_FLAG
-		return config.SaveToFile(logger, filename.c_str());
-	}
-
 	enum VisualizerFlag
 	{
 		#define UBA_VISUALIZER_FLAG(name, defaultValue, desc) VisualizerFlag_##name,
diff --git a/Engine/Source/Programs/UnrealBuildAccelerator/Visualizer/Private/UbaVisualizer.h b/Engine/Source/Programs/UnrealBuildAccelerator/Visualizer/Private/UbaVisualizer.h
--- a/Engine/Source/Programs/UnrealBuildAccelerator/Visualizer/Private/UbaVisualizer.h
+++ b/Engine/Source/Programs/UnrealBuildAccelerator/Visualizer/Private/UbaVisualizer.h
@@ -5,69 +5,10 @@
 #include "UbaNetworkClient.h"
 #include "UbaThread.h"
 #include "UbaTraceReader.h"
+#include "UbaVisualizerConfig.h"
 
 namespace uba
 {
-	class Config;
-
-	#define UBA_VISUALIZER_FLAGS1 \
-		UBA_VISUALIZER_FLAG(Progress, true, L"progress") \
-		UBA_VISUALIZER_FLAG(Status, true, L"status") \
-		UBA_VISUALIZER_FLAG(ActiveProcesses, false, L"active processes") \
-		UBA_VISUALIZER_FLAG(TitleBars, true, L"instance title bars") \
-		UBA_VISUALIZER_FLAG(DetailedData, false, L"detailed data (use -UbaDetailedTrace for even more)") \
-		UBA_VISUALIZER_FLAG(CpuMemStats, true, L"cpu/mem stats") \
-		UBA_VISUALIZER_FLAG(NetworkStats, true, L"network stats") \
-		UBA_VISUALIZER_FLAG(DriveStats, true, L"drive stats") \
-		UBA_VISUALIZER_FLAG(ActiveProcessGraph, false, L"graph of active processes over time") \
-		UBA_VISUALIZER_FLAG(ProcessBars, true, L"process bars") \
-		UBA_VISUALIZER_FLAG(FinishedProcesses, true, L"finished process bars") \
-		UBA_VISUALIZER_FLAG(Timeline, true, L"timeline") \
-		UBA_VISUALIZER_FLAG(Workers, false, L"workers (threads on host taking care of requests from helpers)") \
-		UBA_VISUALIZER_FLAG(CursorLine, false, L"cursor (vertical line)") \
-
-	#define UBA_VISUALIZER_FLAGS2 \
-		UBA_VISUALIZER_FLAG(ShowProcessText, true, L"Show text in process bars") \
-		UBA_VISUALIZER_FLAG(ShowReadWriteColors, true, L"Show colors for read/write times in process bars") \
-		UBA_VISUALIZER_FLAG(ScaleHorizontalWithScrollWheel, false, L"Use scroll wheel to scale horizontally") \
-		UBA_VISUALIZER_FLAG(DarkMode, false, L"Use dark mode to draw visualizer") \
-		UBA_VISUALIZER_FLAG(AutoSaveSettings, true, L"Auto save Position/Settings on close") \
-		UBA_VISUALIZER_FLAG(ShowAllTraces, true, L"Show all traces started on channel") \
-		UBA_VISUALIZER_FLAG(SortActiveRemoteSessions, true, L"Sort active sessions on top") \
-		UBA_VISUALIZER_FLAG(AutoScaleHorizontal, true, L"Automatically scale horizontally to fit processes") \
-		UBA_VISUALIZER_FLAG(LockTimelineToBottom, true, L"Lock timeline to always paint at bottom") \
-
-	struct VisualizerConfig
-	{
-		VisualizerConfig(const tchar* filename);
-
-		bool Load(Logger& logger);
-		bool Save(Logger& logger);
-
-		TString filename;
-
-		int x = 100;
-		int y = 100;
-		u32 width = 1500;
-		u32 height = 1500;
-		u32 fontSize = 13;
-		TString fontName;
-		u32 maxActiveVisible = 5;
-		u32 maxActiveProcessHeight = 16;
-		u32 boxHeight = 12;
-		float horizontalScaleValue = 0.5f;
-
-		#define UBA_VISUALIZER_FLAG(name, defaultValue, desc) bool show##name = defaultValue;
-		UBA_VISUALIZER_FLAGS1
-		#undef UBA_VISUALIZER_FLAG
-
-		#define UBA_VISUALIZER_FLAG(name, defaultValue, desc) bool name = defaultValue;
-		UBA_VISUALIZER_FLAGS2
-		#undef UBA_VISUALIZER_FLAG
-
-		u64 parent = 0;
-	};
-
 	class Visualizer
 	{
 	public:
diff --git a/Engine/Source/Programs/UnrealBuildAccelerator/Visualizer/Private/UbaVisualizerConfig.cpp b/Engine/Source/Programs/UnrealBuildAccelerator/Visualizer/Private/UbaVisualizerConfig.cpp
new file mode 100644
--- /dev/null
+++ b/Engine/Source/Programs/UnrealBuildAccelerator/Visualizer/Private/UbaVisualizerConfig.cpp
@@ -0,0 +1,75 @@
+// Copyright Epic Games, Inc. All Rights Reserved.
+
+#include "UbaVisualizerConfig.h"
+#include "UbaConfig.h"
+#include "UbaPlatform.h"
+
+namespace uba
+{
+	VisualizerConfig::VisualizerConfig(const tchar* fn) : filename(fn)
+	{
+		fontName = TC("Arial");
+	}
+
+	bool VisualizerConfig::Load(Logger& logger)
+	{
+		Config config;
+		if (!config.LoadFromFile(logger, filename.c_str()))
+		{
+			#if PLATFORM_WINDOWS
+			DWORD value = 1;
+			DWORD valueSize = sizeof(value);
+			if (RegGetValueW(HKEY_CURRENT_USER, L"Software\\Microsoft\\Windows\\CurrentVersion\\Themes\\Personalize", L"AppsUseLightTheme", RRF_RT_REG_DWORD, NULL, &value, &valueSize) == ERROR_SUCCESS)
+				DarkMode = value == 0;
+			#endif
+			return false;
+		}
+		config.GetValueAsInt(x, L"X");
+		config.GetValueAsInt(y, L"Y");
+		config.GetValueAsU32(width, L"Width");
+		config.GetValueAsU32(height, L"Height");
+		config.GetValueAsU32(fontSize, L"FontSize");
+		config.GetValueAsString(fontName, L"FontName");
+		config.GetValueAsU32(maxActiveVisible, L"MaxActiveVisible");
+		config.GetValueAsU32(maxActiveProcessHeight, L"MaxActiveProcessHeight");
+
+		config.GetValueAsU32(boxHeight, L"BoxHeight");
+		u64 temp;
+		if (config.GetValueAsU64(temp, L"HorizontalScrollValue"))
+			horizontalScaleValue = float(double(temp) / 10000000.0);
+
+#define UBA_VISUALIZER_FLAG(name, defaultValue, desc) config.GetValueAsBool(show##name, L"Show" #name);
+		UBA_VISUALIZER_FLAGS1
+		#undef UBA_VISUALIZER_FLAG
+		#define UBA_VISUALIZER_FLAG(name, defaultValue, desc) config.GetValueAsBool(name, TC(#name));
+		UBA_VISUALIZER_FLAGS2
+		#undef UBA_VISUALIZER_FLAG
+
+		fontSize = Min(fontSize, 30u);
+		return true;
+	}
+
+	bool VisualizerConfig::Save(Logger& logger)
+	{
+		Config config;
+		config.AddValue(L"X", x);
+		config.AddValue(L"Y", y);
+		config.AddValue(L"Width", width);
+		config.AddValue(L"Height", height);
+		config.AddValue(L"FontSize", fontSize);
+		config.AddValue(L"FontName", fontName.c_str());
+		config.AddValue(L"MaxActiveVisible", maxActiveVisible);
+		config.AddValue(L"MaxActiveProcessHeight", maxActiveProcessHeight);
+
+		config.AddValue(L"BoxHeight", boxHeight);
+		config.AddValue(L"HorizontalScrollValue", u64(double(horizontalScaleValue)*10000000.0));
+
+		#define UBA_VISUALIZER_FLAG(name, defaultValue, desc) config.AddValue(L"Show" #name, show##name);
+		UBA_VISUALIZER_FLAGS1
+		#undef UBA_VISUALIZER_FLAG
+		#define UBA_VISUALIZER_FLAG(name, defaultValue, desc) config.AddValue(TC(#name), name);
+		UBA_VISUALIZER_FLAGS2
+		#undef UBA_VISUALIZER_FLAG
+		return config.SaveToFile(logger, filename.c_str());
+	}
+}
\ No newline at end of file
diff --git a/Engine/Source/Programs/UnrealBuildAccelerator/Visualizer/Private/UbaVisualizerConfig.h b/Engine/Source/Programs/UnrealBuildAccelerator/Visualizer/Private/UbaVisualizerConfig.h
new file mode 100644
--- /dev/null
+++ b/Engine/Source/Programs/UnrealBuildAccelerator/Visualizer/Private/UbaVisualizerConfig.h
@@ -0,0 +1,69 @@
+// Copyright Epic Games, Inc. All Rights Reserved.
+
+#pragma once
+
+#include "UbaMemory.h"
+
+namespace uba
+{
+	class Config;
+	class Logger;
+
+	#define UBA_VISUALIZER_FLAGS1 \
+		UBA_VISUALIZER_FLAG(Progress, true, L"progress") \
+		UBA_VISUALIZER_FLAG(Status, true, L"status") \
+		UBA_VISUALIZER_FLAG(ActiveProcesses, false, L"active processes") \
+		UBA_VISUALIZER_FLAG(TitleBars, true, L"instance title bars") \
+		UBA_VISUALIZER_FLAG(DetailedData, false, L"detailed data (use -UbaDetailedTrace for even more)") \
+		UBA_VISUALIZER_FLAG(CpuMemStats, true, L"cpu/mem stats") \
+		UBA_VISUALIZER_FLAG(NetworkStats, true, L"network stats") \
+		UBA_VISUALIZER_FLAG(DriveStats, true, L"drive stats") \
+		UBA_VISUALIZER_FLAG(ActiveProcessGraph, false, L"graph of active processes over time") \
+		UBA_VISUALIZER_FLAG(ProcessBars, true, L"process bars") \
+		UBA_VISUALIZER_FLAG(FinishedProcesses, true, L"finished process bars") \
+		UBA_VISUALIZER_FLAG(Timeline, true, L"timeline") \
+		UBA_VISUALIZER_FLAG(Workers, false, L"workers (threads on host taking care of requests from helpers)") \
+		UBA_VISUALIZER_FLAG(CursorLine, false, L"cursor (vertical line)") \
+
+	#define UBA_VISUALIZER_FLAGS2 \
+		UBA_VISUALIZER_FLAG(ShowProcessText, true, L"Show text in process bars") \
+		UBA_VISUALIZER_FLAG(ShowReadWriteColors, true, L"Show colors for read/write times in process bars") \
+		UBA_VISUALIZER_FLAG(ScaleHorizontalWithScrollWheel, false, L"Use scroll wheel to scale horizontally") \
+		UBA_VISUALIZER_FLAG(DarkMode, false, L"Use dark mode to draw visualizer") \
+		UBA_VISUALIZER_FLAG(AutoSaveSettings, true, L"Auto save Position/Settings on close") \
+		UBA_VISUALIZER_FLAG(ShowAllTraces, true, L"Show all traces started on channel") \
+		UBA_VISUALIZER_FLAG(SortActiveRemoteSessions, true, L"Sort active sessions on top") \
+		UBA_VISUALIZER_FLAG(AutoScaleHorizontal, true, L"Automatically scale horizontally to fit processes") \
+		UBA_VISUALIZER_FLAG(LockTimelineToBottom, true, L"Lock timeline to always paint at bottom") \
+
+	struct VisualizerConfig
+	{
+		VisualizerConfig(const tchar* filename);
+
+		bool Load(Logger& logger);
+		bool Save(Logger& logger);
+
+		TString filename;
+
+		int x = 100;
+		int y = 100;
+		u32 width = 1500;
+		u32 height = 1500;
+		u32 fontSize = 13;
+		TString fontName;
+		u32 maxActiveVisible = 5;
+		u32 maxActiveProcessHeight = 16;
+		u32 boxHeight = 12;
+		float horizontalScaleValue = 0.5f;
+
+		#define UBA_VISUALIZER_FLAG(name, defaultValue, desc) bool show##name = defaultValue;
+		UBA_VISUALIZER_FLAGS1
+		#undef UBA_VISUALIZER_FLAG
+
+		#define UBA_VISUALIZER_FLAG(name, defaultValue, desc) bool name = defaultValue;
+		UBA_VISUALIZER_FLAGS2
+		#undef UBA_VISUALIZER_FLAG
+
+		u64 parent = 0;
+	};
+}
