[UBA]
* Started splitting up Visualizer class into Visualizer and VisualierWin (windows)

diff --git a/Engine/Source/Programs/UnrealBuildAccelerator/Visualizer/Private/UbaVisualizer.cpp b/Engine/Source/Programs/UnrealBuildAccelerator/Visualizer/Private/UbaVisualizer.cpp
--- a/Engine/Source/Programs/UnrealBuildAccelerator/Visualizer/Private/UbaVisualizer.cpp
+++ b/Engine/Source/Programs/UnrealBuildAccelerator/Visualizer/Private/UbaVisualizer.cpp
@@ -1,4498 +1,13 @@
 // Copyright Epic Games, Inc. All Rights Reserved.
 
 #include "UbaVisualizer.h"
-#include "UbaConfig.h"
-#include <algorithm>
-
-#include <uxtheme.h>
-#include <dwmapi.h>
-#include <shellscalingapi.h>
-
-#pragma comment (lib, "Shcore.lib")
-#pragma comment (lib, "UxTheme.lib")
-#pragma comment (lib, "Dwmapi.lib")
-#define WM_NEWTRACE WM_USER+1
-#define WM_SETTITLE WM_USER+2
 
 namespace uba
 {
-	#if PLATFORM_WINDOWS
-	extern bool g_isInAssertMessageBox;
-	#endif
-
-	enum
-	{
-		Popup_CopySessionInfo = 3,
-		Popup_CopyProcessInfo,
-		Popup_CopyProcessLog,
-		Popup_CopyProcessBreadcrumbs,
-		Popup_CopyWorkInfo,
-		Popup_Replay,
-		Popup_Pause,
-		Popup_Play,
-		Popup_JumpToEnd,
-
-		#define UBA_VISUALIZER_FLAG(name, defaultValue, desc) Popup_##name,
-		UBA_VISUALIZER_FLAGS2
-		#undef UBA_VISUALIZER_FLAG
-
-		Popup_IncreaseFontSize,
-		Popup_DecreaseFontSize,
-
-		Popup_SaveAs,
-		Popup_SaveSettings,
-		Popup_OpenSettings,
-		Popup_Quit,
-	};
-
-	enum VisualizerFlag
-	{
-		#define UBA_VISUALIZER_FLAG(name, defaultValue, desc) VisualizerFlag_##name,
-		UBA_VISUALIZER_FLAGS1
-		#undef UBA_VISUALIZER_FLAG
-		VisualizerFlag_Count
-	};
-
-	class DrawTextLogger : public Logger
-	{
-	public:
-		DrawTextLogger(HWND hw, HDC h, int fh, HBRUSH bb)
-		:	hwnd(hw)
-		,	hdc(h)
-		,	fontHeight(fh)
-		,	backgroundBrush(bb)
-		{
-			textColor = GetTextColor(hdc);
-		}
-
-		virtual void BeginScope() override {}
-		virtual void EndScope() override {}
-		virtual void Log(LogEntryType type, const wchar_t* str, u32 strLen) override
-		{
-			RECT textRect{0,0,0,0};
-			DrawTextW(hdc, str, strLen, &textRect, DT_CALCRECT);
-
-			lines.emplace_back(TString(str, strLen), textOffset, height, textColor);
-			width = Max(width, int(textRect.right + textOffset));
-			height += fontHeight;
-		}
-
-		void AddSpace(int space = 5)
-		{
-			height += space;
-		}
-
-		void AddTextOffset(int offset)
-		{
-			textOffset += offset;
-		}
-
-		void AddWidth(int extra)
-		{
-			extraWidth += extra;
-		}
-
-		void DrawAtPos(int x, int y)
-		{
-			RECT r;
-			r.left = x;
-			r.top = y;
-			r.right = r.left + width;
-			r.bottom = r.top + height;
-
-			RECT clientRect;
-			GetClientRect(hwnd, &clientRect);
-
-			if (r.right > clientRect.right)
-				OffsetRect(&r, -width - 15, 0);
-			if (r.bottom > clientRect.bottom)
-			{
-				OffsetRect(&r, 0, clientRect.bottom - r.bottom);
-				if (r.top < 0)
-					OffsetRect(&r, 0, -r.top);
-			}
-
-			RECT fillRect = r;
-			fillRect.right += 2 + extraWidth;
-			FillRect(hdc, &fillRect, backgroundBrush);
-
-			for (auto& line : lines)
-			{
-				RECT tr = r;
-				tr.left += line.left;
-				tr.top += line.top;
-				SetTextColor(hdc, line.color);
-				DrawTextW(hdc, line.str.data(), u32(line.str.size()), &tr, DT_SINGLELINE|DT_NOPREFIX);
-			}
-		}
-
-		void DrawAtCursor()
-		{
-			POINT p;
-			GetCursorPos(&p);
-			ScreenToClient(hwnd, &p);
-			p.x += 3;
-			p.y += 3;
-			DrawAtPos(p.x, p.y);
-		}
-
-		DrawTextLogger& SetColor(COLORREF c) { textColor = c; return *this; }
-
-		int width = 0;
-		int height = 0;
-		int textOffset = 2;
-		int extraWidth = 0;
-		struct Line { TString str; int left; int top; COLORREF color; };
-		Vector<Line> lines;
-
-		HWND hwnd;
-		HDC hdc;
-		int fontHeight;
-		HBRUSH backgroundBrush;
-		COLORREF textColor;
-		bool isFirst = true;
-		using Logger::Log;
-	};
-
-	class WriteTextLogger : public Logger
-	{
-	public:
-		WriteTextLogger(TString& out) : m_out(out) {}
-		virtual void BeginScope() override {}
-		virtual void EndScope() override {}
-		virtual void Log(LogEntryType type, const wchar_t* str, u32 strLen) override { m_out.append(str, strLen).append(TC("\n")); }
-		TString& m_out;
-	};
-
 	Visualizer::Visualizer(VisualizerConfig& config, Logger& logger)
 	:	m_logger(logger)
 	,	m_config(config)
 	,	m_trace(logger)
 	{
-		memset(m_activeProcessFont, 0, sizeof(m_activeProcessFont));
-		memset(m_activeProcessCountHistory, 0, sizeof(m_activeProcessCountHistory));
-	}
-
-	Visualizer::~Visualizer()
-	{
-		m_looping = false;
-
-		// Make sure GetMessage is triggered out of its slumber
-		PostMessage(m_hwnd, WM_QUIT, 0, 0);
-
-		m_thread.Wait();
-		delete m_client;
-	}
-
-	bool Visualizer::ShowUsingListener(const wchar_t* channelName)
-	{
-		TraceChannel channel(m_logger);
-		if (!channel.Init(channelName))
-		{
-			m_logger.Error(L"TODO");
-			return false;
-		}
-
-		m_listenTimeout.Create(EventResetType_Auto);
-
-		m_listenChannel.Append(channelName);
-		m_looping = true;
-		m_autoScroll = false;
-		if (!StartHwndThread())
-			return true;
-
-		{
-			StringBuffer<> title;
-			PostNewTitle(GetTitlePrefix(title).Appendf(L"Listening for new sessions on channel '%s'", m_listenChannel.data));
-		}
-
-		StringBuffer<256> traceName;
-		while (m_hwnd)
-		{
-			if (m_locked)
-			{
-				m_listenTimeout.IsSet(1000);
-				continue;
-			}
-
-			if (m_parentHwnd && !IsWindow(m_parentHwnd))
-				PostQuit();
-
-			traceName.Clear();
-			if (!channel.Read(traceName))
-			{
-				m_logger.Error(L"TODO2");
-				return false;
-			}
-
-			if (traceName.count)
-			{
-				StringBuffer<128> filter;
-				if (!m_config.ShowAllTraces)
-				{
-					OwnerInfo ownerInfo = GetOwnerInfo();
-					if (ownerInfo.pid)
-						filter.Appendf(L"_%s%u", ownerInfo.id, ownerInfo.pid);
-				}
-
-				if (!traceName.Equals(m_newTraceName.data) && traceName.EndsWith(filter))
-				{
-					m_newTraceName.Clear().Append(traceName);
-					m_usingNamed = true;
-					PostNewTrace(0, false);
-				}
-			}
-			else
-				m_newTraceName.Clear();
-
-			m_listenTimeout.IsSet(1000);
-		}
-
-		return true;
-	}
-
-	bool Visualizer::ShowUsingNamedTrace(const wchar_t* namedTrace)
-	{
-		m_looping = true;
-		if (!StartHwndThread())
-			return true;
-		m_newTraceName.Append(namedTrace);
-		m_usingNamed = true;
-		PostNewTrace(0, false);
-		return true;
-	}
-
-	bool Visualizer::ShowUsingSocket(NetworkBackend& backend, const wchar_t* host, u16 port)
-	{
-		NetworkClient* activeClient = nullptr;
-		NetworkClient* oldClient = nullptr;
-
-		auto destroyClient = MakeGuard([&]() { m_client = nullptr; delete activeClient; });
-		m_looping = true;
-		m_autoScroll = false;
-		if (!StartHwndThread())
-			return true;
-
-		m_clientDisconnect.Create(EventResetType_Manual);
-
-		wchar_t dots[] = TC("....");
-		u32 dotsCounter = 0;
-
-		StringBuffer<256> traceName;
-		while (m_hwnd)
-		{
-			if (!activeClient)
-			{
-				bool ctorSuccess = true;
-				NetworkClientCreateInfo ncci;
-				ncci.workerCount = 0;
-				activeClient = new NetworkClient(ctorSuccess, ncci);
-				if (!ctorSuccess)
-					return false;
-			}
-
-			StringBuffer<> title;
-			PostNewTitle(GetTitlePrefix(title).Appendf(L"Trying to connect to %s:%u%s", host, port, dots + ((dotsCounter--) % 4)));
-
-			bool timedOut;
-			if (!activeClient->Connect(backend, host, port, &timedOut))
-				continue;
-			m_client = activeClient;
-
-			PostNewTitle(GetTitlePrefix(title).Appendf(L"Connected to %s:%u", host, port));
-			PostNewTrace(0, false);
-
-			if (oldClient)
-			{
-				delete oldClient;
-				oldClient = nullptr;
-			}
-
-			while (m_hwnd && !m_clientDisconnect.IsSet(100) && activeClient->IsConnected())
-			{
-				m_trace.UpdateReceiveClient(*activeClient);
-			}
-
-			while (activeClient->IsConnected())
-			{
-				Sleep(500);
-				activeClient->SendKeepAlive();
-			}
-			PostNewTitle(GetTitlePrefix(title).Appendf(L"Disconnected..."));
-
-			oldClient = activeClient;
-			activeClient = nullptr;
-
-			m_clientDisconnect.Reset();
-		}
-		return true;
-	}
-
-	bool Visualizer::ShowUsingFile(const wchar_t* fileName, u32 replay)
-	{
-		m_looping = true;
-		m_autoScroll = false;
-		m_fileName.Append(fileName);
-		if (!StartHwndThread())
-			return true;
-		PostNewTrace(replay, 0);
-		return true;
-	}
-
-	bool Visualizer::StartHwndThread()
-	{
-		m_thread.Start([this]() { ThreadLoop(); return 0;}, TC("UbaHwnd"));
-		while (!m_hwnd)
-			if (m_thread.Wait(10))
-				return false;
-		return true;
-	}
-
-	bool Visualizer::HasWindow()
-	{
-		return m_looping == true;
-	}
-
-	HWND Visualizer::GetHwnd()
-	{
-		return m_hwnd;
-	}
-
-	void Visualizer::Lock(bool lock)
-	{
-		m_locked = lock;
-	}
-
-	StringBufferBase& Visualizer::GetTitlePrefix(StringBufferBase& out)
-	{
-		out.Clear();
-		out.Append(L"UbaVisualizer");
-		#if UBA_DEBUG
-		out.Append(L" (DEBUG)");
-		#endif
-		out.Append(L" - ");
-		return out;
-	}
-
-	bool Visualizer::Unselect()
-	{
-		if (m_processSelected || m_sessionSelectedIndex != ~0u || m_statsSelected || m_timelineSelected != 0 || m_fetchedFilesSelected != ~0u || m_workSelected || !m_hyperLinkSelected.empty())
-		{
-			m_processSelected = false;
-			m_sessionSelectedIndex = ~0u;
-			m_statsSelected = false;
-			m_activeProcessGraphSelected = false;
-			m_buttonSelected = ~0u;
-			m_timelineSelected = 0;
-			m_fetchedFilesSelected = ~0u;
-			m_workSelected = false;
-			m_hyperLinkSelected.clear();
-			return true;
-		}
-		return false;
-	}
-
-	void Visualizer::Reset()
-	{
-		for (HBITMAP bm : m_textBitmaps)
-			DeleteObject(bm);
-		DeleteObject(m_lastBitmap);
-		m_contentWidth = 0;
-		m_contentHeight = 0;
-		m_textBitmaps.clear();
-		m_lastBitmap = 0;
-		m_lastBitmapOffset = BitmapCacheHeight;
-		//m_autoScroll = true;
-		//m_scrollPosX = 0;
-		//m_scrollPosY = 0;
-		//m_zoomValue = 0.75f;
-		//m_horizontalScaleValue = 1.0f;
-
-		m_startTime = GetTime();
-		m_pauseTime = 0;
-
-		Unselect();
-	}
-
-	void Visualizer::InitBrushes()
-	{
-		if (m_config.DarkMode)
-		{
-			m_textColor = RGB(190, 190, 190);
-			m_textWarningColor = RGB(190, 190, 0);
-			m_textErrorColor = RGB(190, 0, 0);
-
-			m_processBrushes[0].inProgress = CreateSolidBrush(RGB(70, 70, 70));
-			m_processBrushes[1].inProgress = CreateSolidBrush(RGB(130, 130, 130));
-			m_processBrushes[0].inProgressRace = CreateSolidBrush(RGB(50, 50, 50));
-			m_processBrushes[1].inProgressRace = CreateSolidBrush(RGB(120, 120, 120));
-
-			m_processBrushes[0].error = CreateSolidBrush(RGB(140, 0, 0));
-			m_processBrushes[1].error = CreateSolidBrush(RGB(190, 0, 0));
-
-			m_processBrushes[0].returned = CreateSolidBrush(RGB(50, 50, 120));
-			m_processBrushes[1].returned = CreateSolidBrush(RGB(70, 70, 160));
-
-			m_processBrushes[0].recv = CreateSolidBrush(RGB(10, 92, 10));
-			m_processBrushes[1].recv = CreateSolidBrush(RGB(10, 130, 10));
-			m_processBrushes[0].success = CreateSolidBrush(RGB(10, 100, 10));
-			m_processBrushes[1].success = CreateSolidBrush(RGB(10, 140, 10));
-			m_processBrushes[0].send = CreateSolidBrush(RGB(10, 115, 10));
-			m_processBrushes[1].send = CreateSolidBrush(RGB(10, 145, 10));
-			m_processBrushes[0].cacheFetchTest = CreateSolidBrush(RGB(24, 120, 110));
-			m_processBrushes[1].cacheFetchTest = CreateSolidBrush(RGB(31, 150, 138));
-			m_processBrushes[0].cacheFetchDownload = CreateSolidBrush(RGB(24, 112, 110));
-			m_processBrushes[1].cacheFetchDownload = CreateSolidBrush(RGB(31, 143, 138));
-			m_processBrushes[0].warning = CreateSolidBrush(RGB(100, 100, 0));
-			m_processBrushes[1].warning = CreateSolidBrush(RGB(100, 100, 0));
-
-			m_backgroundBrush = CreateSolidBrush(0x00252526);
-			m_separatorPen = CreatePen(PS_SOLID, 1, RGB(50, 50, 50));
-			m_tooltipBackgroundBrush = CreateSolidBrush(0x00404040);
-			m_checkboxPen = CreatePen(PS_SOLID, 1, RGB(130, 130, 130));
-
-			m_sendColor = RGB(0, 170, 0);
-			m_recvColor = RGB(0, 170, 255);
-			m_cpuColor = RGB(170, 170, 0);
-			m_memColor = RGB(170, 0, 255);
-			m_driveColor = RGB(170, 65, 55);
-			m_activeProcColor = RGB(0, 170, 170);
-		}
-		else
-		{
-			m_textColor = GetSysColor(COLOR_INFOTEXT);
-			m_textWarningColor = RGB(170, 130, 0);
-			m_textErrorColor = RGB(190, 0, 0);
-
-			m_processBrushes[0].inProgress = CreateSolidBrush(RGB(150, 150, 150));
-			m_processBrushes[1].inProgress = CreateSolidBrush(RGB(180, 180, 180));
-			m_processBrushes[0].inProgressRace = CreateSolidBrush(RGB(130, 130, 130));
-			m_processBrushes[1].inProgressRace = CreateSolidBrush(RGB(170, 170, 170));
-
-			m_processBrushes[0].error = CreateSolidBrush(RGB(255, 70, 70));
-			m_processBrushes[1].error = CreateSolidBrush(RGB(255, 100, 70));
-
-			m_processBrushes[0].returned = CreateSolidBrush(RGB(150, 150, 200));
-			m_processBrushes[1].returned = CreateSolidBrush(RGB(170, 170, 200));
-
-			m_processBrushes[0].recv = CreateSolidBrush(RGB(10, 190, 10));
-			m_processBrushes[1].recv = CreateSolidBrush(RGB(20, 210, 20));
-			m_processBrushes[0].success = CreateSolidBrush(RGB(10, 200, 10));
-			m_processBrushes[1].success = CreateSolidBrush(RGB(20, 220, 20));
-			m_processBrushes[0].send = CreateSolidBrush(RGB(80, 210, 80));
-			m_processBrushes[1].send = CreateSolidBrush(RGB(90, 250, 90));
-			m_processBrushes[0].cacheFetchTest = CreateSolidBrush(RGB(150, 150, 200));
-			m_processBrushes[1].cacheFetchTest = CreateSolidBrush(RGB(170, 170, 200));
-			m_processBrushes[0].cacheFetchDownload = CreateSolidBrush(RGB(150, 150, 200));
-			m_processBrushes[1].cacheFetchDownload = CreateSolidBrush(RGB(170, 170, 200));
-			m_processBrushes[0].warning = CreateSolidBrush(RGB(170, 170, 0));
-			m_processBrushes[1].warning = CreateSolidBrush(RGB(170, 170, 0));
-
-			m_backgroundBrush = GetSysColorBrush(0);
-			m_separatorPen = CreatePen(PS_SOLID, 1, RGB(180, 180, 180));
-			m_tooltipBackgroundBrush = GetSysColorBrush(COLOR_INFOBK);
-			m_checkboxPen = CreatePen(PS_SOLID, 1, RGB(130, 130, 130));
-
-			m_sendColor = RGB(0, 170, 0); // Green
-			m_recvColor = RGB(63, 72, 204); // Blue
-			m_cpuColor = RGB(200, 130, 0); // Orange
-			m_memColor = RGB(170, 0, 255); // Purple
-			m_driveColor = RGB(255, 115, 96);
-			m_activeProcColor = RGB(0, 170, 170);
-		}
-
-		m_textPen = CreatePen(PS_SOLID, 1, m_textColor);
-		m_sendPen = CreatePen(PS_SOLID, 1, m_sendColor);
-		m_recvPen = CreatePen(PS_SOLID, 1, m_recvColor);
-		m_cpuPen = CreatePen(PS_SOLID, 1, m_cpuColor);
-		m_memPen = CreatePen(PS_SOLID, 1, m_memColor);
-		m_drivePen = CreatePen(PS_SOLID, 1, m_driveColor);
-		m_activeProcPen = CreatePen(PS_SOLID, 1, m_activeProcColor);
-	}
-
-	void Visualizer::ThreadLoop()
-	{
-		if (m_config.parent)
-		{
-			SetProcessDpiAwareness(PROCESS_SYSTEM_DPI_AWARE);
-		}
-
-		InitBrushes();
-
-		LOGBRUSH br = { 0 };
-		GetObject(m_backgroundBrush, sizeof(br), &br);
-		m_processUpdatePen = CreatePen(PS_SOLID, 2, RGB(GetRValue(br.lbColor), GetGValue(br.lbColor), GetBValue(br.lbColor)));
-
-		HINSTANCE hInstance = GetModuleHandle(NULL);
-		int winPosX = m_config.x;
-		int winPosY = m_config.y;
-		int winWidth = m_config.width;
-		int winHeight = m_config.height;
-
-		RECT rectCombined;
-		SetRectEmpty(&rectCombined);
-		EnumDisplayMonitors(0, 0, [](HMONITOR hMon,HDC hdc,LPRECT lprcMonitor,LPARAM pData)
-		{
-			RECT* rectCombined = reinterpret_cast<RECT*>(pData);
-			UnionRect(rectCombined, rectCombined, lprcMonitor);
-			return TRUE;
-		}, (LPARAM)&rectCombined);
-
-		winPosX = Max(int(rectCombined.left), winPosX);
-		winPosY = Max(int(rectCombined.top), winPosY);
-		winPosX = Min(int(rectCombined.right) - winWidth, winPosX);
-		winPosY = Min(int(rectCombined.bottom) - winHeight, winPosY);
-
-		WNDCLASSEX wndClassEx;
-		ZeroMemory(&wndClassEx, sizeof(wndClassEx));
-		wndClassEx.cbSize = sizeof(wndClassEx);
-		wndClassEx.style = CS_HREDRAW | CS_VREDRAW;
-		wndClassEx.lpfnWndProc = &StaticWinProc;
-		wndClassEx.hIcon = LoadIcon(hInstance, MAKEINTRESOURCE(123));
-		wndClassEx.hCursor = 0;
-		wndClassEx.hInstance = hInstance;
-		wndClassEx.hbrBackground = NULL;
-		wndClassEx.lpszClassName = TEXT("UbaVisualizer");
-		ATOM wndClassAtom = RegisterClassEx(&wndClassEx);
-		const TCHAR* windowClassName = MAKEINTATOM(wndClassAtom);
-
-		auto unreg = MakeGuard([&]() { UnregisterClass(windowClassName, hInstance); });
-
-		UpdateDefaultFont();
-
-		UpdateProcessFont();
-
-		const TCHAR* fontName = TEXT("Consolas");
-		m_popupFont.handle = CreateFontW(-12, 0, 0, 0, FW_NORMAL, false, false, false, ANSI_CHARSET, OUT_DEFAULT_PRECIS, CLIP_DEFAULT_PRECIS, CLEARTYPE_QUALITY, FIXED_PITCH | FF_MODERN, fontName);
-		m_popupFont.height = 14;
-		//m_popupFont = (HFONT)GetStockObject(SYSTEM_FIXED_FONT);
-
-		DWORD scrollbarFlags = 0;
-		m_verticalScrollBarEnabled = !ActiveProcessesShouldFillHeight();
-		if (m_verticalScrollBarEnabled)
-			scrollbarFlags |= WS_VSCROLL;
-		m_horizontalScrollBarEnabled = !m_config.AutoScaleHorizontal;
-		if (m_horizontalScrollBarEnabled)
-			scrollbarFlags |= WS_HSCROLL;
-
-		DWORD windowStyle = WS_CAPTION | WS_SYSMENU | WS_THICKFRAME | WS_MAXIMIZEBOX | WS_MINIMIZEBOX | WS_CLIPCHILDREN | scrollbarFlags;
-
-		DWORD exStyle = 0;
-		if (m_config.parent)
-			windowStyle = WS_POPUP | scrollbarFlags;
-
-		StringBuffer<> title;
-		GetTitlePrefix(title).Append(L"Initializing...");
-
-		HWND hwnd = CreateWindowEx(exStyle, windowClassName, title.data, windowStyle, winPosX, winPosY, winWidth, winHeight, NULL, NULL, hInstance, this);
-		SetWindowLongPtr(hwnd, GWLP_USERDATA, (LONG_PTR)this);
-		m_hwnd = hwnd;
-
-		auto destroyWin = [&]()
-			{
-				if (m_hwnd)
-				{
-					if (m_config.AutoSaveSettings)
-						SaveSettings();
-					DestroyWindow(m_hwnd);
-					m_hwnd = 0;
-				}
-			};
-
-		BOOL cloak = TRUE;
-		DwmSetWindowAttribute(hwnd, DWMWA_CLOAK, &cloak, sizeof(cloak));
-		auto exitCloak = MakeGuard([&]() { cloak = FALSE; DwmSetWindowAttribute(hwnd, DWMWA_CLOAK, &cloak, sizeof(cloak)); });
-
-		if (m_config.DarkMode)
-			UpdateTheme();
-
-		HitTestResult res;
-		HitTest(res, { -1, -1 });
-
-		if (m_config.parent)
-		{
-			exitCloak.Execute();
-
-			// If not child it will not propagate keyboard presses etc to parent
-			SetWindowLong(hwnd, GWL_STYLE, WS_CHILD | scrollbarFlags);
-
-			m_parentHwnd = (HWND)(uintptr_t)m_config.parent;
-			if (!SetParent(hwnd, m_parentHwnd))
- 				m_logger.Error(L"SetParent failed using parentHwnd 0x%llx", m_parentHwnd);
-
-			UpdateWindow(m_hwnd);
-			UpdateScrollbars(true);
-			PostMessage(m_parentHwnd, 0x0444, 0, (LPARAM)hwnd);
-		}
-		else
-		{
-			ShowWindow(hwnd, SW_SHOW);
-			UpdateWindow(m_hwnd);
-			UpdateScrollbars(true);
-			exitCloak.Execute();
-		}
-
-		m_startTime = GetTime();
-
-		while (m_looping)
-		{
-			DWORD timeoutMs = 2000;
-			DWORD result = MsgWaitForMultipleObjects(0, NULL, FALSE, timeoutMs, QS_ALLINPUT);
-			if (result == WAIT_TIMEOUT)
-				continue;
-			if (result != WAIT_OBJECT_0)
-				break;
-			MSG msg;
-			while (PeekMessage(&msg, NULL, 0, 0, PM_REMOVE))
-			{
-				TranslateMessage(&msg);
-				DispatchMessage(&msg);
-
-				// It may happen that we receive the WM_DESTOY message from within the DistpachMessage above and handle it directly in WndProc.
-				// So before trying to call GetMessage again, we just need to validate that m_looping is still true otherwise we could end
-				// up waiting forever for this thread to exit.
-				if (!m_looping || msg.message == WM_QUIT || msg.message == WM_DESTROY || msg.message == WM_CLOSE)
-				{
-					destroyWin();
-					m_looping = false;
-					if (m_listenTimeout.IsCreated())
-						m_listenTimeout.Set();
-					break;
-				}
-			}
-		}
-
-		destroyWin();
-	}
-
-	void Visualizer::Pause(bool pause)
-	{
-		if (m_paused == pause)
-			return;
-
-		m_paused = pause;
-		if (pause)
-		{
-			m_pauseStart = GetTime();
-		}
-		else
-		{
-			m_replay = 1;
-			m_pauseTime += GetTime() - m_pauseStart;
-			m_traceView.finished = false;
-			SetTimer(m_hwnd, 0, 200, NULL);
-		}
-	}
-
-	void Visualizer::StartDragToScroll(const POINT& anchor)
-	{
-		// Uses reference-counter method since multiple input events (left and middle mouse button) can trigger the drag-to-scroll mechansim
-		if (m_dragToScrollCounter == 0)
-		{
-			m_processSelected = false;
-			m_sessionSelectedIndex = ~0u;
-			m_statsSelected = false;
-			m_activeProcessGraphSelected = false;
-			m_buttonSelected = ~0u;
-			m_timelineSelected = 0;
-			m_fetchedFilesSelected = ~0u;
-			m_workSelected = false;
-			m_hyperLinkSelected.clear();
-			m_autoScroll = false;
-			m_mouseAnchor = { anchor.x, anchor.y };
-			m_scrollAtAnchorX = m_scrollPosX;
-			m_scrollAtAnchorY = m_scrollPosY;
-			SetCapture(m_hwnd);
-			Redraw(false);
-		}
-		++m_dragToScrollCounter;
-	}
-
-	void Visualizer::StopDragToScroll()
-	{
-		if (m_dragToScrollCounter > 0)
-			--m_dragToScrollCounter;
-		if (m_dragToScrollCounter != 0)
-			return;
-
-		ReleaseCapture();
-		if (UpdateSelection())
-			Redraw(false);
-	}
-
-	void Visualizer::SaveSettings()
-	{
-		RECT rect;
-		GetWindowRect(m_hwnd, &rect);
-
-		m_config.x = rect.left;
-		m_config.y = rect.top;
-		m_config.width = rect.right - rect.left;
-		m_config.height = rect.bottom - rect.top;
-		m_config.Save(m_logger);
-	}
-
-	void Visualizer::DirtyBitmaps(bool full)
-	{
-		for (auto& session : m_traceView.sessions)
-			for (auto& processor : session.processors)
-				for (auto& process : processor.processes)
-				{
-					process.bitmapDirty = true;
-					if (full)
-						process.bitmap = 0;
-				}
-
-		for (auto& workTrack : m_traceView.workTracks)
-			for (auto& work : workTrack.records)
-			{
-					work.bitmapDirty = true;
-					if (full)
-						work.bitmap = 0;
-			}
-
-		if (!full)
-			return;
-
-		for (HBITMAP bm : m_textBitmaps)
-			DeleteObject(bm);
-		DeleteObject(m_lastBitmap);
-		m_textBitmaps.clear();
-		m_lastBitmapOffset = BitmapCacheHeight;
-		m_lastBitmap = 0;
-	}
-
-	void Visualizer::PrintCacheWriteStats(Logger& logger, u32 processId)
-	{
-		auto findIt = m_traceView.cacheWrites.find(processId);
-		if (findIt == m_traceView.cacheWrites.end())
-			return;
-		TraceView::CacheWrite& write = findIt->second;
-		logger.Info(L"");
-		logger.Info(L"  -------- Cache write stats ----------");
-		if (write.end)
-		{
-			logger.Info(L"  Duration                    %9s", TimeToText(write.end - write.start).str);
-			logger.Info(L"  Success                     %9s", write.success ? L"true" : L"false");
-
-			BinaryReader reader(write.stats.data(), 0, write.stats.size());
-			CacheSendStats stats;
-			if (m_traceView.version < 45)
-			{
-				stats.sendFile.bytes = reader.Read7BitEncoded();
-				stats.sendFile.count = 1;
-			}
-			else
-				stats.Read(reader, m_traceView.version);
-			stats.Print(logger, m_traceView.frequency);
-		}
-		else
-		{
-			logger.Info(L"  In progress...");
-		}
-	}
-
-	void Visualizer::UpdateFont(Font& font, int height, bool createUnderline)
-	{
-		font.height = height;
-
-		int fh = height;
-		font.offset = 0;
-		if (height <= 13)
-		{
-			++fh;
-			--font.offset;
-		}
-		if (height <= 11)
-			++fh;
-		if (height <= 9)
-			++fh;
-		if (height <= 8)
-			++fh;
-		if (height <= 6)
-			++fh;
-		if (height <= 4)
-			--font.offset;
-
-		if (font.handle)
-			DeleteObject(font.handle);
-		if (font.handleUnderlined)
-			DeleteObject(font.handleUnderlined);
-
-		//NONCLIENTMETRICS nonClientMetrics;
-		//nonClientMetrics.cbSize = sizeof(nonClientMetrics);
-		//SystemParametersInfo(SPI_GETNONCLIENTMETRICS, sizeof(nonClientMetrics), &nonClientMetrics, 0);
-		//m_font = (HFONT)CreateFontIndirect(&nonClientMetrics.lfMessageFont);
-
-		font.handle = CreateFontW(4 - fh, 0, 0, 0, FW_NORMAL, 0, 0, 0, ANSI_CHARSET, OUT_DEFAULT_PRECIS, CLIP_DEFAULT_PRECIS, CLEARTYPE_QUALITY, DEFAULT_PITCH, m_config.fontName.c_str());
-		if (createUnderline)
-			font.handleUnderlined = CreateFontW(4 - fh, 0, 0, 0, FW_NORMAL, 0, 1, 0, ANSI_CHARSET, OUT_DEFAULT_PRECIS, CLIP_DEFAULT_PRECIS, CLEARTYPE_QUALITY, DEFAULT_PITCH, m_config.fontName.c_str());
-	}
-
-	void Visualizer::UpdateDefaultFont()
-	{
-		UpdateFont(m_defaultFont, m_config.fontSize, true);
-		m_sessionStepY = m_defaultFont.height + 4;
-
-		m_timelineFont = m_defaultFont;
-	}
-
-	void Visualizer::UpdateProcessFont()
-	{
-		m_zoomValue = 1.0f + float(m_config.boxHeight) / 30.0f;
-		int fontHeight = Max(int(m_config.boxHeight) - 2, 1);
-		UpdateFont(m_processFont, fontHeight, false);
-		m_progressRectLeft = int(13 + float(m_processFont.height) * 1.5f);
-		DirtyBitmaps(true);
-	}
-
-	void Visualizer::ChangeFontSize(int offset)
-	{
-		m_config.fontSize += offset;
-		m_config.fontSize = Max(m_config.fontSize, 10u);
-		UpdateDefaultFont();
-		Redraw(true);
-	}
-
-	void Visualizer::Redraw(bool now)
-	{
-		u32 flags = RDW_INVALIDATE;
-		if (now)
-			flags |= RDW_UPDATENOW;
-		RedrawWindow(m_hwnd, NULL, NULL, flags);
-
-		u32 activeProcessCount = u32(m_trace.m_activeProcesses.size());
-		for (u32 i=0; i!=sizeof_array(m_activeProcessCountHistory); ++i)
-			m_activeProcessCountHistory[i] = activeProcessCount;
-	}
-
-	void Visualizer::PaintClient(const Function<void(HDC hdc, HDC memDC, RECT& clientRect)>& paintFunc)
-	{
-		HDC hdc = GetDC(m_hwnd);
-
-		RECT rect;
-		GetClientRect(m_hwnd, &rect);
-
-		HDC memDC = CreateCompatibleDC(hdc);
-
-		if (!EqualRect(&m_cachedBitmapRect, &rect))
-		{
-			if (m_cachedBitmap)
-				DeleteObject(m_cachedBitmap);
-			m_cachedBitmap = CreateCompatibleBitmap(hdc, rect.right - rect.left, rect.bottom - rect.top);
-			m_cachedBitmapRect = rect;
-		}
-		HGDIOBJ oldBmp = SelectObject(memDC, m_cachedBitmap);
-		
-		if (!m_isInPaint)
-		{
-			m_isInPaint = true;
-			paintFunc(hdc, memDC, rect);
-			m_isInPaint = false;
-		}
-
-		SelectObject(memDC, oldBmp);
-		DeleteDC(memDC);
-
-		ReleaseDC(m_hwnd, hdc);
-	}
-
-	constexpr int GraphHeight = 30;
-
-	struct SessionRec
-	{
-		TraceView::Session* session;
-		u32 index;
-	};
-	void Populate(SessionRec* recs, TraceView& traceView, bool sort)
-	{
-		u32 count = u32(traceView.sessions.size());
-		for (u32 i = 0, e = count; i != e; ++i)
-			recs[i] = { &traceView.sessions[i], i };
-		if (count <= 1 || !sort)
-			return;
-		std::sort(recs + 1, recs + traceView.sessions.size(), [](SessionRec& a, SessionRec& b)
-			{
-				auto& as = *a.session;
-				auto& bs = *b.session;
-				if ((as.processActiveCount != 0) != (bs.processActiveCount != 0))
-					return as.processActiveCount > bs.processActiveCount;
-				if (as.processActiveCount && as.proxyCreated != bs.proxyCreated)
-					return int(as.proxyCreated) > int(bs.proxyCreated);
-				return a.index < b.index;
-			});
-	}
-
-	void Visualizer::PaintAll(HDC hdc, const RECT& clientRect)
-	{
-		u64 playTime = GetPlayTime();
-
-		SetBkMode(hdc, TRANSPARENT);
-		SetTextColor(hdc, m_textColor);
-		SetBkColor(hdc, m_config.DarkMode ? RGB(70, 70, 70) : RGB(180, 180, 180));
-
-		HBRUSH lastSelectedBrush = 0;
-		HBRUSH textLastSelectedBrush = 0;
-		auto fillRect = [&](const RECT& r, HBRUSH b)
-			{
-				if (lastSelectedBrush != b)
-				{
-					SelectObject(hdc, b);
-					lastSelectedBrush = b;
-				}
-				PatBlt(hdc, r.left, r.top, r.right - r.left, r.bottom - r.top, PATCOPY);
-			};
-
-		auto DrawCenteredText = [&](const StringBuffer<>* lines, u32 lineCount, bool drawBackground)
-			{
-				for (u32 i=0; i!=lineCount; ++i)
-				{
-					RECT rect = { 0, 0, 0, 0 };
-					DrawTextW(hdc, lines[i].data, lines[i].count, &rect, DT_SINGLELINE|DT_NOPREFIX|DT_NOCLIP|DT_CALCRECT);
-					int top = (clientRect.bottom - rect.bottom)/2 + i*20;
-					rect.left = (clientRect.right - rect.right)/2;
-					rect.right = (clientRect.right + rect.right)/2;
-					rect.bottom = top + (rect.bottom - rect.top);
-					rect.top = top;
-
-					if (drawBackground)
-					{
-						RECT r = rect;
-						InflateRect(&r, 4, 2);
-						fillRect(r, m_tooltipBackgroundBrush);
-					}
-					DrawTextW(hdc, lines[i].data, lines[i].count, &rect, DT_SINGLELINE|DT_NOPREFIX|DT_NOCLIP);
-				}
-			};
-
-
-		if (m_traceView.sessions.empty() && !m_parentHwnd)
-		{
-			StringBuffer<> str[2];
-			if (!m_fileName.IsEmpty())
-			{
-				str[0].Append(TCV("Loading trace file"));
-			}
-			else
-			{
-				str[1].Append(TCV("Click here to open trace file"));
-				if (m_listenChannel.count)
-				{
-					str[0].Appendf(TC("Listening for new sessions on channel '%s'"), m_listenChannel.data);
-					str[1].Append(TCV(" instead"));
-				}
-				else
-					str[0].Append(TCV("No trace active"));
-			}
-			SetActiveFont(m_popupFont);
-			DrawCenteredText(str, 2, false);
-			return;
-		}
-
-
-		int posY = int(m_scrollPosY);
-		float scaleX = 50.0f*m_zoomValue*m_config.horizontalScaleValue;
-
-		RECT progressRect = clientRect;
-		progressRect.left += m_progressRectLeft;
-
-		if (m_config.showTimeline)
-		{
-			progressRect.bottom -= m_defaultFont.height + 10;
-		}
-
-		HDC textDC = CreateCompatibleDC(hdc);
-		SetTextColor(textDC, m_textColor);
-		SelectObject(textDC, m_processFont.handle);
-		SelectObject(textDC, GetStockObject(NULL_BRUSH));
-		SetBkMode(textDC, TRANSPARENT);
-
-		//TEXTMETRIC metric;
-		//GetTextMetrics(textDC, &metric);
-
-		HBITMAP nullBmp = CreateCompatibleBitmap(hdc, 1, 1);
-		HBITMAP oldBmp = (HBITMAP)SelectObject(textDC, nullBmp);
-		HBITMAP lastSelectedBitmap = 0;
-
-		u64 lastStop = 0;
-
-		SetActiveFont(m_defaultFont);
-
-		auto drawStatusText = [&](StringView text, LogEntryType type, int posX, int endX, bool moveY, bool underlined = false, bool centered = false)
-			{
-				RECT rect;
-				rect.left = posX;
-				rect.right = endX;
-				rect.top = posY + m_activeFont.offset;
-				rect.bottom = posY + m_activeFont.height + 2;
-				SetTextColor(hdc, type == LogEntryType_Info ? m_textColor : (type == LogEntryType_Error ? m_textErrorColor : m_textWarningColor));
-
-				u32 align = 0;
-				if (centered)
-				{
-					posX += (endX - posX)/2;
-					align = SetTextAlign(hdc, TA_CENTER);
-				}
-				if (underlined)
-					SelectObject(m_activeHdc, m_activeFont.handleUnderlined);
-				ExtTextOutW(hdc, posX, posY, ETO_CLIPPED, &rect, text.data, text.count, NULL);
-				if (underlined)
-					SelectObject(m_activeHdc, m_activeFont.handle);
-				if (centered)
-					SetTextAlign(hdc, align);
-				if (moveY)
-					posY = rect.bottom;
-			};
-
-		auto drawIndentedText = [&](const StringView& text, LogEntryType type, int indent, bool moveY, bool underlined = false)
-			{
-				int posX = 5 + indent*m_defaultFont.height;
-				drawStatusText(text, type, posX, clientRect.right, moveY, underlined);
-			};
-
-		if (m_config.showProgress && m_traceView.progressProcessesTotal)
-		{
-			drawIndentedText(AsView(L"Progress"), LogEntryType_Info, 1, false);
-
-
-			float progress = float(m_traceView.progressProcessesDone) / float(m_traceView.progressProcessesTotal);
-			u32 start = 3 + 6*m_activeFont.height;
-			u32 width = m_activeFont.height * 18;
-			RECT rect;
-			rect.left = start;
-			rect.right = start + width;
-			rect.top = posY;
-			rect.bottom = posY + m_activeFont.height;
-			fillRect(rect, m_processBrushes[0].inProgress);
-
-			rect.right = rect.left + int(progress*float(width));//durationMs/100;
-			fillRect(rect, m_traceView.progressErrorCount ? m_processBrushes[0].error : m_processBrushes[0].success);
-
-			StringBuffer<> str;
-			str.Appendf(TC("%u%%    %u / %u"), u32(progress*100.0f), m_traceView.progressProcessesDone, m_traceView.progressProcessesTotal);
-
-			const tchar* remoteDisabled = TC(""); // TODO: Only show this in "advanced mode" m_traceView.remoteExecutionDisabled ? TC("   (Remote spawn disabled)") : TC("");
-			str.Append(remoteDisabled);
-
-			if (u32 active = m_traceView.totalProcessActiveCount)
-				str.Appendf(TC("   (%u active)"), active);
-
-			if (m_traceView.lastKillProcessTime && playTime >= m_traceView.lastKillProcessTime && playTime < m_traceView.lastKillProcessTime + MsToTime(800))
-			{
-				rect.left = start + width + 20;
-				rect.right = rect.left + m_activeFont.height * 7;
-				fillRect(rect, m_processBrushes[0].error);
-				drawStatusText(TCV("Killing (mem)"), LogEntryType_Info, rect.left, rect.right, 0, false, true);
-			}
-			else if (m_traceView.lastSpawningDelayStartTime && playTime >= m_traceView.lastSpawningDelayStartTime && playTime < m_traceView.lastSpawningDelayEndTime + MsToTime(500))
-			{
-				rect.left = start + width + 20;
-				rect.right = rect.left + m_activeFont.height * 7;
-				fillRect(rect, m_processBrushes[0].warning);
-				drawStatusText(TCV("Throttling (mem)"), LogEntryType_Info, rect.left, rect.right, 0, false, true);
-			}
-
-
-			drawIndentedText(str, LogEntryType_Info, 6, true);
-		}
-
-		if (m_traceView.version && (m_traceView.version < TraceReadCompatibilityVersion || m_traceView.version > TraceVersion))
-		{
-			if (!m_traceView.finished)
-			{
-				SetTextColor(hdc, m_textWarningColor);
-				StringBuffer<> str;
-				str.Appendf(TC("Unsupported trace version %u (Versions supported are %u to %u)"), m_traceView.version, TraceReadCompatibilityVersion, TraceVersion);
-				HFONT handle = CreateFontW(-20, 0, 0, 0, FW_NORMAL, 0, 0, 0, ANSI_CHARSET, OUT_DEFAULT_PRECIS, CLIP_DEFAULT_PRECIS, CLEARTYPE_QUALITY, DEFAULT_PITCH, m_config.fontName.c_str());
-				SelectObject(hdc, handle);
-				DrawTextW(hdc, str.data, str.count, (RECT*)&clientRect, DT_CENTER|DT_NOPREFIX|DT_SINGLELINE|DT_VCENTER);
-				DeleteObject(handle);
-			}
-			else
-			{
-				m_traceView.Clear();
-				StringBuffer<> title;
-				SetWindowTextW(m_hwnd, GetTitlePrefix(title).Appendf(L" (Listening for new sessions on channel '%s')", m_listenChannel.data).data);
-			}
-			return;
-		}
-
-		if (m_config.showStatus && !m_traceView.statusMap.empty())
-		{
-			u32 lastRow = ~0u;
-			u32 row = ~0u;
-			for (auto& kv : m_traceView.statusMap)
-			{
-				auto& status = kv.second;
-				if (status.text.empty())
-					continue;
-				row = u32(kv.first >> 32);
-				if (lastRow != ~0u && lastRow != row)
-					posY += m_activeFont.height + 2;
-				lastRow = row;
-				u32 column = u32(kv.first & ~0u);
-				drawIndentedText(status.text, status.type, column, false, !status.link.empty());
-			}
-			if (row != ~0u)
-				posY += m_activeFont.height + 2;
-
-			SetTextColor(hdc, m_textColor);
-			posY += 3;
-		}
-
-		if (m_config.showActiveProcessGraph)
-		{
-			if (posY + GraphHeight >= progressRect.top && posY + GraphHeight - 5 < progressRect.bottom)
-			{
-				int graphBaseY = posY + GraphHeight - 4;
-				double graphHeight = double(GraphHeight - 2);
-
-				Vector<POINT> line;
-
-				SelectObject(hdc, m_activeProcPen);
-				bool isFirstUpdate = true;
-				bool isFirstDraw = true;
-				u64 prevValue = 0;
-				int prevX = 0;
-				int prevY = 0;
-
-				float timeScale = (m_config.horizontalScaleValue*m_zoomValue)*50.0f;
-				float startOffset = ((m_scrollPosX/timeScale) - float(int(m_scrollPosX/timeScale))) * timeScale;
-				float time = -startOffset/timeScale;
-				float playTimeS = TimeToS(playTime);
-				u32 activeProcessCountIndex = 0;
-				u16 count = 0;
-				while (time < playTimeS && activeProcessCountIndex < m_traceView.activeProcessCounts.size())
-				{
-					int posX = progressRect.left + int(startOffset + time*timeScale);
-					if (posX >= clientRect.right)
-						break;
-
-					while (activeProcessCountIndex < m_traceView.activeProcessCounts.size())
-					{
-						auto& entry = m_traceView.activeProcessCounts[activeProcessCountIndex];
-						count = entry.count;
-						if (float(TimeToMs(entry.time))/1000.0f > time)
-							break;
-						++activeProcessCountIndex;
-					}
-
-					int x = posX;
-					int y = graphBaseY;
-
-					double div = double(m_traceView.maxActiveProcessCount);
-					y -= int(double(count) * graphHeight / div);
-
-					if (prevValue <= 0)
-						prevY = y;
-
-					if (x > clientRect.left)
-					{
-						if (isFirstUpdate)
-						{
-							line.push_back({ x, y });
-							isFirstDraw = false;
-						}
-						else
-						{
-							if (isFirstDraw)
-								line.push_back({ prevX, prevY });
-							line.push_back({ x, y });
-							isFirstDraw = false;
-						}
-					}
-					if (x > clientRect.right)
-						break;
-					isFirstUpdate = false;
-					prevX = x;
-					prevY = y;
-					prevValue = count;
-
-					time += 0.25f;
-				}
-				if (line.size() > 1)
-					Polyline(hdc, line.data(), int(line.size()));
-			}
-			posY += GraphHeight;
-		}
-
-		if (m_config.showActiveProcesses && !m_trace.m_activeProcesses.empty())
-		{
-			auto drawBox = [&](u64 start, u64 stop, int height, bool selected, bool inProgress)
-				{
-					//int left = 3 + 6*m_activeFont.height;
-					//int right = rect.left + durationMs/100;
-					int posX = int(m_scrollPosX) + progressRect.left;
-					bool done = stop != ~u64(0);
-					if (!done)
-						stop = playTime;
-					int left = posX + int(float(TimeToS(start)) * scaleX);
-					int right = posX + int(float(TimeToS(stop)) * scaleX) - 1;
-
-					RECT rect;
-					rect.left = left;
-					rect.right = right;
-					rect.top = posY;
-					rect.bottom = posY + height;
-					fillRect(rect, inProgress ? m_processBrushes[selected].inProgress : m_processBrushes[selected].success);
-					return rect;
-				};
-
-			m_activeProcessCountHistory[m_activeProcessCountHistoryIterator++ % sizeof_array(m_activeProcessCountHistory)] = u32(m_trace.m_activeProcesses.size());
-
-			PaintActiveProcesses(posY, clientRect, [&](TraceView::ProcessLocation& processLocation, u32 boxHeight, bool firstWithHeight)
-				{
-					auto& session = m_trace.GetSession(m_traceView, processLocation.sessionIndex);
-					TraceView::Process& process = session.processors[processLocation.processorIndex].processes[processLocation.processIndex];
-
-					bool selected = m_processSelected && m_processSelectedLocation == processLocation;
-
-					if (m_config.showFinishedProcesses)
-					{
-						u32 index = processLocation.processIndex;
-						while (index > 0)
-						{
-							--index;
-							TraceView::Process& process2 = session.processors[processLocation.processorIndex].processes[index];
-							drawBox(process2.start, process2.stop, boxHeight, false, false);
-						}
-					}
-
-					RECT boxRect = drawBox(process.start, process.stop, boxHeight, selected, true);
-
-					u32 v = boxHeight - 1;
-
-					if (v > 4)
-					{
-						u32 fontIndex = Min(v, u32(sizeof_array(m_activeProcessFont) - 1));
-						if (!m_activeProcessFont[fontIndex].handle)
-						{
-							UpdateFont(m_activeProcessFont[fontIndex], fontIndex - 1, false);
-							m_activeProcessFont[fontIndex].offset += 1;
-						}
-						if (firstWithHeight)
-							SetActiveFont(m_activeProcessFont[fontIndex]);
-
-						StringBuffer<> str;
-						auto firstPar = -1;//process.description.find_first_of('(');
-						if (firstPar != -1 && *process.description.rbegin() == ')')
-						{
-							str.Append(process.description.c_str() + firstPar + 1).Resize(str.count-1);
-							str.Append(L"  ").Append(process.description.c_str(), firstPar);
-						}
-						else
-							str.Append(process.description);
-
-						if (process.isRemote)
-							str.Append(L" [").Append(session.name).Append(']');
-						else if (process.type == TraceView::ProcessType_CacheFetchTest || process.type == TraceView::ProcessType_CacheFetchDownload)
-							str.Append(L" [cache]");
-						if (boxRect.left < 0)
-							str.Appendf(L"   %s", TimeToText(playTime - process.start, true).str);
-						//drawIndentedText(str, LogEntryType_Info, 1, true);
-						drawStatusText(str, LogEntryType_Info, Max(int(boxRect.left+1), 1), boxRect.right, false);
-					}
-				});
-		}
-
-		int boxHeight = m_config.boxHeight;
-		int stepY = int(boxHeight) + 2;
-		int processStepY = boxHeight + 1;
-
-		TraceView::WorkRecord selectedWork;
-
-		SessionRec sortedSessions[1024];
-		Populate(sortedSessions, m_traceView, m_config.SortActiveRemoteSessions);
-
-		u32 visibleBoxes = 0;
-
-		TraceView::ProcessLocation processLocation { 0, 0, 0 };
-		for (u64 sessionIt = 0, sessionEnd = m_traceView.sessions.size(); sessionIt != sessionEnd; ++sessionIt)
-		{
-			bool isFirst = sessionIt == 0;
-			auto& session = *sortedSessions[sessionIt].session;
-			bool hasUpdates = !session.updates.empty();
-			if (!isFirst)
-			{
-				if (!hasUpdates && session.processors.empty())
-					continue;
-
-				if (!m_config.showFinishedProcesses && session.disconnectTime != ~u64(0))
-					continue;
-			}
-
-			processLocation.sessionIndex = sortedSessions[sessionIt].index;
-			if (!isFirst)
-				posY += 3;
-
-			if (m_config.showTitleBars)
-			{
-				if (posY + stepY >= progressRect.top && posY <= progressRect.bottom)
-				{
-					SelectObject(hdc, m_separatorPen);
-					MoveToEx(hdc, 0, posY, NULL);
-					LineTo(hdc, clientRect.right, posY);
-
-					StringBuffer<> text;
-					session.GetFullName(text);
-
-					if (hasUpdates && session.disconnectTime == ~u64(0))
-					{
-						u64 ping = session.ping.back();
-						//u64 memAvail = session.updates.back().memAvail;
-						//float cpuLoad = session.updates.back().cpuLoad;
-
-						//text.Appendf(L" - Cpu: %.1f%%", cpuLoad * 100.0f);
-						//if (memAvail)
-						//	text.Appendf(L" Mem: %ls/%ls", BytesToText(session.memTotal - memAvail).str, BytesToText(session.memTotal).str);
-						if (ping)
-							text.Appendf(L" Ping: %ls", TimeToText(ping, false, m_traceView.frequency).str);
-						if (!session.notification.empty())
-							text.Append(L" - ").Append(session.notification);
-					}
-					else if (!isFirst)
-					{
-						text.Append(L" - Disconnected");
-						if (!session.notification.empty())
-							text.Append(L" (").Append(session.notification).Append(')');
-					}
-
-					bool selected = m_sessionSelectedIndex == processLocation.sessionIndex;
-
-					int textBottom = Min(posY + m_sessionStepY, int(progressRect.bottom));
-
-					RECT rect;
-					rect.left = 5;
-					rect.right = clientRect.right;
-					rect.top = posY;
-					rect.bottom = textBottom;
-
-					if (selected)
-						SetBkMode(hdc, OPAQUE);
-
-					SIZE textSize;
-					if (GetTextExtentPoint32W(hdc, text.data, text.count, &textSize))
-						session.fullNameWidth = u32(textSize.cx);
-					ExtTextOutW(hdc, 5, posY+2, ETO_CLIPPED, &rect, text.data, text.count, NULL);
-					if (selected)
-						SetBkMode(hdc, TRANSPARENT);
-				}
-				posY += m_sessionStepY;
-			}
-
-			bool showGraph = m_config.showNetworkStats || m_config.showCpuMemStats || m_config.showDriveStats;
-			if (showGraph && hasUpdates)
-			{
-				if (posY + GraphHeight >= progressRect.top && posY + GraphHeight - 5 < progressRect.bottom)
-				{
-					int posX = int(m_scrollPosX) + progressRect.left;
-					int graphBaseY = posY + GraphHeight - 4;
-					double graphHeight = double(GraphHeight - 2);
-
-					Vector<POINT> line;
-
-					auto drawGraph = [&](auto& values, auto maxValue, double scale, HPEN pen, bool accumulatingValue, int offsetY = 0)
-						{
-							bool loop = true;
-							u32 reconnectIndex = 0;
-							while (loop)
-							{
-								SelectObject(hdc, pen);
-								bool isFirstUpdate = true;
-								bool isFirstDraw = true;
-								decltype(maxValue) prevValue = 0;
-								u64 prevTime = 0;
-								int prevX = 0;
-								int prevPrevX = 0;
-								int prevY = 0;
-
-								u64 i = 0;
-								u64 e = 0;
-
-								if (reconnectIndex)
-									i = session.reconnectIndices[reconnectIndex - 1];
-
-								if (reconnectIndex < session.reconnectIndices.size())
-									e = session.reconnectIndices[reconnectIndex];
-								else
-								{
-									e = Min(session.updates.size(), values.size());
-									loop = false;
-								}
-
-								line.clear();
-								for (; i!=e; ++i)
-								{
-									u64 updateTime = session.updates[i];
-									auto value = values[i];
-									int x = posX + int(float(TimeToS(updateTime)) * scaleX);
-									int y = graphBaseY;
-
-									double duration = TimeToS(updateTime - prevTime);
-									if (updateTime == 0)
-										isFirstUpdate = true;
-
-									if (accumulatingValue)
-									{
-										double invScaleY = duration * scale;
-										if (invScaleY != 0 && prevValue != 0)
-											y -= int(double(value - prevValue) / invScaleY) + offsetY;
-									}
-									else
-									{
-										double div = 1.0f;
-										if (maxValue != 0)
-											div = double(maxValue);
-										y -= int(double((double)maxValue - (double)value*scale)*graphHeight/div) + offsetY;
-									}
-
-									if (prevValue <= 0)
-										prevY = y;
-
-									if (x > clientRect.left)
-									{
-										if (isFirstUpdate)
-										{
-											isFirstDraw = false;
-										}
-										else
-										{
-											if (isFirstDraw)
-												line.push_back({prevPrevX, prevY});
-											line.push_back({prevX, y});
-											isFirstDraw = false;
-										}
-									}
-									if (prevX > clientRect.right)
-										break;
-									isFirstUpdate = false;
-									prevPrevX = prevX;
-									prevX = x;
-									prevY = y;
-									prevValue = value;
-									prevTime = updateTime;
-								}
-								if (line.size() > 1)
-									Polyline(hdc, line.data(), int(line.size()));
-								++reconnectIndex;
-							}
-						};
-
-					if (m_config.showNetworkStats && session.highestSendPerS != 0 && session.highestRecvPerS != 0)
-					{
-						drawGraph(session.networkSend, 100000000000000ull, double(session.highestSendPerS) / graphHeight, m_sendPen, true);
-						drawGraph(session.networkRecv, 100000000000000ull, double(session.highestRecvPerS) / graphHeight, m_recvPen, true, 1);
-					}
-
-					if (m_config.showCpuMemStats)
-					{
-						drawGraph(session.cpuLoad, 0.0f, -1.0, m_cpuPen, false);
-						drawGraph(session.memAvail, session.memTotal, 1.0, m_memPen, false);
-					}
-
-					if (m_config.showDriveStats)
-						for (auto& pair : session.drives)
-							if (pair.second.busyHighest)
-								drawGraph(pair.second.busyPercent, 0, -0.01, m_drivePen, false);
-
-				}
-				posY += GraphHeight;
-			}
-
-			if (m_config.showDetailedData)
-			{
-				auto drawText = [&](const StringBufferBase& text, RECT& rect, u32* outWidth)
-					{
-						if (rect.top > progressRect.bottom)
-							return;
-						bool selected = m_fetchedFilesSelected == processLocation.sessionIndex && text.StartsWith(TC("Fetched Files"));
-						if (selected)
-							SetBkMode(hdc, OPAQUE);
-						if (rect.bottom > progressRect.bottom)
-						{
-							rect.bottom = progressRect.bottom;
-							ExtTextOutW(hdc, rect.left, rect.top, ETO_CLIPPED, &rect, text.data, text.count, NULL);
-						}
-						else
-							ExtTextOutW(hdc, rect.left, rect.top, 0, NULL, text.data, text.count, NULL);
-						if (selected)
-							SetBkMode(hdc, TRANSPARENT);
-						if (outWidth)
-						{
-							SIZE s;
-							GetTextExtentPoint32W(hdc, text.data, text.count, &s);
-							*outWidth = u32(s.cx);
-						}
-					};
-				bool isRemote = processLocation.sessionIndex != 0;
-				PaintDetailedStats(posY, progressRect, session, isRemote, playTime, drawText);
-			}
-
-			SetActiveFont(m_processFont);
-
-			bool shouldDrawText = m_processFont.height > 4;
-
-			if (m_config.showProcessBars)
-			{
-				processLocation.processorIndex = 0;
-				for (auto& processor : session.processors)
-				{
-					bool drawProcessorIndex = m_config.showFinishedProcesses;
-
-					if (posY + m_sessionStepY >= progressRect.top && posY < progressRect.bottom)
-					{
-						int barHeight = boxHeight;
-						int textOffsetY = 0;
-						if (posY + boxHeight > progressRect.bottom)
-						{
-							int newBarHeight = Min(barHeight, int(progressRect.bottom - posY));
-							textOffsetY = int(barHeight - newBarHeight);
-							barHeight = newBarHeight;
-						}
-
-						const int textHeight = barHeight;
-						const int rectBottom = posY + textHeight;
-						const int offsetY = (textHeight - m_processFont.height + textOffsetY) / 2;
-
-						processLocation.processIndex = 0;
-						int posX = int(m_scrollPosX) + progressRect.left;
-						for (auto& process : processor.processes)
-						{
-							int left = posX + int(float(TimeToS(process.start)) * scaleX);
-
-							auto pig = MakeGuard([&]() { ++processLocation.processIndex; });
-
-							if (left >= progressRect.right)
-								continue;
-
-							u64 stop = process.stop;
-							bool done = stop != ~u64(0);
-							if (!done)
-								stop = playTime;
-							else if (!m_config.showFinishedProcesses)
-								continue;
-
-							drawProcessorIndex = true;
-
-							RECT rect;
-							rect.left = left;
-							rect.right = Max(posX + int(float(TimeToS(stop)) * scaleX) - 1, left + 1);
-							rect.top = posY;
-							rect.bottom = rectBottom - 1;
-
-							if (rect.right <= progressRect.left)
-								continue;
-
-							if (!m_filterString.empty() && !Contains(process.description.c_str(), m_filterString.c_str()) && !Contains(process.breadcrumbs.c_str(), m_filterString.c_str()))
-								continue;
-							++visibleBoxes;
-
-							bool selected = m_processSelected && m_processSelectedLocation == processLocation;
-							if (selected)
-								process.bitmapDirty = true;
-
-							--rect.top;
-							PaintProcessRect(process, hdc, rect, progressRect, selected, false, lastSelectedBrush);
-							++rect.top;
-
-							int processWidth = rect.right - rect.left;
-							if (shouldDrawText && m_config.ShowProcessText && processWidth > 3)
-							{
-								if (!process.bitmap || process.bitmapDirty)
-								{
-									if (!process.bitmap)
-									{
-										if (m_lastBitmapOffset == BitmapCacheHeight)
-										{
-											if (m_lastBitmap)
-												m_textBitmaps.push_back(m_lastBitmap);
-
-											m_lastBitmapOffset = 0;
-											m_lastBitmap = CreateCompatibleBitmap(hdc, 256, BitmapCacheHeight);
-										}
-										process.bitmap = m_lastBitmap;
-										process.bitmapOffset = m_lastBitmapOffset;
-										m_lastBitmapOffset += m_processFont.height;
-									}
-									if (lastSelectedBitmap != process.bitmap)
-									{
-										SelectObject(textDC, process.bitmap);
-										lastSelectedBitmap = process.bitmap;
-									}
-
-									RECT rect2{ 0, int(process.bitmapOffset), 256, int(process.bitmapOffset) + m_processFont.height };
-									RECT rect3{ 0, int(process.bitmapOffset), processWidth, int(process.bitmapOffset) + m_processFont.height };
-									if (!done)
-										rect3.right = 256;
-
-									PaintProcessRect(process, textDC, rect3, rect2, selected, true, textLastSelectedBrush);
-
-									rect2.left += 3; // Move in text a bit
-									
-									int textY = rect2.top + m_processFont.offset;
-
-									bool dropShadow = m_config.DarkMode;
-									if (dropShadow)
-									{
-										SetTextColor(textDC, RGB(5, 60, 5));
-										++rect2.left;
-										++textY;
-										ExtTextOutW(textDC, rect2.left, textY, ETO_CLIPPED, &rect2, process.description.c_str(), int(process.description.size()), NULL);
-										--rect2.left;
-										--textY;
-									}
-									SetTextColor(textDC, m_textColor);
-									ExtTextOutW(textDC, rect2.left, textY, ETO_CLIPPED, &rect2, process.description.c_str(), int(process.description.size()), NULL);
-
-									if (!selected)
-										process.bitmapDirty = false;
-								}
-
-								if (lastSelectedBitmap != process.bitmap)
-								{
-									SelectObject(textDC, process.bitmap);
-									lastSelectedBitmap = process.bitmap;
-								}
-
-								int width = Min(processWidth, 256);
-								int bitmapOffsetY = process.bitmapOffset;
-								int bltOffsetY = offsetY;
-								if (bltOffsetY < 0)
-								{
-									bitmapOffsetY -= bltOffsetY;
-									bltOffsetY = 0;
-								}
-								int height = Min(textHeight, m_processFont.height);
-								if (bltOffsetY + height > textHeight)
-									height = textHeight - bltOffsetY;
-
-								if (left > -256 && height >= 0)
-								{
-									int bitmapOffsetX = rect.left - left;
-
-									if (left < progressRect.left)
-									{
-										int diff = progressRect.left - left;
-										rect.left = progressRect.left;
-										width -= diff;
-										bitmapOffsetX += diff;
-									}
-									BitBlt(hdc, rect.left, rect.top + bltOffsetY, width, height, textDC, bitmapOffsetX, bitmapOffsetY, SRCCOPY);
-								}
-							}
-						}
-
-						if (drawProcessorIndex)
-						{
-							RECT rect;
-							rect.left = 5;
-							rect.right = progressRect.left - 2;
-							rect.top = posY;
-							rect.bottom = rectBottom;
-
-							StringBuffer<> buf;
-							buf.AppendValue(u64(processLocation.processorIndex) + 1);
-							ExtTextOutW(hdc, 5, posY + offsetY, ETO_CLIPPED, &rect, buf.data, buf.count, NULL);
-						}
-					}
-
-					lastStop = Max(lastStop, processor.processes.rbegin()->stop);
-
-					++processLocation.processorIndex;
-
-					if (drawProcessorIndex)
-						posY += processStepY;
-				}
-			}
-			else
-			{
-				for (auto& processor : session.processors)
-					if (!processor.processes.empty())
-						lastStop = Max(lastStop, processor.processes.rbegin()->stop);
-			}
-
-			if (m_config.showWorkers && isFirst)
-			{
-				u32 trackIndex = 0;
-				for (auto& workTrack : m_traceView.workTracks)
-				{
-					if (posY + m_sessionStepY >= progressRect.top && posY <= progressRect.bottom)
-					{
-						int textOffsetY = 0;
-						int barHeight = boxHeight;
-						if (posY + int(boxHeight) > progressRect.bottom)
-						{
-							int newBarHeight = Min(barHeight, int(progressRect.bottom - posY));
-							textOffsetY = barHeight - newBarHeight;
-							barHeight = newBarHeight;
-						}
-
-						const int textHeight = barHeight;
-						const int rectBottom = posY + textHeight;
-						const int offsetY = (textHeight - m_processFont.height + textOffsetY) / 2;
-
-						if (shouldDrawText)
-						{
-							RECT rect;
-							rect.left = 5;
-							rect.right = progressRect.left - 5;
-							rect.top = posY;
-							rect.bottom = rectBottom;
-
-							StringBuffer<> buf;
-							buf.AppendValue(u64(trackIndex) + 1);
-							ExtTextOutW(hdc, 5, posY + offsetY, ETO_CLIPPED, &rect, buf.data, buf.count, NULL);
-						}
-
-						int lastDrawnRight = 0;
-
-						u32 workIndex = 0;
-						int posX = int(m_scrollPosX) + progressRect.left;
-						for (auto& work : workTrack.records)
-						{
-							auto inc = MakeGuard([&](){ ++workIndex; });
-
-							if (work.start == work.stop)
-								continue;
-
-							float startTime = TimeToS(work.start);
-
-							int left = posX + int(startTime * scaleX);
-
-							if (left >= progressRect.right)
-								continue;
-
-							if (!m_filterString.empty())
-							{
-								bool keep = Contains(work.description, m_filterString.c_str());
-								if (!keep)
-									for (auto& en : work.entries)
-										keep |= Contains(en.text, m_filterString.c_str());
-								if (!keep)
-									continue;
-							}
-
-							HBRUSH brush;
-
-							u64 stop = work.stop;
-
-							float stopTime = TimeToS(stop);
-
-							RECT rect;
-							rect.left = left;
-							rect.right = posX + int(stopTime * scaleX) - 1;
-							rect.top = posY;
-							rect.bottom = rectBottom - 1;
-
-							if (rect.right <= progressRect.left)
-								continue;
-
-							++visibleBoxes;
-
-							rect.right = Max(int(rect.right), left + 1);
-
-							Color color = work.color;
-							bool selected = m_workSelected && m_workTrack == trackIndex && m_workIndex == workIndex;
-							if (selected)
-							{
-								selectedWork = work;
-								work.bitmapDirty = true;
-								auto c = (u8*)&color;
-								for (u32 j=0; j!=3; ++j)
-									c[j] = u8(Min(int(c[j]) + 40, 255));
-							}
-							else
-							{
-								if (rect.left + 1 == rect.right)
-									if (lastDrawnRight == rect.right)
-										continue;
-							}
-
-							lastDrawnRight = rect.right;
-
-							bool done = stop != ~u64(0);
-							if (done)
-							{
-								auto insres = m_coloredBrushes.try_emplace(color);
-								if (insres.second)
-								{
-									auto c = (u8*)&color;
-									insres.first->second = CreateSolidBrush(RGB(c[2], c[1], c[0]));
-								}
-								brush = insres.first->second;
-							}
-							else
-							{
-								//stop = playTime;
-								brush = m_processBrushes[0].inProgress;
-							}
-
-							--rect.top;
-
-							auto clampRect = [&](RECT& r) { r.left = Min(Max(r.left, progressRect.left), progressRect.right); r.right = Max(Min(r.right, progressRect.right), progressRect.left); };
-
-							clampRect(rect);
-							fillRect(rect, brush);
-							++rect.top;
-
-							int processWidth = rect.right - rect.left;
-							if (shouldDrawText && m_config.ShowProcessText && processWidth > 3)
-							{
-								if (!work.bitmap || work.bitmapDirty)
-								{
-									if (!work.bitmap)
-									{
-										if (m_lastBitmapOffset == BitmapCacheHeight)
-										{
-											if (m_lastBitmap)
-												m_textBitmaps.push_back(m_lastBitmap);
-
-											m_lastBitmapOffset = 0;
-											m_lastBitmap = CreateCompatibleBitmap(hdc, 256, BitmapCacheHeight);
-										}
-										work.bitmap = m_lastBitmap;
-										work.bitmapOffset = m_lastBitmapOffset;
-										m_lastBitmapOffset += m_processFont.height;
-									}
-									if (lastSelectedBitmap != work.bitmap)
-									{
-										SelectObject(textDC, work.bitmap);
-										lastSelectedBitmap = work.bitmap;
-									}
-
-									RECT rect2{ 0,int(work.bitmapOffset), 256, int(work.bitmapOffset) + m_processFont.height };
-
-									SelectObject(textDC, brush);
-									PatBlt(textDC, rect2.left, rect2.top, rect2.right - rect2.left, rect2.bottom - rect2.top, PATCOPY);
-									ExtTextOutW(textDC, rect2.left, rect2.top, ETO_CLIPPED, &rect2, work.description, int(wcslen(work.description)), NULL);
-									if (!selected)
-										work.bitmapDirty = false;
-								}
-								else if (lastSelectedBitmap != work.bitmap)
-								{
-									SelectObject(textDC, work.bitmap);
-									lastSelectedBitmap = work.bitmap;
-								}
-
-								int width = Min(processWidth, 256);
-								int bitmapOffsetY = work.bitmapOffset;
-								int bltOffsetY = offsetY;
-								if (bltOffsetY < 0)
-								{
-									bitmapOffsetY -= bltOffsetY;
-									bltOffsetY = 0;
-								}
-								int height = Min(textHeight, m_processFont.height);
-								if (bltOffsetY + height > textHeight)
-									height = textHeight - bltOffsetY;
-
-								if (left > -256 && height >= 0)
-								{
-									int bitmapOffsetX = rect.left - left;
-
-									if (left < progressRect.left)
-									{
-										int diff = progressRect.left - left;
-										rect.left = progressRect.left;
-										width -= diff;
-										bitmapOffsetX += diff;
-									}
-									BitBlt(hdc, rect.left, rect.top + bltOffsetY, width, height, textDC, bitmapOffsetX, bitmapOffsetY, SRCCOPY);
-								}
-							}
-						}
-					}
-					++trackIndex;
-					posY += processStepY;
-				}
-			}
-
-			SetActiveFont(m_defaultFont);
-		}
-
-		if (true)
-		{
-			for (auto& kv : m_traceView.tables)
-			{
-				auto& table = kv.second;
-				u32 tableOffsetX = m_activeFont.height*40;
-				u32 tableOffsetY = 30;
-
-				u32 rowOffset = m_activeFont.height + 2;
-				u32 columnOffset = table.rowCharWidth * m_activeFont.height;
-
-				u32 columnIndex = 0;
-				u32 rowIndex = 0;
-				for (const tchar* item : table.items)
-				{
-					u32 y = tableOffsetY + rowIndex*rowOffset;
-					u32 x = tableOffsetX + columnIndex*columnOffset;
-					if (int(x) >= progressRect.right)
-						break;
-
-					ExtTextOutW(hdc, x, y, ETO_CLIPPED, nullptr, item, TStrlen(item), NULL);
-
-					++rowIndex;
-					if (rowIndex > table.rowMaxCount || int(y) >= progressRect.bottom)
-					{
-						rowIndex = 0;
-						++columnIndex;
-					}
-				}
-			}
-		}
-
-		SelectObject(textDC, oldBmp);
-		DeleteObject(nullBmp);
-		DeleteDC(textDC);
-
-		m_contentWidth = m_progressRectLeft + Max(0, int(TimeToS((lastStop != 0 && lastStop != ~u64(0)) ? lastStop : playTime) * scaleX));
-		
-
-		m_contentHeight = posY - int(m_scrollPosY) + stepY + 14;
-
-		float timelineSelected = m_timelineSelected;
-
-		if (m_config.showTimeline && !m_traceView.sessions.empty())
-			PaintTimeline(hdc, clientRect);
-
-		if (!m_filterString.empty())
-		{
-			StringBuffer<> str[2];
-			str[0].Append(TCV("Box Filter: ")).Append(m_filterString);
-			str[1].Appendf(TC("Box Count: %u"), visibleBoxes);
-
-			SetActiveFont(m_popupFont);
-			DrawCenteredText(str, 2, true);
-		}
-
-		if (m_config.showCursorLine && m_mouseOverWindow)
-		{
-			float timeScale = (m_config.horizontalScaleValue * m_zoomValue)*50.0f;
-			float startOffset = -(m_scrollPosX / timeScale);
-			POINT pos;
-			GetCursorPos(&pos);
-			ScreenToClient(m_hwnd, &pos);
-			timelineSelected = startOffset + float(pos.x - m_progressRectLeft) / timeScale;
-		}
-
-		if (timelineSelected != 0)
-		{
-			int posX = int(m_scrollPosX) + progressRect.left;
-			int left = posX + int(timelineSelected * scaleX);
-			int timelineTop = GetTimelineTop(clientRect);
-
-			// TODO: Draw line up
-			MoveToEx(hdc, left, 2, NULL);
-			LineTo(hdc, left, timelineTop);
-
-			if (timelineSelected >= 0)
-			{
-				StringBuffer<> b;
-				u32 milliseconds = u32(timelineSelected * 1000.0f);
-				u32 seconds = milliseconds / 1000;
-				milliseconds -= seconds * 1000;
-				u32 minutes = seconds / 60;
-				seconds -= minutes * 60;
-				u32 hours = minutes / 60;
-				minutes -= hours * 60;
-				if (hours)
-				{
-					b.AppendValue(hours).Append('h');
-					if (minutes < 10)
-						b.Append('0');
-				}
-				if (minutes || hours)
-				{
-					b.AppendValue(minutes).Append('m');
-					if (seconds < 10)
-						b.Append('0');
-				}
-				b.AppendValue(seconds).Append('.');
-				if (milliseconds < 100)
-					b.Append('0');
-				if (milliseconds < 10)
-					b.Append('0');
-				b.AppendValue(milliseconds);
-
-				StringBuffer<128> wt = GetWorldTime(timelineSelected);
-
-				SetActiveFont(m_popupFont);
-				DrawTextLogger logger(m_hwnd, hdc, m_popupFont.height, m_tooltipBackgroundBrush);
-
-				logger.Info(L"%s (%s)", b.data, wt.data);
-				logger.DrawAtPos(left + 4, timelineTop - 20);
-			}
-		}
-
-		{
-			int boxSide = 8;
-			int boxStride = boxSide + 2;
-			int top = 5;
-			int bottom = top + boxSide;
-			int left = progressRect.right - 7 - boxSide;
-			int right = progressRect.right - 7;
-			bool* values = &m_config.showProgress;
-			for (int i = VisualizerFlag_Count - 1; i >= 0; --i)
-			{
-				SelectObject(hdc, m_buttonSelected == u32(i) ? m_textPen : m_checkboxPen);
-				SelectObject(hdc, GetStockObject(NULL_BRUSH));
-				Rectangle(hdc, left, top, right, bottom);
-
-				if (values[i])
-				{
-					MoveToEx(hdc, left + 2, top + 2, NULL);
-					LineTo(hdc, right - 2, bottom - 2);
-					MoveToEx(hdc, right - 3, top + 2, NULL);
-					LineTo(hdc, left + 1, bottom - 2);
-				}
-
-				left -= boxStride;
-				right -= boxStride;
-			}
-
-			top -= 2;
-
-			auto drawText = [&](const tchar* text, COLORREF color)
-				{
-					SetTextColor(hdc, color);
-					RECT r{left, top, left+200, top+200};
-					auto strLen = TStrlen(text);
-					DrawTextW(hdc, text, strLen, &r, DT_SINGLELINE|DT_NOCLIP|DT_NOPREFIX|DT_CALCRECT);
-					left -= (r.right - r.left) + 5;
-					r.left = left;
-					DrawTextW(hdc, text, strLen, &r, DT_SINGLELINE|DT_NOCLIP|DT_NOPREFIX);
-				};
-
-			if (m_config.showDriveStats)
-			{
-				SetActiveFont(m_defaultFont);
-				drawText(L"DRV", m_driveColor);
-			}
-
-			if (m_config.showNetworkStats)
-			{
-				SetActiveFont(m_defaultFont);
-				drawText(L"SND", m_sendColor);
-				drawText(L"RCV", m_recvColor);
-			}
-
-			if (m_config.showCpuMemStats)
-			{
-				SetActiveFont(m_defaultFont);
-				drawText(L"CPU", m_cpuColor);
-				drawText(L"MEM", m_memColor);
-			}
-
-			SetTextColor(hdc, m_textColor);
-		}
-
-		if (m_processSelected)
-		{
-			const TraceView::Process& process = m_traceView.GetProcess(m_processSelectedLocation);
-			u64 duration = 0;
-			
-			Vector<TString> logLines;
-			u32 maxCharCount = 50u;
-
-			bool hasExited = process.stop != ~u64(0);
-			if (hasExited)
-			{
-				duration = process.stop - process.start;
-
-				if (!process.logLines.empty())
-				{
-					u32 lineMaxCount = 0;
-					for (auto& line : process.logLines)
-					{
-						u32 offset = 0;
-						u32 left = u32(line.text.size());
-						while (left)
-						{
-							u32 toCopy = Min(left, maxCharCount);
-							lineMaxCount = Max(lineMaxCount, toCopy);
-							logLines.push_back(line.text.substr(offset, toCopy));
-							offset += toCopy;
-							left -= toCopy;
-						}
-					}
-				}
-			}
-			else
-			{
-				duration = playTime - process.start;
-			}
-
-			SetActiveFont(m_popupFont);
-			DrawTextLogger logger(m_hwnd, hdc, m_popupFont.height, m_tooltipBackgroundBrush);
-
-			logger.AddTextOffset(-10); // Remove spaces in the front
-			logger.AddWidth(3);
-
-			logger.AddSpace(2);
-			logger.Info(L"  %ls", process.description.c_str());
-			if (process.type != TraceView::ProcessType_Task)
-			{
-				logger.Info(L"  Host:        %ls", m_processSelectedLocation.sessionIndex == 0 ? TC("local") : m_traceView.GetSession(m_processSelectedLocation).name.c_str());
-				logger.Info(L"  ProcessId:  %6u", process.id);
-			}
-			logger.Info(L"  Start:     %7ls (%ls)", TimeToText(process.start, true).str, GetWorldTime(process.start).data);
-			logger.Info(L"  Duration:  %7ls", TimeToText(duration, true).str);
-			if (!process.returnedReason.empty())
-				logger.Info(L"  Returned: %7s", process.returnedReason.data());
-			if (hasExited && process.exitCode != 0)
-			{
-				if (process.exitCode == ProcessCancelExitCode)
-					logger.Info(L"  ExitCode: Cancelled");
-				else if (process.exitCode >= 0xC0000000u)
-					logger.Info(L"  ExitCode: 0x%x", process.exitCode);
-				else
-					logger.Info(L"  ExitCode: %7u", process.exitCode);
-			}
-
-			const auto& breadcrumbs = process.breadcrumbs;
-			if (!breadcrumbs.empty())
-			{
-				constexpr TString::size_type maxLineLen = 50;
-				logger.Info(L"");
-				logger.Info(L"  ------------ Breadcrumbs ------------");
-				for (TString::size_type lineStart = 0, lineEnd = 0; lineEnd < breadcrumbs.size(); lineStart = lineEnd + 1)
-				{
-					// Log each individual line
-					lineEnd = breadcrumbs.find(L'\n', lineStart);
-					TString line = (lineEnd == TString::npos ? breadcrumbs.substr(lineStart) : breadcrumbs.substr(lineStart, lineEnd - lineStart));
-
-					// Break each line down into smaller section if they are longer than the maximum allowed length
-					if (line.size() > maxLineLen)
-					{
-						for (TString::size_type sectionStart = 0, sectionEnd = 0; sectionStart < line.size(); sectionStart = sectionEnd)
-						{
-							const TString::size_type maxSectionLen = sectionStart == 0 ? maxLineLen : maxLineLen - 2;
-							sectionEnd = std::min<TString::size_type>(sectionEnd + maxSectionLen, line.size());
-							TString section = (sectionStart == 0 ? L"  " : L"    ") + line.substr(sectionStart, sectionEnd - sectionStart);
-							logger.Info(section.c_str());
-						}
-					}
-					else
-					{
-						line = L"  " + line;
-						logger.Info(line.c_str());
-					}
-				}
-			}
-
-			if (process.stop != ~u64(0) && !process.stats.empty())
-			{
-				BinaryReader reader(process.stats.data(), 0, process.stats.size());
-				ProcessStats processStats;
-				SessionStats sessionStats;
-				StorageStats storageStats;
-				KernelStats kernelStats;
-				CacheFetchStats cacheStats;
-
-				if (process.type == TraceView::ProcessType_CacheFetchTest || process.type == TraceView::ProcessType_CacheFetchDownload)
-				{
-					if (!process.returnedReason.empty())
-						logger.Info(L"  Cache:       Miss");
-					else
-						logger.Info(L"  Cache:        Hit");
-					cacheStats.Read(reader, m_traceView.version);
-					if (reader.GetLeft())
-					{
-						storageStats.Read(reader, m_traceView.version);
-						kernelStats.Read(reader, m_traceView.version);
-					}
-				}
-				else
-				{
-					processStats.Read(reader, m_traceView.version);
-
-					if (reader.GetLeft())
-					{
-						if (process.isRemote || (m_traceView.version >= 36 && !process.isReuse))
-							sessionStats.Read(reader, m_traceView.version);
-						storageStats.Read(reader, m_traceView.version);
-						kernelStats.Read(reader, m_traceView.version);
-					}
-				}
-
-				if (processStats.hostTotalTime)
-				{
-					logger.Info(L"");
-					logger.Info(L"  ----------- Detours stats -----------");
-					processStats.Print(logger, m_traceView.frequency);
-				}
-				else if (processStats.peakMemory)
-				{
-					logger.Info(L"");
-					logger.Info(L"  ----------- Process stats -----------");
-					processStats.Print(logger, m_traceView.frequency);
-				}
-
-				if (!sessionStats.IsEmpty())
-				{
-					logger.Info(L"");
-					logger.Info(L"  ----------- Session stats -----------");
-					sessionStats.Print(logger, m_traceView.frequency);
-				}
-
-				if (!cacheStats.IsEmpty())
-				{
-					logger.Info(L"");
-					logger.Info(L"  ------------ Cache stats ------------");
-					cacheStats.Print(logger, m_traceView.frequency);
-				}
-
-				if (!storageStats.IsEmpty())
-				{
-					logger.Info(L"");
-					logger.Info(L"  ----------- Storage stats -----------");
-					storageStats.Print(logger, m_traceView.frequency);
-				}
-
-				if (!kernelStats.IsEmpty())
-				{
-					logger.Info(L"");
-					logger.Info(L"  ----------- Kernel stats ------------");
-					kernelStats.Print(logger, false, m_traceView.frequency);
-				}
-
-				PrintCacheWriteStats(logger, process.id);
-
-				if (!logLines.empty())
-				{
-					logger.Info(L"");
-					logger.Info(L"  ---------------- Log ----------------");
-					logger.AddTextOffset(14);
-					for (auto& line : logLines)
-						logger.Log(LogEntryType_Info, line);
-				}
-			}
-			logger.AddSpace(3);
-			logger.DrawAtCursor();
-		}
-		else if (m_workSelected && selectedWork.description)
-		{
-			u64 duration;
-			if (selectedWork.stop != ~u64(0))
-				duration = selectedWork.stop - selectedWork.start;
-			else
-				duration = playTime - selectedWork.start;
-
-			SetActiveFont(m_popupFont);
-			DrawTextLogger logger(m_hwnd, hdc, m_popupFont.height, m_tooltipBackgroundBrush);
-
-			logger.AddSpace();
-			logger.Info(L"  %ls", selectedWork.description);
-			logger.Info(L"  Start:     %ls (%ls)", TimeToText(selectedWork.start, true).str, GetWorldTime(selectedWork.start).data);
-			logger.Info(L"  Duration:  %ls", TimeToText(duration, true).str);
-			if (!selectedWork.entries.empty())
-				logger.AddSpace();
-			for (u32 i=0,e=u32(selectedWork.entries.size()); i!=e; ++i)
-			{
-				auto& entry = selectedWork.entries[i];
-				u64 time = 0;
-				if (i == 0)
-				{
-					u64 start = entry.startTime ? entry.startTime : entry.time;
-					if (TimeToMs(start - selectedWork.start) > 1)
-						logger.Info(L"   Start (%ls)", TimeToText(start - selectedWork.start).str);
-				}
-
-				if (entry.startTime)
-					time = entry.time - entry.startTime;
-				
-				if (!time)
-					for (u32 j=i+1;j<e;++j)
-					{
-						auto& next = selectedWork.entries[j];
-						if (next.startTime)
-							continue;
-						time = next.time - entry.time;
-						break;
-					}
-				
-				if (!time && selectedWork.stop != ~0u)
-					time = selectedWork.stop - entry.time;
-				if (entry.count == 1)
-					logger.Info(L"%s  %ls (%ls)", (entry.startTime ? L" " : L""), entry.text, TimeToText(time).str);
-				else
-					logger.Info(L"%s  %ls (%ls %u)", (entry.startTime ? L" " : L""), entry.text, TimeToText(time).str, entry.count);
-			}
-			logger.AddSpace();
-			logger.DrawAtCursor();
-		}
-		else if (m_sessionSelectedIndex != ~0u)
-		{
-			SetActiveFont(m_popupFont);
-			DrawTextLogger logger(m_hwnd, hdc, m_popupFont.height, m_tooltipBackgroundBrush);
-			logger.AddWidth(3);
-			logger.AddSpace(2);
-
-			auto& session = m_traceView.sessions[m_sessionSelectedIndex];
-			Vector<TString>& summary = session.summary;
-			if (summary.empty())
-			{
-				if (m_traceView.finished)
-					logger.Info(L" Session summary was never received");
-				else
-					logger.Info(L" Session summary not available until session is done");
-			}
-
-			logger.AddTextOffset(-10); // Remove spaces in the front
-			for (auto& line : summary)
-				logger.Log(LogEntryType_Info, line);
-			logger.AddTextOffset(0);
-
-			if (m_sessionSelectedIndex == 0)
-			{
-				for (auto& cacheSummary : m_traceView.cacheSummaries)
-					for (auto& line : cacheSummary.lines)
-						logger.Log(LogEntryType_Info, line);
-			}
-
-			for (auto& pair : session.drives)
-			{
-				auto& d = pair.second;
-				logger.Info(L"  %c: Read: %s (%u) Write: %s (%u)", pair.first, BytesToText(d.totalReadBytes).str, d.totalReadCount, BytesToText(d.totalWriteBytes).str, d.totalWriteCount);
-			}
-
-			logger.AddSpace(3);
-			logger.DrawAtCursor();
-		}
-		else if (m_statsSelected)
-		{
-			SetActiveFont(m_popupFont);
-			DrawTextLogger logger(m_hwnd, hdc, m_popupFont.height, m_tooltipBackgroundBrush);
-			logger.AddSpace(3);
-			logger.SetColor(m_cpuColor).Info(L"  Cpu: %.1f%%", m_stats.cpuLoad * 100.0f);
-			logger.SetColor(m_memColor).Info(L"  Mem: %ls/%ls", BytesToText(m_stats.memTotal - m_stats.memAvail).str, BytesToText(m_stats.memTotal).str);
-			if (m_stats.recvBytes || m_stats.sendBytes)
-			{
-				logger.SetColor(m_recvColor).Info(L"  Recv: %lsit/s", BytesToText(m_stats.recvBytesPerSecond*8).str);
-				logger.SetColor(m_sendColor).Info(L"  Send: %lsit/s", BytesToText(m_stats.sendBytesPerSecond*8).str);
-			}
-			logger.SetColor(m_textColor);
-
-			if (m_stats.ping)
-				logger.Info(L"  Ping: %ls", TimeToText(m_stats.ping, false, m_traceView.frequency).str);
-			if (m_stats.connectionCount)
-				logger.Info(L"  Conn: %u", u32(m_stats.connectionCount));
-
-			for (auto& pair : m_stats.drives)
-				logger.SetColor(m_driveColor).Info(L"  %c: %u%% R:%s/s W:%s/s", pair.first, pair.second.busyPercent, BytesToText(pair.second.readPerSecond).str, BytesToText(pair.second.writePerSecond).str);
-
-			logger.AddSpace(3);
-			logger.DrawAtCursor();
-		}
-		else if (m_activeProcessGraphSelected)
-		{
-			SetActiveFont(m_popupFont);
-			DrawTextLogger logger(m_hwnd, hdc, m_popupFont.height, m_tooltipBackgroundBrush);
-			logger.AddSpace(3);
-			logger.AddSpace(3);
-			logger.SetColor(m_activeProcColor).Info(L"  Active Processes: %u", m_activeProcessCount);
-			logger.DrawAtCursor();
-		}
-		else if (m_buttonSelected != ~0u)
-		{
-			const wchar_t* tooltip[] =
-			{
-				#define UBA_VISUALIZER_FLAG(name, defaultValue, desc) desc,
-				UBA_VISUALIZER_FLAGS1
-				#undef UBA_VISUALIZER_FLAG
-			};
-
-			SetActiveFont(m_popupFont);
-			DrawTextLogger logger(m_hwnd, hdc, m_popupFont.height, m_tooltipBackgroundBrush);
-			logger.Info(L"%ls %ls", L"Show", tooltip[m_buttonSelected]);
-			logger.DrawAtCursor();
-		}
-		else if (m_fetchedFilesSelected != ~0u)
-		{
-			auto& session = m_traceView.sessions[m_fetchedFilesSelected];
-			auto& fetchedFiles = session.fetchedFiles;
-			if (!fetchedFiles.empty() && !fetchedFiles[0].hint.empty())
-			{
-				/*
-				int colWidth = 500;
-				int width = colWidth * 2;
-				int height = Min(int(clientRect.bottom), int(fetchedFiles.size() * m_popupFont.height));
-
-				SetActiveFont(m_defaultFont);
-				DrawTextLogger logger(m_hwnd, hdc, r, m_font.height, m_tooltipBackgroundBrush);
-				for (auto& f : fetchedFiles)
-				{
-					if (f.hint == TC("KnownInput"))
-						continue;
-					if (logger.rect.top >= r.bottom - m_font.height)
-					{
-						if (logger.rect.left + colWidth >= r.right)
-						{
-							logger.Info(L"...");
-							break;
-						}
-						logger.rect.top = r.top;
-						logger.rect.left += colWidth;
-					}
-					logger.Info(L"%s", f.hint.c_str());
-				}
-				logger.DrawAtCursor();
-				*/
-			}
-		}
-	}
-
-	u64 ConvertTime(const TraceView& view, u64 time);
-
-	void Visualizer::PaintActiveProcesses(int& posY, const RECT& clientRect, const Function<void(TraceView::ProcessLocation&, u32, bool)>& drawProcess)
-	{
-		SetActiveFont(m_processFont);
-		int startPosY = posY;
-
-		Map<u64, TraceView::ProcessLocation*> activeProcesses;
-		u32 remoteCount = 0;
-		for (auto& kv : m_trace.m_activeProcesses)
-		{
-			auto& active = kv.second;
-			TraceView::Process& process = m_trace.GetSession(m_traceView, active.sessionIndex).processors[active.processorIndex].processes[active.processIndex];
-			u64 start = process.start; //~0llu - process.start;
-			activeProcesses.try_emplace(start, &active);
-			if (process.isRemote)
-				++remoteCount;
-		}
-
-		u32 maxHeight = clientRect.bottom;
-
-		bool fillHeight = ActiveProcessesShouldFillHeight();
-		if (fillHeight)
-		{
-			maxHeight = clientRect.bottom - posY;
-			if (m_config.showTimeline)
-				maxHeight -= GetTimelineHeight();
-		}
-		else
-		{
-			u32 maxHeight2 = m_config.maxActiveVisible * (m_activeFont.height + 2);
-			maxHeight = Min(maxHeight, maxHeight2);
-		}
-
-		u32 maxSize = Min(m_config.maxActiveProcessHeight, 32u);
-		u32 maxSizeMinusOne = maxSize - 1;
-
-		u32 counts[128] = { 0 };
-
-		u32 highestHistoryCount = 0;
-		for (u32 i=0; i!= sizeof_array(m_activeProcessCountHistory) - 1; ++i)
-			highestHistoryCount = Max(highestHistoryCount, m_activeProcessCountHistory[i]);
-
-		u32 activeProcessCount = highestHistoryCount;
-		counts[0] = activeProcessCount;
-		u32 totalHeight = counts[0]*2;
-		while (totalHeight < maxHeight && counts[maxSizeMinusOne] != activeProcessCount)
-		{
-			bool changed = false;
-			for (u32 i=0; i!=maxSizeMinusOne; ++i)
-			{
-				if (counts[i] && counts[i] > counts[i+1]*2 + 1)
-				{
-					--counts[i];
-					++counts[i+1];
-					++totalHeight;
-					changed = true;
-				}
-			}
-			if (!changed)
-			{
-				for (u32 j=0; j!=maxSizeMinusOne; ++j)
-				{
-					if (!counts[j])
-						continue;
-					++counts[j + 1];
-					--counts[j];
-					++totalHeight;
-					break;
-				}
-			}
-		}
-
-		auto it = activeProcesses.begin();
-		auto itEnd = activeProcesses.end();
-		int endY = int(startPosY + maxHeight);
-		for (u32 i=0;i!=maxSize;++i)
-		{
-			u32 v = maxSizeMinusOne - i;
-			u32 boxHeight = v + 1;
-
-			for (u32 j=0;j!=counts[v];++j)
-			{
-				if (it == itEnd || posY >= endY)
-					break;
-				auto& active = *it->second;
-				++it;
-				drawProcess(active, boxHeight, j == 0);
-				posY += boxHeight + 1;
-			}
-		}
-
-		//if (fillHeight)
-		//	posY = clientRect.bottom-1;//startPosY; // To prevent vertical scrollbar
-		if (fillHeight || counts[maxSizeMinusOne] != activeProcessCount)
-			posY = startPosY + maxHeight;
-		else
-			posY += 3;
-
-		SetActiveFont(m_defaultFont);
-	}
-
-	void Visualizer::PaintProcessRect(TraceView::Process& process, HDC hdc, RECT rect, const RECT& progressRect, bool selected, bool writingBitmap, HBRUSH& lastSelectedBrush)
-	{
-		auto clampRect = [&](RECT& r) { r.left = Min(Max(r.left, progressRect.left), progressRect.right); r.right = Max(Min(r.right, progressRect.right), progressRect.left); };
-		bool done = process.stop != ~u64(0);
-
-		HBRUSH brush = m_processBrushes[selected].success;
-		if (!process.returnedReason.empty())
-			brush = m_processBrushes[selected].returned;
-		else if (!done)
-			brush = process.isRacing ? m_processBrushes[selected].inProgressRace : m_processBrushes[selected].inProgress;
-		else if (process.type == TraceView::ProcessType_CacheFetchTest)
-			brush = m_processBrushes[selected].cacheFetchTest;
-		else if (process.type == TraceView::ProcessType_CacheFetchDownload)
-			brush = m_processBrushes[selected].cacheFetchDownload;
-		else if (process.exitCode == ProcessCancelExitCode)
-			brush = m_processBrushes[selected].returned;
-		else if (process.exitCode != 0)
-			brush = m_processBrushes[selected].error;
-
-		u64 writeFilesTime = process.writeFilesTime;
-
-		auto fillRect = [&](const RECT& r, HBRUSH b)
-			{
-				if (lastSelectedBrush != b)
-				{
-					SelectObject(hdc, b);
-					lastSelectedBrush = b;
-				}
-				PatBlt(hdc, r.left, r.top, r.right - r.left, r.bottom - r.top, PATCOPY);
-			};
-
-		if (!done || process.exitCode != 0 || !m_config.ShowReadWriteColors)
-		{
-			if (writingBitmap)
-				rect.right = 256;
-			clampRect(rect);
-			fillRect(rect, brush);
-			return;
-		}
-
-		double duration = double(process.stop - process.start);
-
-		RECT main = rect;
-		int width = rect.right - rect.left;
-
-		double recvPart = (double(ConvertTime(m_traceView, process.createFilesTime)) / duration);
-		if (int headSize = int(recvPart * width))
-		{
-			main.left += headSize;
-			RECT r2 = rect;
-			r2.right = r2.left + headSize;
-			clampRect(r2);
-			if (r2.left != r2.right)
-				fillRect(r2, m_processBrushes[selected].recv);
-		}
-
-		double sendPart = (double(ConvertTime(m_traceView, writeFilesTime)) / duration);
-		if (int tailSize = int(sendPart * width))
-		{
-			main.right -= tailSize;
-			RECT r2 = rect;
-			r2.left = r2.right - tailSize;
-			clampRect(r2);
-			if (r2.left != r2.right)
-				fillRect(r2, m_processBrushes[selected].send);
-		}
-
-		clampRect(main);
-		if (main.left != main.right)
-			fillRect(main, brush);
-
-		//clampRect(rect);
-		/*
-		if (!process.updates.empty())
-		{
-			shouldDrawText = false;
-			int prevUpdateX = rect.left;
-			for (auto& update : process.updates)
-			{
-				int updateX = int(posX + TimeToS(update.time) * scaleX) - 1;
-
-				RECT textRect{ prevUpdateX + 5, rect.top, updateX, rect.bottom };
-				DrawTextW(hdc, update.reason.c_str(), u32(update.reason.size()), &textRect, DT_SINGLELINE);
-
-				SelectObject(hdc, m_processUpdatePen);
-				MoveToEx(hdc, updateX, rect.top, NULL);
-				LineTo(hdc, updateX, rect.bottom - 1);
-				prevUpdateX = updateX;
-			}
-		}
-		*/
-	}
-
-	void Visualizer::PaintTimeline(HDC hdc, const RECT& clientRect)
-	{
-		SetActiveFont(m_timelineFont);
-		int top = GetTimelineTop(clientRect);
-		float timeScale = (m_config.horizontalScaleValue*m_zoomValue)*50.0f;
-		float startOffset = ((m_scrollPosX/timeScale) - float(int(m_scrollPosX/timeScale))) * timeScale;
-		int index = -int(startOffset/timeScale);
-			
-		int number = -int(float(m_scrollPosX)/timeScale);
-
-		int textStepSize = int((5.0f / timeScale) + 1) * 5;
-		if (textStepSize > 150)
-			textStepSize = 600;
-		else if (textStepSize > 120)
-			textStepSize = 300;
-		else if (textStepSize > 90)
-			textStepSize = 240;
-		else if (textStepSize > 45)
-			textStepSize = 120;
-		else if (textStepSize > 30)
-			textStepSize = 60;
-		else if (textStepSize > 10)
-			textStepSize = 30;
-
-		int lineStepSize = textStepSize / 5;
-
-		RECT progressRect = clientRect;
-		progressRect.left += m_progressRectLeft;
-
-		SelectObject(hdc, m_textPen);
-		Vector<DWORD> lines;
-		Vector<POINT> points;
-
-		while (true)
-		{
-			int pos = progressRect.left + int(startOffset + float(index)*timeScale);
-			if (pos >= clientRect.right)
-				break;
-
-			int lineBottom = top + 5;
-			if (!(number % textStepSize))
-			{
-				bool shouldDraw = true;
-				int seconds = number;
-				StringBuffer<> buffer;
-				if (seconds >= 60)
-				{
-					int min = seconds / 60;
-					seconds -= min * 60;
-					if (!seconds)
-					{
-						buffer.Appendf(L"%um", min);
-						lineBottom += 4;
-					}
-				}
-				if (!number || seconds)
-					buffer.Appendf(L"%u", seconds);
-
-				if (shouldDraw)
-				{
-					SIZE s;
-					GetTextExtentPoint32W(hdc, buffer.data, buffer.count, &s);
-					RECT textRect;
-					textRect.top = top + 8;
-					textRect.bottom = textRect.top + m_activeFont.height;
-					textRect.right = pos + s.cx/2;
-					textRect.left = pos - s.cx/2;
-					ExtTextOutW(hdc, textRect.left, textRect.top, 0, NULL, buffer.data, buffer.count, NULL);
-				}
-			}
-			if (!(number % lineStepSize))
-			{
-				points.push_back({pos, top});
-				points.push_back({pos, lineBottom});
-				lines.push_back(2);
-			}
-
-			++number;
-			++index;
-		}
-
-		points.push_back({m_contentWidth, top - 25});
-		points.push_back({m_contentWidth, top});
-		lines.push_back(2);
-
-		PolyPolyline(hdc, points.data(), lines.data(), int(lines.size()));
-
-		/*
-		POINT cursorPos;
-		GetCursorPos(&cursorPos);
-		ScreenToClient(m_hwnd, &cursorPos);
-		RECT lineRect;
-		lineRect.left = cursorPos.x;
-		lineRect.right= cursorPos.x + 1;
-		lineRect.top = top;
-		lineRect.bottom = top + 15;
-		FillRect(hdc, &lineRect, m_lineBrush);
-		*/
-	}
-
-	void Visualizer::PaintDetailedStats(int& posY, const RECT& progressRect, TraceView::Session& session, bool isRemote, u64 playTime, const DrawTextFunc& drawTextFunc)
-	{
-		int stepY = m_activeFont.height;
-		int startPosY = posY;
-		int posX = progressRect.left + 5;
-		int maxY = posY + stepY;
-		u32 maxTextWidth = 0;
-
-		RECT textRect;
-		auto drawText = [&](const wchar_t* format, ...)
-			{
-				textRect.top = posY;
-				textRect.bottom = posY + 20;
-				textRect.left = posX;
-				textRect.right = posX + 1000;
-				posY += stepY;
-				StringBuffer<> str;
-				va_list arg;
-				va_start(arg, format);
-				str.Append(format, arg);
-				va_end(arg);
-				u32 lastTextWidth = 0;
-				drawTextFunc(str, textRect, &lastTextWidth);
-				maxTextWidth = Max(maxTextWidth, lastTextWidth);
-				maxY = Max(maxY, posY);
-			};
-
-		if (isRemote)
-		{
-			drawText(L"Finished Processes: %u", session.processExitedCount);
-			drawText(L"Active Processes: %u", session.processActiveCount);
-			drawText(L"ClientId: %u  TcpCount: %u", session.clientUid.data1, session.connectionCount.empty() ? 1 : session.connectionCount.back());
-
-			if (session.disconnectTime == ~u64(0))
-			{
-				if (session.proxyCreated)
-					drawText(L"Proxy(HOSTED): %ls", session.proxyName.c_str());
-				else if (!session.proxyName.empty())
-					drawText(L"Proxy: %ls", session.proxyName.c_str());
-				else
-					drawText(L"Proxy: None");
-			}
-
-			bool hasFileDetails = !session.fetchedFiles.empty() || !session.storedFiles.empty();
-
-			if (!hasFileDetails)
-			{
-				posY = startPosY;
-				posX += maxTextWidth + 15;
-			}
-
-			if (!session.updates.empty())
-			{
-				auto updateTime = session.updates.back();
-				auto updateSend = session.networkSend.back();
-				auto updateRecv = session.networkRecv.back();
-				u64 sendPerS = 0;
-				u64 recvPerS = 0;
-				float duration = TimeToS(updateTime - session.prevUpdateTime);
-				if (duration != 0)
-				{
-					sendPerS = u64(float(updateSend - session.prevSend) / duration);
-					recvPerS = u64(float(updateRecv - session.prevRecv) / duration);
-				}
-				drawText(L"Recv: %ls (%sit/s)", BytesToText(updateRecv), BytesToText(recvPerS*8));
-				drawText(L"Send: %ls (%sit/s)", BytesToText(updateSend), BytesToText(sendPerS*8));
-			}
-
-			int fileWidth = 700;
-
-			auto drawFiles = [&](const wchar_t* fileType, Vector<TraceView::FileTransfer>& files, u64 count, u64 bytes, u64 activeCount, u32& maxVisibleFiles)
-				{
-					textRect.left = posX;
-					textRect.right = posX + fileWidth;
-					drawText(L"%ls Files: %u (%u) %s", fileType, u32(count), u32(activeCount), BytesToText(bytes));
-					u32 fileCount = 0;
-					for (auto rit = files.rbegin(), rend = files.rend(); rit != rend; ++rit)
-					{
-						TraceView::FileTransfer& file = *rit;
-						if (file.stop != ~u64(0))
-							continue;
-						u64 time = 0;
-						if (file.start < playTime)
-							time = playTime - file.start;
-
-						if ((((u8*)&file.key)[19] & IsProxyMask) != 0)
-							drawText(L"%s (proxy) %s", file.hint.c_str(), TimeToText(time, true).str);
-						else if (file.size == 0)
-							drawText(L"%s (calc) %s", file.hint.c_str(), TimeToText(time, true).str);
-						else
-							drawText(L"%s (%s) %s", file.hint.c_str(), BytesToText(file.size).str, TimeToText(time, true).str);
-
-						if (fileCount++ > 5)
-							break;
-					}
-					posY += stepY * (maxVisibleFiles - fileCount);
-					maxVisibleFiles = Max(maxVisibleFiles, fileCount);
-				};
-
-			Vector<TraceView::FileTransfer> fetchedFiles;
-			for (auto& kv : session.fetchedFilesActive)
-				fetchedFiles.push_back(session.fetchedFiles[kv.second]);
-			std::sort(fetchedFiles.begin(), fetchedFiles.end(), [](const TraceView::FileTransfer& a, const TraceView::FileTransfer& b) { return a.start > b.start; });
-
-			if (hasFileDetails)
-			{
-				posY = startPosY;
-				posX += maxTextWidth + 15;
-			}
-			drawFiles(L"Fetched", fetchedFiles, session.fetchedFilesCount, session.fetchedFilesBytes, session.fetchedFilesActive.size(), session.maxVisibleFiles);
-			if (hasFileDetails)
-			{
-				posY = startPosY;
-				posX += fileWidth;
-			}
-			drawFiles(L"Stored", session.storedFiles, session.storedFilesCount, session.storedFilesBytes, session.storedFilesActive.size(), session.maxVisibleFiles);
-		}
-		else
-		{
-			if (m_traceView.totalProcessActiveCount || m_traceView.totalProcessExitedCount)
-			{
-				drawText(L"Finished Processes: %u (local: %u)", m_traceView.totalProcessExitedCount, session.processExitedCount);
-				drawText(L"Active Processes: %u (local: %u)", m_traceView.totalProcessActiveCount, session.processActiveCount);
-				drawText(L"Active Helpers: %u", Max(1u, m_traceView.activeSessionCount) - 1);
-			}
-
-			if (session.highestSendPerS > 0.0f || session.highestRecvPerS > 0.0f)
-			{
-				auto updateTime = session.updates.back();
-				auto updateSend = session.networkSend.back();
-				auto updateRecv = session.networkRecv.back();
-				if (updateSend || updateRecv)
-				{
-					u64 sendPerS = 0;
-					u64 recvPerS = 0;
-					float duration = TimeToS(updateTime - session.prevUpdateTime);
-					if (duration != 0)
-					{
-						sendPerS = u64(float(updateSend - session.prevSend) / duration);
-						recvPerS = u64(float(updateRecv - session.prevRecv) / duration);
-					}
-
-					drawText(L"Recv: %ls (%sit/s)", BytesToText(updateRecv), BytesToText(recvPerS*8));
-					drawText(L"Send: %ls (%sit/s)", BytesToText(updateSend), BytesToText(sendPerS*8));
-				}
-			}
-
-			if (!session.updates.empty())
-			{
-				posY = startPosY;
-				posX += maxTextWidth + 10;
-
-				for (auto& pair : session.drives)
-				{
-					auto& drive = pair.second;
-
-					u64 readPerS = 0;
-					u64 writePerS = 0;
-
-					u64 updateCount = session.updates.size();
-					float duration = TimeToS(session.updates.back() - session.prevUpdateTime);
-					if (duration != 0 && updateCount > 1)
-					{
-						readPerS = u64(float(drive.readBytes.back()) / duration);
-						writePerS = u64(float(drive.writeBytes.back()) / duration);
-					}
-				
-					drawText(L"%c: Rd %ls (%s/s) Wr %ls (%s/s) ", pair.first, BytesToText(drive.totalReadBytes).str, BytesToText(readPerS).str, BytesToText(drive.totalWriteBytes).str, BytesToText(writePerS).str);
-				}
-			}
-		}
-
-		posY = maxY;
-	}
-
-	u64 Visualizer::GetPlayTime()
-	{
-		u64 currentTime = m_paused ? m_pauseStart : GetTime();
-		u64 playTime = 0;
-		if (m_traceView.startTime && currentTime > (m_traceView.startTime + m_pauseTime))
-			playTime = currentTime - m_traceView.startTime - m_pauseTime;
-		if (m_replay)
-			playTime *= m_replay;
-		return playTime;
-	}
-
-	int Visualizer::GetTimelineHeight()
-	{
-		return m_timelineFont.height + 8;
-	}
-
-	int Visualizer::GetTimelineTop(const RECT& clientRect)
-	{
-		int timelineHeight = GetTimelineHeight();
-		int posY = m_contentHeight - timelineHeight;
-		int maxY = int(clientRect.bottom - timelineHeight);
-		return m_config.LockTimelineToBottom ? maxY : Min(posY, maxY);
-	}
-
-	void Visualizer::HitTest(HitTestResult& outResult, const POINT& pos)
-	{
-		SetActiveFont(m_defaultFont);
-
-		u64 playTime = GetPlayTime();
-
-		RECT clientRect;
-		GetClientRect(m_hwnd, &clientRect);
-
-		int posY = int(m_scrollPosY);
-		int boxHeight = m_config.boxHeight;
-		int processStepY = boxHeight + 1;
-		float scaleX = 50.0f*m_zoomValue*m_config.horizontalScaleValue;
-
-		RECT progressRect = clientRect;
-		progressRect.left += m_progressRectLeft;
-		progressRect.bottom -= 30;
-
-		{
-			int boxSide = 8;
-			int boxStride = boxSide + 2;
-			int top = 5;
-			int bottom = top + boxSide;
-			int left = progressRect.right - 7 - boxSide;
-			int right = progressRect.right - 7;
-			for (int i = VisualizerFlag_Count - 1; i >= 0; --i)
-			{
-				if (pos.x >= left && pos.x <= right && pos.y >= top && pos.y <= bottom)
-				{
-					outResult.buttonSelected = i;
-					return;
-				}
-				left -= boxStride;
-				right -= boxStride;
-			}
-		}
-
-		outResult.section = 0;
-
-		if (m_config.showTimeline && !m_traceView.sessions.empty())
-		{
-			int timelineTop = GetTimelineTop(clientRect);
-			if (pos.y >= timelineTop)// && pos.y < timelineTop + 40)
-			{
-				outResult.section = 3;
-				float timeScale = (m_config.horizontalScaleValue * m_zoomValue)*50.0f;
-				float startOffset = -(m_scrollPosX / timeScale);
-				outResult.timelineSelected = startOffset + float(pos.x - m_progressRectLeft) / timeScale;
-				return;
-			}
-		}
-
-		u64 lastStop = 0;
-
-		if (m_config.showProgress && m_traceView.progressProcessesTotal)
-			posY += m_activeFont.height + 2;
-
-		if (m_config.showStatus && !m_traceView.statusMap.empty())
-		{
-			u32 lastRow = ~0u;
-			u32 row = ~0u;
-			for (auto& kv : m_traceView.statusMap)
-			{
-				if (kv.second.text.empty())
-					continue;
-
-				row = u32(kv.first >> 32);
-				if (lastRow != ~0u && lastRow != row)
-					posY += m_activeFont.height + 2;
-				lastRow = row;
-
-				if (!kv.second.link.empty())
-				{
-					if (pos.y >= posY && pos.y < posY + m_activeFont.height && pos.x > 20 && pos.x < 80) // TODO: x is just hard coded to fit horde for now
-					{
-						outResult.hyperLink = kv.second.link;
-						return;
-					}
-				}
-			}
-			if (row != ~0u)
-				posY += m_activeFont.height + 2;
-			posY += 3;
-		}
-
-		if (pos.y < posY)
-			return;
-
-		outResult.section = 1;
-
-		if (m_config.showActiveProcessGraph)
-		{
-			if (pos.y > posY && pos.y < posY + GraphHeight)
-			{
-				float timeScale = (m_config.horizontalScaleValue * m_zoomValue)*50.0f;
-				float startOffset = -(m_scrollPosX / timeScale);
-				u64 selectedTimeMs = u64(1000.0f*(startOffset + float(pos.x - m_progressRectLeft) / timeScale));
-
-				u64 lastTime = m_traceView.activeProcessCounts.empty() ? 0 : m_traceView.activeProcessCounts[m_traceView.activeProcessCounts.size()-1].time;
-				if (selectedTimeMs < TimeToMs(lastTime))
-				{
-					u16 count = 0;
-					for (auto& entry : m_traceView.activeProcessCounts)
-					{
-						count = entry.count;
-						if (TimeToMs(entry.time) > selectedTimeMs)
-							break;
-					}
-					
-					outResult.activeProcessGraphSelected = true;
-					outResult.activeProcessCount = count;
-				}
-			}
-			posY += GraphHeight;
-		}
-
-		TraceView::ProcessLocation& outLocation = outResult.processLocation;
-
-		if (m_config.showActiveProcesses && !m_trace.m_activeProcesses.empty())
-		{
-			PaintActiveProcesses(posY, clientRect, [&](TraceView::ProcessLocation& processLocation, u32 boxHeight, bool firstWithHeight)
-				{
-					if (pos.y < posY || pos.y > posY + int(boxHeight))
-						return;
-					auto& session = m_trace.GetSession(m_traceView, processLocation.sessionIndex);
-					TraceView::Process& process = session.processors[processLocation.processorIndex].processes[processLocation.processIndex];
-
-					int posX = int(m_scrollPosX) + progressRect.left;
-					u64 stop = process.stop;
-					bool done = stop != ~u64(0);
-					if (!done)
-						stop = playTime;
-					int left = posX + int(float(TimeToS(process.start)) * scaleX);
-					int right = posX + int(float(TimeToS(stop)) * scaleX) - 1;
-
-					if (pos.x >= left && pos.x <= right)
-					{
-						outLocation.sessionIndex = processLocation.sessionIndex;
-						outLocation.processorIndex = processLocation.processorIndex;
-						outLocation.processIndex = processLocation.processIndex;
-						outResult.processSelected = true;
-						return;
-					}
-				});
-			if (outResult.processSelected)
-				return;
-		}
-
-		if (pos.y < posY)
-			return;
-
-		outResult.section = 2;
-
-		SessionRec sortedSessions[1024];
-		Populate(sortedSessions, m_traceView, m_config.SortActiveRemoteSessions);
-
-		for (u64 sessionIt = 0, sessionEnd = m_traceView.sessions.size(); sessionIt != sessionEnd; ++sessionIt)
-		{
-			bool isFirst = sessionIt == 0;
-			auto& session = *sortedSessions[sessionIt].session;
-			bool hasUpdates = !session.updates.empty();
-			if (!isFirst)
-			{
-				if (!hasUpdates && session.processors.empty())
-					continue;
-				if (!m_config.showFinishedProcesses && session.disconnectTime != ~u64(0))
-					continue;
-			}
-
-			u32 sessionIndex = sortedSessions[sessionIt].index;
-			if (!isFirst)
-				posY += 3;
-
-			if (m_config.showTitleBars)
-			{
-				if (pos.y >= posY && pos.y < posY + m_sessionStepY)
-				{
-					if (pos.x < int(session.fullNameWidth) + 5)
-					{
-						outResult.sessionSelectedIndex = sessionIndex;
-						return;
-					}
-				}
-
-				posY += m_sessionStepY;
-			}
-
-			bool showGraph = m_config.showNetworkStats || m_config.showCpuMemStats || m_config.showDriveStats;
-			if (showGraph && !session.updates.empty())
-			{
-				int posX = int(m_scrollPosX) + progressRect.left;
-
-				if (pos.y >= posY && pos.y < posY + GraphHeight && pos.x >= posX)
-				{
-					bool loop = true;
-					u32 reconnectIndex = 0;
-					while (loop)
-					{
-						u64 i = 0;
-						u64 e = 0;
-
-						if (reconnectIndex)
-							i = session.reconnectIndices[reconnectIndex - 1];
-
-						if (reconnectIndex < session.reconnectIndices.size())
-							e = session.reconnectIndices[reconnectIndex];
-						else
-						{
-							e = session.updates.size();
-							loop = false;
-						}
-
-						int prevX = -1;
-						int prevCenterX = -1;
-						for (; i!=e; ++i)
-						{
-							u64 updateTime = session.updates[i];
-							auto updateSend = session.networkSend[i];
-							auto updateRecv = session.networkRecv[i];
-
-							int x = posX + int(float(TimeToS(updateTime)) * scaleX);
-							int centerX = prevX != -1 ? (prevX + (x - prevX)/2) : x;
-
-							if (prevCenterX != -1 && pos.x >= prevCenterX && pos.x < centerX)
-							{
-								u64 prevTime = 0;
-								u64 prevSend = 0;
-								u64 prevRecv = 0;
-								if (i > 0)
-								{
-									prevTime = session.updates[i - 1];
-									prevSend = Min(session.networkSend[i - 1], updateSend);
-									prevRecv = Min(session.networkRecv[i - 1], updateRecv);
-								}
-
-								double duration = TimeToS(updateTime - prevTime);
-								outResult.stats.recvBytes = updateRecv;
-								outResult.stats.sendBytes = updateSend;
-								outResult.stats.recvBytesPerSecond = u64(float(updateRecv - prevRecv) / duration);
-								outResult.stats.sendBytesPerSecond = u64(float(updateSend - prevSend) / duration);
-								outResult.stats.ping = session.ping[i];
-								outResult.stats.memAvail = session.memAvail[i];
-								outResult.stats.cpuLoad = session.cpuLoad[i];
-								outResult.stats.memTotal = session.memTotal;
-								outResult.stats.connectionCount = session.connectionCount[i];
-								outResult.statsSelected = true;
-
-								for (auto& pair : session.drives)
-								{
-									auto& updateBusy = pair.second.busyPercent[i];
-									u64 updateRead = pair.second.readBytes[i];
-									u64 updateWrite = pair.second.writeBytes[i];
-									auto& statsDrive = outResult.stats.drives[pair.first];
-									statsDrive.busyPercent = updateBusy;
-									statsDrive.readPerSecond = u64(float(updateRead) / duration);
-									statsDrive.writePerSecond = u64(float(updateWrite) / duration);
-								}
-								return;
-							}
-							
-							prevCenterX = centerX;
-							prevX = x;
-						}
-						++reconnectIndex;
-					}
-					posY += GraphHeight;
-				}
-
-				posY += GraphHeight;
-			}
-
-			if (m_config.showDetailedData)
-			{
-				auto drawText = [&](const StringBufferBase& text, RECT& rect, u32* outWidth)
-					{
-						if (pos.x >= rect.left && pos.x < rect.right && pos.y >= rect.top && pos.y < rect.bottom && text.StartsWith(TC("Fetched Files")))
-							outResult.fetchedFilesSelected = sessionIndex;
-					};
-				PaintDetailedStats(posY, progressRect, session, sessionIt != 0, playTime, drawText);
-			}
-
-			if (m_config.showProcessBars)
-			{
-				u32 processorIndex = 0;
-				for (auto& processor : session.processors)
-				{
-					bool drawProcessorIndex = m_config.showFinishedProcesses;
-
-					if (pos.y < progressRect.bottom && posY + processStepY >= progressRect.top && posY <= progressRect.bottom && pos.y >= posY-1 && pos.y < posY-1 + processStepY)
-					{
-						u32 processIndex = 0;
-						int posX = int(m_scrollPosX) + progressRect.left;
-						for (auto& process : processor.processes)
-						{
-							int left = posX + int(float(TimeToS(process.start)) * scaleX);
-
-							auto pig = MakeGuard([&]() { ++processIndex; });
-
-							if (left >= progressRect.right)
-								continue;
-
-							if (left < progressRect.left)
-								left = progressRect.left;
-
-							u64 stopTime = process.stop;
-							bool done = stopTime != ~u64(0);
-							if (!done)
-								stopTime = playTime;
-							else if (!m_config.showFinishedProcesses)
-								continue;
-
-							drawProcessorIndex = true;
-
-							RECT rect;
-							rect.left = left;
-							rect.right = posX + int(float(TimeToS(stopTime)) * scaleX);
-
-							if (rect.right <= progressRect.left)
-								continue;
-
-							if (!m_filterString.empty() && !Contains(process.description.c_str(), m_filterString.c_str()) && !Contains(process.breadcrumbs.c_str(), m_filterString.c_str()))
-								continue;
-
-							rect.right = Max(int(rect.right), left + 1);
-							rect.top = posY;
-							rect.bottom = posY + int(float(18) * m_zoomValue);
-
-							if (pos.x >= rect.left && pos.x <= rect.right)
-							{
-								outLocation.sessionIndex = sessionIndex;
-								outLocation.processorIndex = processorIndex;
-								outLocation.processIndex = processIndex;
-								outResult.processSelected = true;
-								return;
-							}
-						}
-					}
-
-					if (!processor.processes.empty())
-						lastStop = Max(lastStop, processor.processes.rbegin()->stop);
-
-					if (drawProcessorIndex)
-						posY += processStepY;
-
-					++processorIndex;
-				}
-			}
-			else
-			{
-				for (auto& processor : session.processors)
-					if (!processor.processes.empty())
-						lastStop = Max(lastStop, processor.processes.rbegin()->stop);
-			}
-
-			if (m_config.showWorkers && isFirst)
-			{
-				int trackIndex = 0;
-				for (auto& workTrack : m_traceView.workTracks)
-				{
-					if (pos.y < progressRect.bottom && posY + processStepY >= progressRect.top && posY <= progressRect.bottom && pos.y >= posY-1 && pos.y < posY-1 + processStepY)
-					{
-						u32 workIndex = 0;
-						int posX = int(m_scrollPosX) + progressRect.left;
-						for (auto& work : workTrack.records)
-						{
-							auto inc = MakeGuard([&](){ ++workIndex; });
-
-							int left = posX + int(float(TimeToS(work.start)) * scaleX);
-
-							if (left >= progressRect.right)
-								continue;
-
-							if (!m_filterString.empty())
-							{
-								bool keep = Contains(work.description, m_filterString.c_str());
-								if (!keep)
-									for (auto& en : work.entries)
-										keep |= Contains(en.text, m_filterString.c_str());
-								if (!keep)
-									continue;
-							}
-
-							if (left < progressRect.left)
-								left = progressRect.left;
-
-							u64 stopTime = work.stop;
-							bool done = stopTime != ~u64(0);
-							if (!done)
-								stopTime = playTime;
-
-							RECT rect;
-							rect.left = left;
-							rect.right = posX + int(float(TimeToS(stopTime)) * scaleX);
-
-							if (rect.right <= progressRect.left)
-								continue;
-
-							rect.right = Max(int(rect.right), left + 1);
-							rect.top = posY;
-							rect.bottom = posY + int(float(18) * m_zoomValue);
-
-							if (pos.x >= rect.left && pos.x <= rect.right)
-							{
-								outResult.workTrack = trackIndex;
-								outResult.workIndex = workIndex;
-								outResult.workSelected = true;
-								return;
-							}
-						}
-					}
-					++trackIndex;
-					posY += processStepY;
-				}
-			}
-		}
-
-		m_contentWidth = m_progressRectLeft + Max(0, int(TimeToS((lastStop != 0 && lastStop != ~u64(0)) ? lastStop : playTime) * scaleX));
-		m_contentHeight = posY - int(m_scrollPosY) + processStepY + 14;
-	}
-
-	void Visualizer::WriteProcessStats(Logger& out, const TraceView::Process& process)
-	{
-		bool hasExited = process.stop != ~u64(0);
-		out.Info(L"  %ls", process.description.c_str());
-		out.Info(L"  ProcessId: %u", process.id);
-		out.Info(L"  Start:     %ls", TimeToText(process.start, true).str);
-		if (hasExited)
-			out.Info(L"  Duration:  %ls", TimeToText(process.stop - process.start, true).str);
-		if (hasExited && process.exitCode != 0)
-			out.Info(L"  ExitCode:  %u", process.exitCode);
-
-		if (process.stop != ~u64(0))
-		{
-			out.Info(L"");
-
-			BinaryReader reader(process.stats.data(), 0, process.stats.size());
-			ProcessStats processStats;
-			SessionStats sessionStats;
-			StorageStats storageStats;
-			KernelStats kernelStats;
-
-			processStats.Read(reader, m_traceView.version);
-			if (reader.GetLeft())
-			{
-				if (process.isRemote || (m_traceView.version >= 36 && !process.isReuse))
-					sessionStats.Read(reader, m_traceView.version);
-				storageStats.Read(reader, m_traceView.version);
-				kernelStats.Read(reader, m_traceView.version);
-			}
-
-			out.Info(L"  ----------- Detours stats -----------");
-			processStats.Print(out, m_traceView.frequency);
-
-			if (!sessionStats.IsEmpty())
-			{
-				out.Info(L"");
-				out.Info(L"  ----------- Session stats -----------");
-				sessionStats.Print(out, m_traceView.frequency);
-			}
-
-			if (!storageStats.IsEmpty())
-			{
-				out.Info(L"");
-				out.Info(L"  ----------- Storage stats -----------");
-				storageStats.Print(out, m_traceView.frequency);
-			}
-
-			if (!kernelStats.IsEmpty())
-			{
-				out.Info(L"");
-				out.Info(L"  ----------- Kernel stats ------------");
-				kernelStats.Print(out, false, m_traceView.frequency);
-			}
-
-			PrintCacheWriteStats(out, process.id);
-		}
-	}
-
-	void Visualizer::WriteWorkStats(Logger& out, const TraceView::WorkRecord& record)
-	{
-		out.Info(L"  %ls", record.description);
-		out.Info(L"  Start:     %ls", TimeToText(record.start, true).str);
-		if (record.stop != ~u64(0))
-			out.Info(L"  Duration:  %ls", TimeToText(record.stop - record.start, true).str);
-		for (auto& e : record.entries)
-			out.Info(L"   %ls (%ls)", e.text, TimeToText(e.time - record.start).str);
-	}
-
-	void Visualizer::CopyTextToClipboard(const TString& str)
-	{
-		if (!OpenClipboard(m_hwnd))
-			return;
-		if (auto hglbCopy = GlobalAlloc(GMEM_MOVEABLE, (str.size() + 1) * sizeof(TCHAR)))
-		{
-			if (auto lptstrCopy = GlobalLock(hglbCopy))
-			{
-				memcpy(lptstrCopy, str.data(), (str.size() + 1) * sizeof(TCHAR));
-				GlobalUnlock(hglbCopy);
-				EmptyClipboard();
-				SetClipboardData(CF_UNICODETEXT, hglbCopy);
-			}
-		}
-		CloseClipboard();
-	}
-
-	void Visualizer::UnselectAndRedraw()
-	{
-		if (Unselect() || m_config.showCursorLine)
-			Redraw(false);
-	}
-
-	bool Visualizer::UpdateAutoscroll()
-	{
-		if (!m_autoScroll)
-			return false;
-
-		u64 playTime = GetPlayTime();
-
-
-		RECT rect;
-		GetClientRect(m_hwnd, &rect);
-		if (rect.right == 0)
-			return false;
-
-		float timeS = TimeToS(playTime);
-
-		if (m_config.AutoScaleHorizontal)
-		{
-			m_scrollPosX = 0;
-			timeS = Max(timeS, 20.0f/m_zoomValue);
-			m_config.horizontalScaleValue = Max(float(rect.right - m_progressRectLeft - 2)/(m_zoomValue*timeS*50.0f), 0.001f);
-			return true;
-		}
-		else
-		{
-			float oldScrollPosX = m_scrollPosX;
-			m_scrollPosX = Min(0.0f, (float)rect.right - timeS*50.0f*m_config.horizontalScaleValue*m_zoomValue - float(m_progressRectLeft));
-			return oldScrollPosX != m_scrollPosX;
-		}
-	}
-
-	bool Visualizer::UpdateSelection()
-	{
-		if (!m_mouseOverWindow || m_dragToScrollCounter > 0)
-			return false;
-		POINT pos;
-		GetCursorPos(&pos);
-		ScreenToClient(m_hwnd, &pos);
-
-		HitTestResult res;
-		HitTest(res, pos);
-
-		m_activeSection = res.section;
-
-		if (res.processSelected == m_processSelected && res.processLocation == m_processSelectedLocation &&
-			res.sessionSelectedIndex == m_sessionSelectedIndex &&
-			res.statsSelected == m_statsSelected && memcmp(&res.stats, &m_stats, sizeof(Stats)) == 0 &&
-			res.buttonSelected == m_buttonSelected && res.timelineSelected == m_timelineSelected &&
-			res.activeProcessGraphSelected == m_activeProcessGraphSelected && 
-			res.activeProcessCount == m_activeProcessCount && 
-			res.fetchedFilesSelected == m_fetchedFilesSelected &&
-			res.workSelected == m_workSelected && res.workTrack == m_workTrack && res.workIndex == m_workIndex &&
-			res.hyperLink == m_hyperLinkSelected)
-			return false;
-		m_processSelected = res.processSelected;
-		m_processSelectedLocation = res.processLocation;
-		m_sessionSelectedIndex = res.sessionSelectedIndex;
-		m_statsSelected = res.statsSelected;
-		m_stats = res.stats;
-		m_activeProcessGraphSelected = res.activeProcessGraphSelected;
-		m_activeProcessCount = res.activeProcessCount;
-		m_buttonSelected = res.buttonSelected;
-		m_timelineSelected = res.timelineSelected;
-		m_fetchedFilesSelected = res.fetchedFilesSelected;
-		m_workSelected = res.workSelected;
-		m_workTrack = res.workTrack;
-		m_workIndex = res.workIndex;
-		m_hyperLinkSelected = res.hyperLink;
-		return true;
-	}
-
-	void Visualizer::UpdateScrollbars(bool redraw)
-	{
-		RECT rect;
-		GetClientRect(m_hwnd, &rect);
-
-		SCROLLINFO si;
-		si.cbSize = sizeof(SCROLLINFO);
-		si.fMask = SIF_ALL | SIF_DISABLENOSCROLL;
-		si.nMin = 0;
-		si.nMax = m_contentHeight;
-		si.nPage = int(rect.bottom);
-		si.nPos = -int(m_scrollPosY);
-		si.nTrackPos = 0;
-
-		bool updateFrame = false;
-		if (ActiveProcessesShouldFillHeight())
-		{
-			if (m_verticalScrollBarEnabled)
-			{
-				SetWindowLong(m_hwnd, GWL_STYLE, GetWindowLong(m_hwnd, GWL_STYLE) & ~WS_VSCROLL);
-				m_verticalScrollBarEnabled = false;
-				updateFrame = true;
-			}
-		}
-		else
-		{
-			if (!m_verticalScrollBarEnabled)
-			{
-				SetWindowLong(m_hwnd, GWL_STYLE, GetWindowLong(m_hwnd, GWL_STYLE) | WS_VSCROLL);
-				m_verticalScrollBarEnabled = true;
-				updateFrame = true;
-			}
-			SetScrollInfo(m_hwnd, SB_VERT, &si, redraw);
-		}
-
-		if (m_config.AutoScaleHorizontal)
-		{
-			if (m_horizontalScrollBarEnabled)
-			{
-				SetWindowLong(m_hwnd, GWL_STYLE, GetWindowLong(m_hwnd, GWL_STYLE) & ~WS_HSCROLL);
-				m_horizontalScrollBarEnabled = false;
-				updateFrame = true;
-			}
-		}
-		else
-		{
-			if (!m_horizontalScrollBarEnabled)
-			{
-				SetWindowLong(m_hwnd, GWL_STYLE, GetWindowLong(m_hwnd, GWL_STYLE) | WS_HSCROLL);
-				m_horizontalScrollBarEnabled = true;
-				updateFrame = true;
-			}
-
-			si.nMax = m_contentWidth;
-			si.nPage = rect.right;
-			si.nPos = -int(m_scrollPosX);
-			SetScrollInfo(m_hwnd, SB_HORZ, &si, redraw);
-		}
-
-		if (updateFrame)
-			SetWindowPos(m_hwnd, NULL, 0, 0, 0, 0, SWP_NOSIZE | SWP_NOMOVE | SWP_NOZORDER | SWP_FRAMECHANGED);
-	}
-
-	void Visualizer::SetActiveFont(const Font& font)
-	{
-		m_activeFont = font;
-		if (m_activeHdc)
-			SelectObject(m_activeHdc, font.handle);
-	}
-
-	void Visualizer::UpdateTheme()
-	{
-		SetWindowTheme(m_hwnd, m_config.DarkMode ? L"DarkMode_Explorer" : L"Explorer", NULL);
-		SendMessageW(m_hwnd, WM_THEMECHANGED, 0, 0);
-		BOOL useDarkMode = m_config.DarkMode;
-		u32 attribute = 20; // DWMWA_USE_IMMERSIVE_DARK_MODE
-		DwmSetWindowAttribute(m_hwnd, attribute, &useDarkMode, sizeof(useDarkMode));
-	}
-
-	bool Visualizer::ActiveProcessesShouldFillHeight()
-	{
-		return !m_config.showDetailedData && !m_config.showTitleBars && !m_config.showCpuMemStats && !m_config.showNetworkStats && !m_config.showDriveStats && !m_config.showProcessBars;
-	}
-
-	StringBuffer<128> Visualizer::GetWorldTime(u64 time)
-	{
-		return GetWorldTime(TimeToS(time));
-	}
-
-	StringBuffer<128> Visualizer::GetWorldTime(float seconds)
-	{
-		return StringBuffer<128>(DateToText(float(m_traceView.traceSystemStartTimeUs / 1'000'000) + seconds).str);
-	}
-
-	void Visualizer::PostNewTrace(u32 replay, bool paused)
-	{
-		KillTimer(m_hwnd, 0);
-		PostMessage(m_hwnd, WM_NEWTRACE, replay, paused);
-	}
-
-	void Visualizer::PostNewTitle(const StringView& title)
-	{
-		#if !defined( __clang_analyzer__ )
-		PostMessage(m_hwnd, WM_SETTITLE, 0, (LPARAM)_wcsdup(title.data));
-		#endif
-	}
-
-	void Visualizer::PostQuit()
-	{
-		m_looping = false;
-		PostMessage(m_hwnd, WM_USER+666, 0, 0);
-	}
-
-	LRESULT Visualizer::WinProc(HWND hWnd, UINT Msg, WPARAM wParam, LPARAM lParam)
-	{
-		switch (Msg)
-		{
-		case WM_SETTITLE:
-		{
-			auto title = (wchar_t*)lParam;
-			SetWindowTextW(hWnd, title);
-			free(title);
-			break;
-		}
-		case WM_NEWTRACE:
-		{
-			m_replay = u32(wParam);
-			m_paused = lParam;
-			m_autoScroll = true;
-			m_scrollPosX = 0;
-			m_scrollPosY = 0;
-			Reset();
-			StringBuffer<> title;
-			GetTitlePrefix(title);
-
-			auto g = MakeGuard([&]()
-				{
-					Redraw(true);
-					UpdateScrollbars(true);
-				});
-
-			if (m_client)
-			{
-				if (!m_trace.StartReadClient(m_traceView, *m_client))
-				{
-					m_clientDisconnect.Set();
-					return false;
-				}
-				m_namedTrace.Clear().Append(m_newTraceName);
-				m_traceView.finished = false;
-			}
-			else if (!m_fileName.IsEmpty())
-			{
-				m_trace.ReadFile(m_traceView, m_fileName.data, m_replay != 0);
-				m_traceView.finished = m_replay == 0;
-				PostNewTitle(GetTitlePrefix(title).Appendf(L"%s (v%u) - %s", m_fileName.data, m_traceView.version, GetWorldTime((u64)0).data));
-			}
-			else if (m_usingNamed)
-			{
-				if (!m_trace.StartReadNamed(m_traceView, m_newTraceName.data, true, m_replay != 0))
-					return false;
-				m_namedTrace.Clear().Append(m_newTraceName);
-				m_traceView.finished = false;
-				PostNewTitle(GetTitlePrefix(title).Appendf(L"%s (Listening for new sessions on channel '%s')", m_namedTrace.data, m_listenChannel.data));
-			}
-
-			SetTimer(m_hwnd, 0, 200, NULL);
-			return 0;
-		}
-
-		case WM_SYSCOMMAND:
-			// Don't send this message through DefWindowProc as it will destroy the window 
-			// and GetMessage will stay stuck indefinitely.
-			if (wParam == SC_CLOSE)
-			{
-				PostQuit();
-				return 0;
-			}
-		break;
-
-		case WM_DESTROY:
-			PostQuit();
-			return 0;
-
-		case WM_ERASEBKGND:
-			return 1;
-
-		case WM_PAINT:
-		{
-			u64 start = GetTime();
-			PaintClient([&](HDC hdc, HDC memDC, RECT& rect)
-				{
-					FillRect(memDC, &rect, m_backgroundBrush);
-					m_activeHdc = memDC;
-					PaintAll(memDC, rect);
-					m_activeHdc = 0;
-					BitBlt(hdc, 0, 0, rect.right - rect.left, rect.bottom - rect.top, memDC, 0, 0, SRCCOPY);
-				});
-			m_lastPaintTimeMs = TimeToMs(GetTime() - start);
-			break;
-		}
-		case WM_SIZE:
-		{
-			int height = HIWORD(lParam);
-			if (m_contentHeight && m_contentHeight + int(m_scrollPosY) < height)
-				m_scrollPosY = float(Min(0, height - m_contentHeight));
-			int width = LOWORD(lParam);
-			if (m_contentWidth && m_contentWidth + int(m_scrollPosX) < width)
-				m_scrollPosX = float(Min(0, width - m_contentWidth));
-			UpdateScrollbars(true);
-			break;
-		}
-		case WM_TIMER:
-		{
-			if (m_handlingTimer) // This is to prevent hangs caused by message boxes from asserts
-				break;
-			m_handlingTimer = true;
-
-			bool changed = false;
-			if (!m_paused)
-			{
-				u64 timeOffset = (GetTime() - m_startTime - m_pauseTime) * m_replay;
-				if (!m_fileName.IsEmpty())
-				{
-					if (m_replay)
-						m_trace.UpdateReadFile(m_traceView, timeOffset, changed);
-				}
-				else if (m_client)
-				{
-					if (!m_trace.UpdateReadClient(m_traceView, *m_client, changed))
-						m_clientDisconnect.Set();
-				}
-				else if (m_usingNamed)
-				{
-					if (!m_trace.UpdateReadNamed(m_traceView, m_replay ? timeOffset : ~u64(0), changed))
-						if (m_listenTimeout.IsCreated())
-							m_listenTimeout.Set();
-				}
-			}
-
-			if (m_traceView.finished)
-			{
-				m_autoScroll = false;
-				KillTimer(m_hwnd, 0);
-				changed = true;
-			}
-
-			changed = UpdateAutoscroll() || changed;
-			changed = UpdateSelection() || changed;
-			if (changed && !IsIconic(m_hwnd))
-			{
-				UpdateScrollbars(true);
-
-				RedrawWindow(m_hwnd, NULL, NULL, RDW_INVALIDATE); // Don't use RDW_UPDATENOW.. it will cause stutter
-				u32 waitTime = 60u;//u32(Min(m_lastPaintTimeMs * 5, 200ull));
-				if (!m_traceView.finished)
-					SetTimer(m_hwnd, 0, waitTime, NULL);
-			}
-
-			m_handlingTimer = false;
-			break;
-		}
-		case WM_MOUSEWHEEL:
-		{
-			if (m_dragToScrollCounter > 0)
-				break;
-
-			int delta = GET_WHEEL_DELTA_WPARAM(wParam);
-			bool controlDown = GetAsyncKeyState(VK_CONTROL) & (1<<15);
-			bool shiftDown = GetAsyncKeyState(VK_LSHIFT) & (1<<15);
-
-			if (m_config.ScaleHorizontalWithScrollWheel || controlDown || shiftDown)
-			{
-				if (m_activeSection == 2 || !controlDown) // process bars
-				{
-					RECT r;
-					GetClientRect(hWnd, &r);
-
-					// Use mouse cursor as scroll anchor point
-					POINT cursorPos = {};
-					GetCursorPos(&cursorPos);
-					ScreenToClient(m_hwnd, &cursorPos);
-
-					float horizontalScaleValue = m_config.horizontalScaleValue;
-					float newScaleValue = horizontalScaleValue;
-					int newBoxHeight = m_config.boxHeight;
-
-					if (controlDown)
-					{
-						if (delta < 0)
-						{
-							if (newBoxHeight > 1)
-								--newBoxHeight;
-						}
-						else if (delta > 0)
-							++newBoxHeight;
-												}
-					else
-						newScaleValue = Max(horizontalScaleValue + horizontalScaleValue*float(delta)*0.0006f, 0.001f);
-
-					// TODO: m_progressRectLeft changes with zoom so anchor logic is wrong
-					const float scrollAnchorOffsetX = float(cursorPos.x) - float(m_progressRectLeft);
-					const float scrollAnchorOffsetY = 0;//float(cursorPos.y)*m_zoomValue;// - m_progressRectLeft;
-
-					float oldZoomValue = m_zoomValue;
-					if (newBoxHeight != int(m_config.boxHeight))
-					{
-						m_config.boxHeight = newBoxHeight;
-						UpdateProcessFont();
-					}
-
-					m_scrollPosY = Min(0.0f, float(m_scrollPosY - scrollAnchorOffsetY)*(m_zoomValue/oldZoomValue) + scrollAnchorOffsetY);
-					m_scrollPosX = Min(0.0f, float(m_scrollPosX - scrollAnchorOffsetX)*(m_zoomValue/oldZoomValue)*(newScaleValue/horizontalScaleValue) + scrollAnchorOffsetX);//LOWORD(lParam);
-
-					if (horizontalScaleValue != newScaleValue)
-						m_config.horizontalScaleValue = newScaleValue;
-
-
-					UpdateAutoscroll();
-					UpdateSelection();
-
-					int minScroll = r.right - m_contentWidth;
-					m_scrollPosX = Min(0.0f, Max(m_scrollPosX, float(minScroll)));
-					m_scrollPosY = Min(0.0f, Max(m_scrollPosY, float(r.bottom - m_contentHeight)));
-
-					//if (!m_traceView.finished && m_scrollPosX <= minScroll)
-					//	m_autoScroll = true;
-
-		
-					if (m_config.ShowReadWriteColors)
-						for (auto& session : m_traceView.sessions)
-							for (auto& processor : session.processors)
-								for (auto& process : processor.processes)
-									//if (TimeToMs(process.writeFilesTime, m_traceView.frequency) >= 300 || TimeToMs(process.createFilesTime, m_traceView.frequency) >= 300)
-										process.bitmapDirty = true;
-				}
-				else if (m_activeSection == 1) // active processes
-				{
-					if (delta < 0)
-						m_config.maxActiveProcessHeight = Max(m_config.maxActiveProcessHeight - 1u, 5u);
-					else if (delta > 0)
-						m_config.maxActiveProcessHeight = Min(m_config.maxActiveProcessHeight + 1u, 32u);
-				}
-				else if (m_activeSection == 0 || m_activeSection == 3) // status/timeline
-				{
-					if (delta < 0)
-						m_config.fontSize -= 1;
-					else if (delta > 0)
-						m_config.fontSize += 1;
-					UpdateDefaultFont();
-				}
-				UpdateScrollbars(true);
-				Redraw(false);
-			}
-			else
-			{
-				RECT r;
-				GetClientRect(hWnd, &r);
-				float oldScrollY = m_scrollPosY;
-				m_scrollPosY = m_scrollPosY + float(delta);
-				m_scrollPosY = Min(Max(m_scrollPosY, float(r.bottom - m_contentHeight)), 0.0f);
-				if (oldScrollY != m_scrollPosY)
-				{
-					UpdateScrollbars(true);
-					Redraw(false);
-				}
-			}
-			break;
-		}
-		case WM_NCHITTEST:
-			if (m_parentHwnd)
-				return HTCLIENT;
-			break;
-		case WM_MOUSEMOVE:
-		{
-			//m_logger.Info(TC("Section: %u"), m_activeSection);
-
-			POINTS p = MAKEPOINTS(lParam);
-			POINT pos{ p.x, p.y };
-			if (m_dragToScrollCounter > 0)
-			{
-				RECT r;
-				GetClientRect(hWnd, &r);
-
-				if (m_contentHeight <= r.bottom)
-					m_scrollPosY = 0;
-				else
-					m_scrollPosY = Max(Min(m_scrollAtAnchorY + float(pos.y - m_mouseAnchor.y), 0.0f), float(r.bottom - m_contentHeight));
-
-				if (m_contentWidth <= r.right)
-					m_scrollPosX = 0;
-				else
-				{
-					int minScroll = r.right - m_contentWidth;
-					m_scrollPosX = Max(Min(m_scrollAtAnchorX + float(pos.x - m_mouseAnchor.x), 0.0f), float(minScroll));
-					if (!m_traceView.finished && m_scrollPosX <= float(minScroll))
-						m_autoScroll = true;
-				}
-				UpdateScrollbars(true);
-				Redraw(false);
-			}
-			else
-			{
-				if (UpdateSelection() || m_config.showCursorLine)
-					Redraw(false);
-				/*
-				else
-				{
-					PaintClient([&](HDC hdc, HDC memDC, RECT& rect)
-						{
-							RECT timelineRect = rect;
-							rect.top = Min(rect.bottom, m_contentHeight) - 28;
-							rect.bottom = Min(rect.bottom, rect.top + 28);
-							FillRect(memDC, &rect, m_backgroundBrush);
-							SetBkMode(memDC, TRANSPARENT);
-							SetTextColor(memDC, m_textColor);
-							SelectObject(memDC, m_font);
-							PaintTimeline(memDC, timelineRect);
-							BitBlt(hdc, 0, rect.top, rect.right - rect.left, rect.bottom - rect.top, memDC, 0, rect.top, SRCCOPY);
-						});
-				}
-				*/
-			}
-
-			TRACKMOUSEEVENT tme;
-			tme.cbSize = sizeof(tme);
-			tme.dwFlags = TME_LEAVE;
-			tme.hwndTrack = hWnd;
-			TrackMouseEvent(&tme);
-			m_mouseOverWindow = true;
-			break;
-		}
-		case WM_MOUSELEAVE:
-			m_mouseOverWindow = false;
-			TRACKMOUSEEVENT tme;
-			tme.cbSize = sizeof(tme);
-			tme.dwFlags = TME_CANCEL;
-			tme.hwndTrack = hWnd;
-			TrackMouseEvent(&tme);
-
-			if (!m_showPopup)
-				UnselectAndRedraw();
-			break;
-
-		case WM_MBUTTONDOWN:
-		{
-			POINTS p = MAKEPOINTS(lParam);
-			StartDragToScroll(POINT{ p.x, p.y });
-			break;
-		}
-
-		case WM_LBUTTONDOWN:
-		{
-			if (!m_parentHwnd && m_traceView.sessions.empty())
-			{
-				RECT r;
-				GetClientRect(hWnd, &r);
-				int centerHrz = r.right/2;
-				int centerVrt = r.bottom/2;
-				r = { centerHrz, centerVrt, centerHrz, centerVrt };
-				InflateRect(&r, 180, 40);
-				POINTS p = MAKEPOINTS(lParam);
-				if (PtInRect(&r, { p.x, p.y }))
-				{
-					OPENFILENAME ofn;
-					tchar szFile[MAX_PATH] = TC("");
-					ZeroMemory(&ofn, sizeof(ofn));
-					ofn.lStructSize = sizeof(ofn);
-					ofn.hwndOwner = m_hwnd;
-					ofn.lpstrFilter = L"Uba Files\0*.uba\0All Files\0*.*\0";
-					ofn.lpstrFile = szFile;
-					ofn.nMaxFile = MAX_PATH;
-					ofn.Flags = OFN_FILEMUSTEXIST | OFN_PATHMUSTEXIST;
-
-					if (GetOpenFileName(&ofn))
-					{
-						m_fileName.Append(szFile);
-						PostNewTrace(0, false);
-					}
-				}
-				break;
-			}
-
-			if (m_buttonSelected != ~0u)
-			{
-				bool* values = &m_config.showProgress;
-				values[m_buttonSelected] = !values[m_buttonSelected];
-				HitTestResult res;
-				HitTest(res, { -1, -1 });
-				UpdateScrollbars(true);
-				Redraw(false);
-			}
-			else if (m_timelineSelected != 0)
-			{
-				if (!m_client) // Does not work for network streams
-				{
-					float timelineSelected = Max(m_timelineSelected, 0.0f);
-					Reset();
-
-					bool changed;
-					u64 time = MsToTime(u64(timelineSelected * 1000.0));
-
-					if (!m_fileName.IsEmpty())
-					{
-						if (!m_trace.ReadFile(m_traceView, m_fileName.data, true))
-							return false;
-					}
-					else
-					{
-						if (!m_trace.StartReadNamed(m_traceView, nullptr, true, true))
-							return false;
-						if (m_traceView.realStartTime + time > m_startTime)
-							time = m_startTime - m_traceView.realStartTime;
-					}
-
-					m_traceView.finished = false;
-					if (!m_fileName.IsEmpty())
-						m_trace.UpdateReadFile(m_traceView, time, changed);
-					else if (m_usingNamed)
-						m_trace.UpdateReadNamed(m_traceView, time, changed);
-
-					m_pauseStart = m_startTime + time;
-					m_pauseTime = m_startTime - m_pauseStart;
-
-					if (!m_paused)
-					{
-						m_autoScroll = true;
-						m_replay = 1;
-						SetTimer(m_hwnd, 0, 200, NULL);
-					}
-					else
-					{
-						m_pauseTime = 0;
-					}
-
-					HitTestResult res;
-					HitTest(res, { -1, -1 });
-					
-					RECT r;
-					GetClientRect(hWnd, &r);
-					m_scrollPosX = Min(Max(m_scrollPosX, float(r.right - m_contentWidth)), 0.0f);
-					m_scrollPosY = Min(0.0f, Max(m_scrollPosY, float(r.bottom - m_contentHeight)));
-
-					UpdateScrollbars(true);
-					Redraw(true);
-				}
-			}
-			else if (!m_hyperLinkSelected.empty())
-			{
-				ShellExecuteW(NULL, L"open", m_hyperLinkSelected.c_str(), NULL, NULL, SW_SHOW);
-			}
-			else if (m_sessionSelectedIndex != ~0u)
-			{
-				auto& session = m_traceView.sessions[m_sessionSelectedIndex];
-				for (auto& info : session.infos)
-				{
-					if (info.hyperlink.empty())
-						continue;
-					ShellExecuteW(NULL, L"open", info.hyperlink.c_str(), NULL, NULL, SW_SHOW);
-					break;
-				}
-			}
-			else
-			{
-				POINTS p = MAKEPOINTS(lParam);
-				StartDragToScroll(POINT{ p.x, p.y });
-			}
-			break;
-		}
-
-		case WM_SETCURSOR:
-		{
-			static HCURSOR arrow = LoadCursorW(NULL, IDC_ARROW);
-			static HCURSOR hand = LoadCursorW(NULL, IDC_HAND);
-			bool useHand = false;
-			if (!m_parentHwnd && m_traceView.sessions.empty())
-			{
-				RECT r;
-				GetClientRect(hWnd, &r);
-				int centerHrz = r.right/2;
-				int centerVrt = r.bottom/2;
-				r = { centerHrz, centerVrt, centerHrz, centerVrt };
-				InflateRect(&r, 180, 40);
-				POINT pt;
-				GetCursorPos(&pt);
-				ScreenToClient(hWnd, &pt);
-				useHand = PtInRect(&r, { pt.x, pt.y });
-			}
-			else if (!m_hyperLinkSelected.empty())
-			{
-				useHand = true;
-			}
-			else if (m_sessionSelectedIndex != ~0u)
-			{
-				auto& session = m_traceView.sessions[m_sessionSelectedIndex];
-				for (auto& info : session.infos)
-					if (!info.hyperlink.empty())
-						useHand = true;
-			}
-			SetCursor(useHand ? hand : arrow);
-			break;
-		}
-		case WM_LBUTTONUP:
-		{
-			if (!(m_buttonSelected != ~0u || m_timelineSelected != 0))
-			{
-				StopDragToScroll();
-			}
-			break;
-		}
-
-		case WM_RBUTTONUP:
-		{
-			POINT point;
-			point.x = LOWORD(lParam);
-			point.y = HIWORD(lParam);
-
-			HMENU hMenu = CreatePopupMenu();
-			ClientToScreen(hWnd, &point);
-
-			#define UBA_VISUALIZER_FLAG(name, defaultValue, desc) \
-				AppendMenuW(hMenu, MF_STRING | (m_config.name ? MF_CHECKED : 0), Popup_##name, desc);
-			UBA_VISUALIZER_FLAGS2
-			#undef UBA_VISUALIZER_FLAG
-
-			AppendMenuW(hMenu, MF_STRING, Popup_IncreaseFontSize, L"&Increase Font Size");
-			AppendMenuW(hMenu, MF_STRING, Popup_DecreaseFontSize, L"&Decrease Font Size");
-
-			AppendMenuW(hMenu, MF_SEPARATOR, 0, NULL);
-
-			if (m_sessionSelectedIndex != ~0u)
-			{
-				AppendMenuW(hMenu, MF_STRING, Popup_CopySessionInfo, L"&Copy Session Info");
-				AppendMenuW(hMenu, MF_SEPARATOR, 0, NULL);
-			}
-			else if (m_processSelected)
-			{
-				const TraceView::Process& process = m_traceView.GetProcess(m_processSelectedLocation);
-
-				AppendMenuW(hMenu, MF_STRING, Popup_CopyProcessInfo, L"&Copy Process Info");
-				if (!process.logLines.empty())
-					AppendMenuW(hMenu, MF_STRING, Popup_CopyProcessLog, L"Copy Process &Log");
-				if (!process.breadcrumbs.empty())
-					AppendMenuW(hMenu, MF_STRING, Popup_CopyProcessBreadcrumbs, L"Copy Process &Breadcrumbs");
-				AppendMenuW(hMenu, MF_SEPARATOR, 0, NULL);
-			}
-			else if (m_workSelected)
-			{
-				AppendMenuW(hMenu, MF_STRING, Popup_CopyWorkInfo, L"&Copy Work Info");
-			}
-
-			if (!m_traceView.sessions.empty())
-			{
-				if (!m_client)
-				{
-					if (!m_replay || m_traceView.finished)
-						AppendMenuW(hMenu, MF_STRING, Popup_Replay, L"&Replay Trace");
-					else
-					{
-						if (m_paused)
-							AppendMenuW(hMenu, MF_STRING, Popup_Play, L"&Play");
-						else
-							AppendMenuW(hMenu, MF_STRING, Popup_Pause, L"&Pause");
-						AppendMenuW(hMenu, MF_STRING, Popup_JumpToEnd, L"&Jump To End");
-					}
-				}
-
-				if (m_fileName.IsEmpty())
-					AppendMenuW(hMenu, MF_STRING, Popup_SaveAs, L"&Save Trace");
-
-				AppendMenuW(hMenu, MF_SEPARATOR, 0, NULL);
-			}
-
-			AppendMenuW(hMenu, MF_STRING, Popup_SaveSettings, L"Save Position/Settings");
-			AppendMenuW(hMenu, MF_STRING, Popup_OpenSettings, L"Open Settings file");
-			AppendMenuW(hMenu, MF_STRING, Popup_Quit, L"&Quit");
-			m_showPopup = true;
-			switch (TrackPopupMenu(hMenu, TPM_RETURNCMD | TPM_RIGHTBUTTON, point.x, point.y, 0, hWnd, NULL))
-			{
-			case Popup_SaveAs:
-			{
-				OPENFILENAME ofn;
-				TCHAR szFile[260] = { 0 };
-
-				// Initialize OPENFILENAME
-				ZeroMemory(&ofn, sizeof(ofn));
-				ofn.lStructSize = sizeof(ofn);
-				ofn.hwndOwner = hWnd;
-				ofn.lpstrFile = szFile;
-				ofn.nMaxFile = sizeof(szFile);
-				ofn.lpstrDefExt = TC("uba");
-				ofn.lpstrFilter = TC("Uba\0*.uba\0All\0*.*\0");
-				ofn.nFilterIndex = 1;
-				ofn.lpstrFileTitle = NULL;
-				ofn.nMaxFileTitle = 0;
-				ofn.lpstrInitialDir = NULL;
-				//ofn.Flags = OFN_PATHMUSTEXIST | OFN_FILEMUSTEXIST;
-				if (GetSaveFileName(&ofn))
-					m_trace.SaveAs(ofn.lpstrFile);
-				break;
-			}
-			case Popup_ShowProcessText:
-				m_config.ShowProcessText = !m_config.ShowProcessText;
-				Redraw(true);
-				break;
-			case Popup_ShowReadWriteColors:
-				m_config.ShowReadWriteColors = !m_config.ShowReadWriteColors;
-				DirtyBitmaps(false);
-				Redraw(true);
-				break;
-			case Popup_ScaleHorizontalWithScrollWheel:
-				m_config.ScaleHorizontalWithScrollWheel = !m_config.ScaleHorizontalWithScrollWheel;
-				break;
-			case Popup_ShowAllTraces:
-				m_config.ShowAllTraces = !m_config.ShowAllTraces;
-				break;
-			case Popup_SortActiveRemoteSessions:
-				m_config.SortActiveRemoteSessions = !m_config.SortActiveRemoteSessions;
-				Redraw(true);
-				break;
-			case Popup_AutoScaleHorizontal:
-				m_config.AutoScaleHorizontal = !m_config.AutoScaleHorizontal;
-				UpdateScrollbars(true);
-				Redraw(true);
-				break;
-			case Popup_LockTimelineToBottom:
-				m_config.LockTimelineToBottom = !m_config.LockTimelineToBottom;
-				Redraw(true);
-				break;
-			case Popup_DarkMode:
-			{
-				m_config.DarkMode = !m_config.DarkMode;
-				DirtyBitmaps(false);
-				InitBrushes();
-				UpdateTheme();
-				Redraw(true);
-				break;
-			}
-			case Popup_AutoSaveSettings:
-				m_config.AutoSaveSettings = !m_config.AutoSaveSettings;
-				break;
-
-			case Popup_Replay:
-				PostNewTrace(1, false);
-				break;
-
-			case Popup_Play:
-				Pause(false);
-				break;
-
-			case Popup_Pause:
-				Pause(true);
-				break;
-
-			case Popup_JumpToEnd:
-				m_traceView.finished = true;
-				PostNewTrace(0, false);
-				break;
-
-			case Popup_SaveSettings:
-				SaveSettings();
-				break;
-
-			case Popup_OpenSettings:
-				ShellExecuteW(NULL, L"open", m_config.filename.c_str(), NULL, NULL, SW_SHOW);
-				break;
-
-			case Popup_Quit: // Quit
-				PostQuit();
-				break;
-
-			case Popup_IncreaseFontSize:
-				ChangeFontSize(1);
-				break;
-			case Popup_DecreaseFontSize:
-				ChangeFontSize(-1);
-				break;
-
-			case Popup_CopySessionInfo:
-			{
-				TString str;
-				auto& session = m_traceView.sessions[m_sessionSelectedIndex];
-				StringBuffer<> fullName;
-				session.GetFullName(fullName);
-				str.append(fullName.data).append(TC("\n"));
-				for (auto& line : session.summary)
-					str.append(line).append(TC("\n"));
-				CopyTextToClipboard(str);
-				break;
-			}
-
-			case Popup_CopyProcessInfo:
-			{
-				TString str;
-				WriteTextLogger logger(str);
-				const TraceView::Process& process = m_traceView.GetProcess(m_processSelectedLocation);
-				WriteProcessStats(logger, process);
-				CopyTextToClipboard(str);
-				break;
-			}
-			case Popup_CopyProcessLog:
-			{
-				TString str;
-				const TraceView::Process& process = m_traceView.GetProcess(m_processSelectedLocation);
-				bool isFirst = true;
-				for (auto line : process.logLines)
-				{
-					if (!isFirst)
-						str += '\n';
-					isFirst = false;
-					str += line.text;
-				}
-				CopyTextToClipboard(str);
-				break;
-			}
-			case Popup_CopyProcessBreadcrumbs:
-			{
-				const TraceView::Process& process = m_traceView.GetProcess(m_processSelectedLocation);
-				CopyTextToClipboard(process.breadcrumbs);
-				break;
-			}
-			case Popup_CopyWorkInfo:
-			{
-				TString str;
-				WriteTextLogger logger(str);
-				auto& record = m_traceView.workTracks[m_workTrack].records[m_workIndex];
-				WriteWorkStats(logger, record);
-				CopyTextToClipboard(str);
-				break;
-			}
-			}
-
-			DestroyMenu(hMenu);
-			m_showPopup = false;
-			UnselectAndRedraw();
-			break;
-		}
-
-		case WM_MBUTTONUP:
-		{
-			StopDragToScroll();
-			//m_processSelected = false;
-			break;
-		}
-		case WM_KEYDOWN:
-		{
-			if (wParam == VK_SPACE)
-				Pause(!m_paused);
-			if (wParam == VK_ADD)
-				++m_replay;
-			if (wParam == VK_SUBTRACT)
-				--m_replay;
-			if (wParam == VK_BACK)
-				if (!m_filterString.empty())
-				{
-					m_filterString.resize(m_filterString.size() -1);
-					Redraw(true);
-				}
-			break;
-		}
-		case WM_CHAR:
-		{
-			tchar c = static_cast<tchar>(wParam);
-			if (c > 32 && c != '\t' && c != '\n')
-				m_filterString += static_cast<tchar>(wParam);
-			Redraw(true);
-		}
-		case WM_VSCROLL:
-			{
-				RECT r;
-				GetClientRect(hWnd, &r);
-				float oldScrollY = m_scrollPosY;
-
-				// HIWORD(wParam) only carries 16-bits, so use GetScrollInfo for larger scroll areas
-				SCROLLINFO scrollInfo = {};
-				scrollInfo.cbSize = sizeof(scrollInfo);
-				scrollInfo.fMask = SIF_TRACKPOS;
-				GetScrollInfo(m_hwnd, SB_VERT, &scrollInfo);
-
-				switch (LOWORD(wParam))
-				{
-				case SB_THUMBTRACK:
-				case SB_THUMBPOSITION:
-					m_scrollPosY = -float(scrollInfo.nTrackPos);
-					break;
-				case SB_PAGEDOWN:
-					m_scrollPosY = m_scrollPosY - float(r.bottom);
-					break;
-				case SB_PAGEUP:
-					m_scrollPosY = m_scrollPosY + float(r.bottom);
-					break;
-				case SB_LINEDOWN:
-					m_scrollPosY -= 30;
-					break;
-				case SB_LINEUP:
-					m_scrollPosY += 30;
-					break;
-				}
-				m_scrollPosY = Min(Max(m_scrollPosY, float(r.bottom - m_contentHeight)), 0.0f);
-
-				if (oldScrollY != m_scrollPosY)
-				{
-					UpdateScrollbars(true);
-					Redraw(false);
-				}
-				return 0;
-			}
-		case WM_HSCROLL:
-			{
-				RECT r;
-				GetClientRect(hWnd, &r);
-				float oldScrollX = m_scrollPosX;
-				bool autoScroll = false;
-
-				// HIWORD(wParam) only carries 16-bits, so use GetScrollInfo for larger scroll areas
-				SCROLLINFO scrollInfo = {};
-				scrollInfo.cbSize = sizeof(scrollInfo);
-				scrollInfo.fMask = SIF_TRACKPOS;
-				GetScrollInfo(m_hwnd, SB_HORZ, &scrollInfo);
-
-				switch (LOWORD(wParam))
-				{
-				case SB_THUMBTRACK:
-					m_scrollPosX = -float(scrollInfo.nTrackPos);
-					if (m_contentWidthWhenThumbTrack == 0)
-						m_contentWidthWhenThumbTrack = m_contentWidth;
-					break;
-				case SB_THUMBPOSITION:
-					autoScroll = m_contentWidthWhenThumbTrack - r.right <= HIWORD(wParam) + 10;
-					m_contentWidthWhenThumbTrack = 0;
-					m_scrollPosX = -float(scrollInfo.nTrackPos);
-					break;
-				case SB_PAGEDOWN:
-					m_scrollPosX = m_scrollPosX - float(r.right);
-					break;
-				case SB_PAGEUP:
-					m_scrollPosX = m_scrollPosX + float(r.right);
-					break;
-				case SB_LINEDOWN:
-					m_scrollPosX -= 30;
-					break;
-				case SB_LINEUP:
-					m_scrollPosX += 30;
-					break;
-				case SB_ENDSCROLL:
-					return 0;
-				}
-
-				int minScroll = r.right - m_contentWidth;
-				m_autoScroll = !m_traceView.finished && (m_scrollPosX <= float(minScroll) || autoScroll);
-				m_scrollPosX = Min(Max(m_scrollPosX, float(r.right - m_contentWidth)), 0.0f);
-
-				if (oldScrollX != m_scrollPosX)
-				{
-					UpdateScrollbars(true);
-					Redraw(false);
-				}
-				return 0;
-			}
-		}
-		return DefWindowProc(hWnd, Msg, wParam, lParam);
-	}
-
-	LRESULT CALLBACK Visualizer::StaticWinProc(HWND hWnd, UINT Msg, WPARAM wParam, LPARAM lParam)
-	{
-		auto thisPtr = (Visualizer*)GetWindowLongPtr(hWnd, GWLP_USERDATA);
-		if (!thisPtr && Msg == WM_CREATE)
-		{
-			thisPtr = (Visualizer*)lParam;
-			SetWindowLongPtr(hWnd, GWLP_USERDATA, lParam);
-			//EnableNonClientDpiScaling(hWnd);
-		}
-		if (thisPtr && hWnd == thisPtr->m_hwnd)
-			return thisPtr->WinProc(hWnd, Msg, wParam, lParam);
-		else
-			return DefWindowProc(hWnd, Msg, wParam, lParam);
 	}
 }
\ No newline at end of file
diff --git a/Engine/Source/Programs/UnrealBuildAccelerator/Visualizer/Private/UbaVisualizer.h b/Engine/Source/Programs/UnrealBuildAccelerator/Visualizer/Private/UbaVisualizer.h
--- a/Engine/Source/Programs/UnrealBuildAccelerator/Visualizer/Private/UbaVisualizer.h
+++ b/Engine/Source/Programs/UnrealBuildAccelerator/Visualizer/Private/UbaVisualizer.h
@@ -13,251 +13,10 @@ namespace uba
 	{
 	public:
 		Visualizer(VisualizerConfig& config, Logger& logger);
-		~Visualizer();
-
-		bool ShowUsingListener(const wchar_t* channelName);
-		bool ShowUsingNamedTrace(const wchar_t* namedTrace);
-		bool ShowUsingSocket(NetworkBackend& backend, const wchar_t* host, u16 port = DefaultPort);
-		bool ShowUsingFile(const wchar_t* fileName, u32 replay);
-
-		bool HasWindow();
-		HWND GetHwnd();
-		void Lock(bool lock);
-
-	private:
-		bool StartHwndThread();
-		bool Unselect();
-		void Reset();
-		void PaintClient(const Function<void(HDC hdc, HDC memDC, RECT& clientRect)>& paintFunc);
-		void PaintAll(HDC hdc, const RECT& clientRect);
-		void PaintActiveProcesses(int& posY, const RECT& clientRect, const Function<void(TraceView::ProcessLocation&, u32, bool)>& drawProcess);
-		void PaintProcessRect(TraceView::Process& process, HDC hdc, RECT rect, const RECT& progressRect, bool selected, bool writingBitmap, HBRUSH& lastSelectedBrush);
-		void PaintTimeline(HDC hdc, const RECT& clientRect);
-		using DrawTextFunc = Function<void(const StringBufferBase& text, RECT& rect, u32* outWidth)>;
-		void PaintDetailedStats(int& posY, const RECT& progressRect, TraceView::Session& session, bool isRemote, u64 playTime, const DrawTextFunc& drawTextFunc);
-
-		struct Stats
-		{
-			u64 recvBytesPerSecond = 0;
-			u64 sendBytesPerSecond = 0;
-			u64 ping = 0;
-			u64 memAvail = 0;
-			u64 memTotal = 0;
-			u64 recvBytes = 0;
-			u64 sendBytes = 0;
-			u64 procActive = 0;
-			u16 connectionCount = 0;
-			float cpuLoad = 0;
-			struct Drive
-			{
-				u8 busyPercent;
-				u64 readPerSecond = 0;
-				u64 writePerSecond = 0;
-			};
-			Map<char, Drive> drives;
-		};
-		struct HitTestResult
-		{
-			u32 section = ~0u;
-			TraceView::ProcessLocation processLocation;
-			bool processSelected = false;
-			u32 sessionSelectedIndex = ~0u;
-			bool statsSelected = false;
-			Stats stats;
-			u32 buttonSelected = ~0u;
-			float timelineSelected = 0;
-			u32 fetchedFilesSelected = ~0u;
-			bool workSelected = false;
-			u32 workTrack = ~0u;
-			u32 workIndex = ~0u;
-			bool activeProcessGraphSelected = false;
-			u16 activeProcessCount = 0u;
-			TString hyperLink;
-		};
-		u64 GetPlayTime();
-		int GetTimelineHeight();
-		int GetTimelineTop(const RECT& clientRect);
-		void HitTest(HitTestResult& outResult, const POINT& pos);
-
-		void WriteProcessStats(Logger& out, const TraceView::Process& process);
-		void WriteWorkStats(Logger& out, const TraceView::WorkRecord& record);
-		void CopyTextToClipboard(const TString& str);
-		void UnselectAndRedraw();
-		bool UpdateAutoscroll();
-		bool UpdateSelection();
-		void UpdateScrollbars(bool redraw);
-		StringBufferBase& GetTitlePrefix(StringBufferBase& out);
-		void InitBrushes();
-		void ThreadLoop();
-		void Pause(bool pause);
-		void StartDragToScroll(const POINT& anchor);
-		void StopDragToScroll();
-		void SaveSettings();
-		void DirtyBitmaps(bool full);
-		void PrintCacheWriteStats(Logger& logger, u32 processId);
-
-		struct Font
-		{
-			HFONT handle = 0;
-			HFONT handleUnderlined = 0;
-			int height = 0;
-			int offset = 0;
-		};
-
-		void UpdateFont(Font& font, int height, bool createUnderline);
-		void UpdateDefaultFont();
-		void UpdateProcessFont();
-		void ChangeFontSize(int offset);
-		void Redraw(bool now);
-		void SetActiveFont(const Font& font);
-		void UpdateTheme();
-		bool ActiveProcessesShouldFillHeight();
-		StringBuffer<128> GetWorldTime(u64 time);
-		StringBuffer<128> GetWorldTime(float seconds);
-
-		StringBuffer<256> m_namedTrace;
-		StringBuffer<256> m_fileName;
-		u32 m_replay = 0;
-		u64 m_startTime = 0;
-		u64 m_pauseTime = 0;
-
-		u64 m_lastPaintTimeMs;
-
-		struct ProcessBrushes
-		{
-			HBRUSH inProgress = 0;
-			HBRUSH inProgressRace = 0;
-			HBRUSH success = 0;
-			HBRUSH error = 0;
-			HBRUSH returned = 0;
-			HBRUSH recv = 0;
-			HBRUSH send = 0;
-			HBRUSH cacheFetchTest = 0;
-			HBRUSH cacheFetchDownload = 0;
-			HBRUSH warning = 0;
-		};
-
-		ProcessBrushes m_processBrushes[2]; // Non-selected and selected
-
-		Atomic<bool> m_looping = false;
-		HWND m_hwnd = 0;
-		HWND m_parentHwnd = 0;
-		COLORREF m_textColor = {};
-		COLORREF m_textWarningColor = {};
-		COLORREF m_textErrorColor = {};
-		COLORREF m_sendColor = {};
-		COLORREF m_recvColor = {};
-		COLORREF m_cpuColor = {};
-		COLORREF m_memColor = {};
-		COLORREF m_driveColor = {};
-		COLORREF m_activeProcColor = {};
-		HBRUSH m_backgroundBrush = 0;
-		HBRUSH m_tooltipBackgroundBrush = 0;
-		HPEN m_textPen = 0;
-		HPEN m_separatorPen = 0;
-		HPEN m_sendPen = 0;
-		HPEN m_recvPen = 0;
-		HPEN m_cpuPen = 0;
-		HPEN m_memPen = 0;
-		HPEN m_drivePen = 0;
-		HPEN m_activeProcPen = 0;
-		HPEN m_processUpdatePen = 0;
-		HPEN m_checkboxPen = 0;
-		int m_sessionStepY = 0;
-
-		Font m_defaultFont;
-		Font m_processFont;
-		Font m_timelineFont;
-		Font m_popupFont;
-
-		int m_processFontOffsetY = 0;
-
-		Font m_activeProcessFont[32];
-		u32 m_activeProcessCountHistory[5];
-		u32 m_activeProcessCountHistoryIterator = 0;
-
-		HDC m_activeHdc = 0;
-		Font m_activeFont;
-
-		int m_progressRectLeft = 30;
 
 		Logger& m_logger;
 		VisualizerConfig m_config;
 		TraceReader m_trace;
 		TraceView m_traceView;
-
-		NetworkClient* m_client = nullptr;
-		Event m_clientDisconnect;
-
-		StringBuffer<256>m_listenChannel;
-		StringBuffer<256> m_newTraceName;
-		Event m_listenTimeout;
-
-		int m_contentWidth = 0;
-		int m_contentHeight = 0;
-
-		int m_contentWidthWhenThumbTrack = 0;
-
-		float m_scrollPosX = 0;
-		float m_scrollPosY = 0;
-		float m_zoomValue = 0.5f;
-		float m_horizontalScaleValue = 0.5f;
-		bool m_autoScroll = true;
-		bool m_paused = false;
-		u64 m_pauseStart = 0;
-
-		bool m_isInPaint = false;
-
-		static constexpr int BitmapCacheHeight = 1024*1024;
-		HBITMAP m_lastBitmap = 0;
-		int m_lastBitmapOffset = BitmapCacheHeight;
-
-		u32 m_activeSection = ~0u;
-		TraceView::ProcessLocation m_processSelectedLocation;
-		bool m_processSelected = false;
-		u32 m_sessionSelectedIndex = ~0u;
-		bool m_statsSelected = false;
-		bool m_activeProcessGraphSelected = false;
-		u64 m_activeProcessCount = 0u;
-		Stats m_stats;
-		u32 m_buttonSelected = ~0u;
-		float m_timelineSelected = 0;
-		u32 m_fetchedFilesSelected = ~0u;
-		TString m_hyperLinkSelected;
-
-		bool m_workSelected = false;
-		u32 m_workTrack = ~0u;
-		u32 m_workIndex = ~0u;
-		
-		bool m_usingNamed = false;
-		bool m_mouseOverWindow = false;
-		bool m_showPopup = false;
-		bool m_locked = false;
-		bool m_handlingTimer = false;
-
-		HBITMAP m_cachedBitmap = 0;
-		RECT m_cachedBitmapRect = { INT_MIN, INT_MIN, INT_MIN, INT_MIN };
-
-		Vector<HBITMAP> m_textBitmaps;
-
-		POINT m_mouseAnchor = {};
-		float m_scrollAtAnchorX = 0;
-		float m_scrollAtAnchorY = 0;
-		int m_dragToScrollCounter = 0;
-
-		bool m_horizontalScrollBarEnabled = true;
-		bool m_verticalScrollBarEnabled = true;
-
-		TString m_filterString;
-
-		Thread m_thread;
-
-		UnorderedMap<Color, HBRUSH> m_coloredBrushes;
-
-		void PostNewTrace(u32 replay, bool paused);
-		void PostNewTitle(const StringView& title);
-		void PostQuit();
-		LRESULT WinProc(HWND hWnd, UINT Msg, WPARAM wParam, LPARAM lParam);
-		static LRESULT CALLBACK StaticWinProc(HWND hWnd, UINT Msg, WPARAM wParam, LPARAM lParam);
 	};
 }
diff --git a/Engine/Source/Programs/UnrealBuildAccelerator/Visualizer/Private/UbaVisualizerApp.cpp b/Engine/Source/Programs/UnrealBuildAccelerator/Visualizer/Private/UbaVisualizerApp.cpp
--- a/Engine/Source/Programs/UnrealBuildAccelerator/Visualizer/Private/UbaVisualizerApp.cpp
+++ b/Engine/Source/Programs/UnrealBuildAccelerator/Visualizer/Private/UbaVisualizerApp.cpp
@@ -1,9 +1,9 @@
 // Copyright Epic Games, Inc. All Rights Reserved.
 
-#include "UbaVisualizer.h"
 #include "UbaConfig.h"
 #include "UbaNetworkBackendTcp.h"
 #include "UbaVersion.h"
+#include "Windows/UbaVisualizerWin.h"
 #include <tlhelp32.h>
 #include <psapi.h>
 
@@ -76,7 +76,7 @@ struct MessageBoxLogWriter : public LogWriter
 			m_visualizer->Lock(false);
 	}
 	
-	Visualizer* m_visualizer = nullptr;
+	VisualizerWin* m_visualizer = nullptr;
 };
 
 int WINAPI wWinMain(_In_ HINSTANCE hInstance, _In_opt_ HINSTANCE hPrevInstance, _In_ PWSTR pCmdLine, _In_ int nShowCmd)
@@ -314,7 +314,7 @@ int WINAPI wWinMain(_In_ HINSTANCE hInstance, _In_opt_ HINSTANCE hPrevInstance,
 	visualizerConfig.Load(logger);
 
 	NetworkBackendTcp networkBackend(logWriter);
-	Visualizer visualizer(visualizerConfig, logger);
+	VisualizerWin visualizer(visualizerConfig, logger);
 	logWriter.m_visualizer = &visualizer;
 
 	if (channel.count)
diff --git a/Engine/Source/Programs/UnrealBuildAccelerator/Visualizer/Private/Windows/UbaVisualizerWin.cpp b/Engine/Source/Programs/UnrealBuildAccelerator/Visualizer/Private/Windows/UbaVisualizerWin.cpp
new file mode 100644
--- /dev/null
+++ b/Engine/Source/Programs/UnrealBuildAccelerator/Visualizer/Private/Windows/UbaVisualizerWin.cpp
@@ -0,0 +1,4496 @@
+// Copyright Epic Games, Inc. All Rights Reserved.
+
+#include "UbaVisualizerWin.h"
+#include "UbaConfig.h"
+#include <algorithm>
+
+#include <uxtheme.h>
+#include <dwmapi.h>
+#include <shellscalingapi.h>
+
+#pragma comment (lib, "Shcore.lib")
+#pragma comment (lib, "UxTheme.lib")
+#pragma comment (lib, "Dwmapi.lib")
+#define WM_NEWTRACE WM_USER+1
+#define WM_SETTITLE WM_USER+2
+
+namespace uba
+{
+#if PLATFORM_WINDOWS
+	extern bool g_isInAssertMessageBox;
+#endif
+
+	enum
+	{
+		Popup_CopySessionInfo = 3,
+		Popup_CopyProcessInfo,
+		Popup_CopyProcessLog,
+		Popup_CopyProcessBreadcrumbs,
+		Popup_CopyWorkInfo,
+		Popup_Replay,
+		Popup_Pause,
+		Popup_Play,
+		Popup_JumpToEnd,
+
+#define UBA_VISUALIZER_FLAG(name, defaultValue, desc) Popup_##name,
+		UBA_VISUALIZER_FLAGS2
+#undef UBA_VISUALIZER_FLAG
+
+		Popup_IncreaseFontSize,
+		Popup_DecreaseFontSize,
+
+		Popup_SaveAs,
+		Popup_SaveSettings,
+		Popup_OpenSettings,
+		Popup_Quit,
+	};
+
+	enum VisualizerFlag
+	{
+#define UBA_VISUALIZER_FLAG(name, defaultValue, desc) VisualizerFlag_##name,
+		UBA_VISUALIZER_FLAGS1
+#undef UBA_VISUALIZER_FLAG
+		VisualizerFlag_Count
+	};
+
+	class DrawTextLogger : public Logger
+	{
+	public:
+		DrawTextLogger(HWND hw, HDC h, int fh, HBRUSH bb)
+			:	hwnd(hw)
+			,	hdc(h)
+			,	fontHeight(fh)
+			,	backgroundBrush(bb)
+		{
+			textColor = GetTextColor(hdc);
+		}
+
+		virtual void BeginScope() override {}
+		virtual void EndScope() override {}
+		virtual void Log(LogEntryType type, const wchar_t* str, u32 strLen) override
+		{
+			RECT textRect{0,0,0,0};
+			DrawTextW(hdc, str, strLen, &textRect, DT_CALCRECT);
+
+			lines.emplace_back(TString(str, strLen), textOffset, height, textColor);
+			width = Max(width, int(textRect.right + textOffset));
+			height += fontHeight;
+		}
+
+		void AddSpace(int space = 5)
+		{
+			height += space;
+		}
+
+		void AddTextOffset(int offset)
+		{
+			textOffset += offset;
+		}
+
+		void AddWidth(int extra)
+		{
+			extraWidth += extra;
+		}
+
+		void DrawAtPos(int x, int y)
+		{
+			RECT r;
+			r.left = x;
+			r.top = y;
+			r.right = r.left + width;
+			r.bottom = r.top + height;
+
+			RECT clientRect;
+			GetClientRect(hwnd, &clientRect);
+
+			if (r.right > clientRect.right)
+				OffsetRect(&r, -width - 15, 0);
+			if (r.bottom > clientRect.bottom)
+			{
+				OffsetRect(&r, 0, clientRect.bottom - r.bottom);
+				if (r.top < 0)
+					OffsetRect(&r, 0, -r.top);
+			}
+
+			RECT fillRect = r;
+			fillRect.right += 2 + extraWidth;
+			FillRect(hdc, &fillRect, backgroundBrush);
+
+			for (auto& line : lines)
+			{
+				RECT tr = r;
+				tr.left += line.left;
+				tr.top += line.top;
+				SetTextColor(hdc, line.color);
+				DrawTextW(hdc, line.str.data(), u32(line.str.size()), &tr, DT_SINGLELINE|DT_NOPREFIX);
+			}
+		}
+
+		void DrawAtCursor()
+		{
+			POINT p;
+			GetCursorPos(&p);
+			ScreenToClient(hwnd, &p);
+			p.x += 3;
+			p.y += 3;
+			DrawAtPos(p.x, p.y);
+		}
+
+		DrawTextLogger& SetColor(COLORREF c) { textColor = c; return *this; }
+
+		int width = 0;
+		int height = 0;
+		int textOffset = 2;
+		int extraWidth = 0;
+		struct Line { TString str; int left; int top; COLORREF color; };
+		Vector<Line> lines;
+
+		HWND hwnd;
+		HDC hdc;
+		int fontHeight;
+		HBRUSH backgroundBrush;
+		COLORREF textColor;
+		bool isFirst = true;
+		using Logger::Log;
+	};
+
+	class WriteTextLogger : public Logger
+	{
+	public:
+		WriteTextLogger(TString& out) : m_out(out) {}
+		virtual void BeginScope() override {}
+		virtual void EndScope() override {}
+		virtual void Log(LogEntryType type, const wchar_t* str, u32 strLen) override { m_out.append(str, strLen).append(TC("\n")); }
+		TString& m_out;
+	};
+
+	VisualizerWin::VisualizerWin(VisualizerConfig& config, Logger& logger)
+		:	Visualizer(config, logger)
+	{
+		memset(m_activeProcessFont, 0, sizeof(m_activeProcessFont));
+		memset(m_activeProcessCountHistory, 0, sizeof(m_activeProcessCountHistory));
+	}
+
+	VisualizerWin::~VisualizerWin()
+	{
+		m_looping = false;
+
+		// Make sure GetMessage is triggered out of its slumber
+		PostMessage(m_hwnd, WM_QUIT, 0, 0);
+
+		m_thread.Wait();
+		delete m_client;
+	}
+
+	bool VisualizerWin::ShowUsingListener(const wchar_t* channelName)
+	{
+		TraceChannel channel(m_logger);
+		if (!channel.Init(channelName))
+		{
+			m_logger.Error(L"TODO");
+			return false;
+		}
+
+		m_listenTimeout.Create(EventResetType_Auto);
+
+		m_listenChannel.Append(channelName);
+		m_looping = true;
+		m_autoScroll = false;
+		if (!StartHwndThread())
+			return true;
+
+		{
+			StringBuffer<> title;
+			PostNewTitle(GetTitlePrefix(title).Appendf(L"Listening for new sessions on channel '%s'", m_listenChannel.data));
+		}
+
+		StringBuffer<256> traceName;
+		while (m_hwnd)
+		{
+			if (m_locked)
+			{
+				m_listenTimeout.IsSet(1000);
+				continue;
+			}
+
+			if (m_parentHwnd && !IsWindow(m_parentHwnd))
+				PostQuit();
+
+			traceName.Clear();
+			if (!channel.Read(traceName))
+			{
+				m_logger.Error(L"TODO2");
+				return false;
+			}
+
+			if (traceName.count)
+			{
+				StringBuffer<128> filter;
+				if (!m_config.ShowAllTraces)
+				{
+					OwnerInfo ownerInfo = GetOwnerInfo();
+					if (ownerInfo.pid)
+						filter.Appendf(L"_%s%u", ownerInfo.id, ownerInfo.pid);
+				}
+
+				if (!traceName.Equals(m_newTraceName.data) && traceName.EndsWith(filter))
+				{
+					m_newTraceName.Clear().Append(traceName);
+					m_usingNamed = true;
+					PostNewTrace(0, false);
+				}
+			}
+			else
+				m_newTraceName.Clear();
+
+			m_listenTimeout.IsSet(1000);
+		}
+
+		return true;
+	}
+
+	bool VisualizerWin::ShowUsingNamedTrace(const wchar_t* namedTrace)
+	{
+		m_looping = true;
+		if (!StartHwndThread())
+			return true;
+		m_newTraceName.Append(namedTrace);
+		m_usingNamed = true;
+		PostNewTrace(0, false);
+		return true;
+	}
+
+	bool VisualizerWin::ShowUsingSocket(NetworkBackend& backend, const wchar_t* host, u16 port)
+	{
+		NetworkClient* activeClient = nullptr;
+		NetworkClient* oldClient = nullptr;
+
+		auto destroyClient = MakeGuard([&]() { m_client = nullptr; delete activeClient; });
+		m_looping = true;
+		m_autoScroll = false;
+		if (!StartHwndThread())
+			return true;
+
+		m_clientDisconnect.Create(EventResetType_Manual);
+
+		wchar_t dots[] = TC("....");
+		u32 dotsCounter = 0;
+
+		StringBuffer<256> traceName;
+		while (m_hwnd)
+		{
+			if (!activeClient)
+			{
+				bool ctorSuccess = true;
+				NetworkClientCreateInfo ncci;
+				ncci.workerCount = 0;
+				activeClient = new NetworkClient(ctorSuccess, ncci);
+				if (!ctorSuccess)
+					return false;
+			}
+
+			StringBuffer<> title;
+			PostNewTitle(GetTitlePrefix(title).Appendf(L"Trying to connect to %s:%u%s", host, port, dots + ((dotsCounter--) % 4)));
+
+			bool timedOut;
+			if (!activeClient->Connect(backend, host, port, &timedOut))
+				continue;
+			m_client = activeClient;
+
+			PostNewTitle(GetTitlePrefix(title).Appendf(L"Connected to %s:%u", host, port));
+			PostNewTrace(0, false);
+
+			if (oldClient)
+			{
+				delete oldClient;
+				oldClient = nullptr;
+			}
+
+			while (m_hwnd && !m_clientDisconnect.IsSet(100) && activeClient->IsConnected())
+			{
+				m_trace.UpdateReceiveClient(*activeClient);
+			}
+
+			while (activeClient->IsConnected())
+			{
+				Sleep(500);
+				activeClient->SendKeepAlive();
+			}
+			PostNewTitle(GetTitlePrefix(title).Appendf(L"Disconnected..."));
+
+			oldClient = activeClient;
+			activeClient = nullptr;
+
+			m_clientDisconnect.Reset();
+		}
+		return true;
+	}
+
+	bool VisualizerWin::ShowUsingFile(const wchar_t* fileName, u32 replay)
+	{
+		m_looping = true;
+		m_autoScroll = false;
+		m_fileName.Append(fileName);
+		if (!StartHwndThread())
+			return true;
+		PostNewTrace(replay, 0);
+		return true;
+	}
+
+	bool VisualizerWin::StartHwndThread()
+	{
+		m_thread.Start([this]() { ThreadLoop(); return 0;}, TC("UbaHwnd"));
+		while (!m_hwnd)
+			if (m_thread.Wait(10))
+				return false;
+		return true;
+	}
+
+	bool VisualizerWin::HasWindow()
+	{
+		return m_looping == true;
+	}
+
+	HWND VisualizerWin::GetHwnd()
+	{
+		return m_hwnd;
+	}
+
+	void VisualizerWin::Lock(bool lock)
+	{
+		m_locked = lock;
+	}
+
+	StringBufferBase& VisualizerWin::GetTitlePrefix(StringBufferBase& out)
+	{
+		out.Clear();
+		out.Append(L"UbaVisualizer");
+#if UBA_DEBUG
+		out.Append(L" (DEBUG)");
+#endif
+		out.Append(L" - ");
+		return out;
+	}
+
+	bool VisualizerWin::Unselect()
+	{
+		if (m_processSelected || m_sessionSelectedIndex != ~0u || m_statsSelected || m_timelineSelected != 0 || m_fetchedFilesSelected != ~0u || m_workSelected || !m_hyperLinkSelected.empty())
+		{
+			m_processSelected = false;
+			m_sessionSelectedIndex = ~0u;
+			m_statsSelected = false;
+			m_activeProcessGraphSelected = false;
+			m_buttonSelected = ~0u;
+			m_timelineSelected = 0;
+			m_fetchedFilesSelected = ~0u;
+			m_workSelected = false;
+			m_hyperLinkSelected.clear();
+			return true;
+		}
+		return false;
+	}
+
+	void VisualizerWin::Reset()
+	{
+		for (HBITMAP bm : m_textBitmaps)
+			DeleteObject(bm);
+		DeleteObject(m_lastBitmap);
+		m_contentWidth = 0;
+		m_contentHeight = 0;
+		m_textBitmaps.clear();
+		m_lastBitmap = 0;
+		m_lastBitmapOffset = BitmapCacheHeight;
+		//m_autoScroll = true;
+		//m_scrollPosX = 0;
+		//m_scrollPosY = 0;
+		//m_zoomValue = 0.75f;
+		//m_horizontalScaleValue = 1.0f;
+
+		m_startTime = GetTime();
+		m_pauseTime = 0;
+
+		Unselect();
+	}
+
+	void VisualizerWin::InitBrushes()
+	{
+		if (m_config.DarkMode)
+		{
+			m_textColor = RGB(190, 190, 190);
+			m_textWarningColor = RGB(190, 190, 0);
+			m_textErrorColor = RGB(190, 0, 0);
+
+			m_processBrushes[0].inProgress = CreateSolidBrush(RGB(70, 70, 70));
+			m_processBrushes[1].inProgress = CreateSolidBrush(RGB(130, 130, 130));
+			m_processBrushes[0].inProgressRace = CreateSolidBrush(RGB(50, 50, 50));
+			m_processBrushes[1].inProgressRace = CreateSolidBrush(RGB(120, 120, 120));
+
+			m_processBrushes[0].error = CreateSolidBrush(RGB(140, 0, 0));
+			m_processBrushes[1].error = CreateSolidBrush(RGB(190, 0, 0));
+
+			m_processBrushes[0].returned = CreateSolidBrush(RGB(50, 50, 120));
+			m_processBrushes[1].returned = CreateSolidBrush(RGB(70, 70, 160));
+
+			m_processBrushes[0].recv = CreateSolidBrush(RGB(10, 92, 10));
+			m_processBrushes[1].recv = CreateSolidBrush(RGB(10, 130, 10));
+			m_processBrushes[0].success = CreateSolidBrush(RGB(10, 100, 10));
+			m_processBrushes[1].success = CreateSolidBrush(RGB(10, 140, 10));
+			m_processBrushes[0].send = CreateSolidBrush(RGB(10, 115, 10));
+			m_processBrushes[1].send = CreateSolidBrush(RGB(10, 145, 10));
+			m_processBrushes[0].cacheFetchTest = CreateSolidBrush(RGB(24, 120, 110));
+			m_processBrushes[1].cacheFetchTest = CreateSolidBrush(RGB(31, 150, 138));
+			m_processBrushes[0].cacheFetchDownload = CreateSolidBrush(RGB(24, 112, 110));
+			m_processBrushes[1].cacheFetchDownload = CreateSolidBrush(RGB(31, 143, 138));
+			m_processBrushes[0].warning = CreateSolidBrush(RGB(100, 100, 0));
+			m_processBrushes[1].warning = CreateSolidBrush(RGB(100, 100, 0));
+
+			m_backgroundBrush = CreateSolidBrush(0x00252526);
+			m_separatorPen = CreatePen(PS_SOLID, 1, RGB(50, 50, 50));
+			m_tooltipBackgroundBrush = CreateSolidBrush(0x00404040);
+			m_checkboxPen = CreatePen(PS_SOLID, 1, RGB(130, 130, 130));
+
+			m_sendColor = RGB(0, 170, 0);
+			m_recvColor = RGB(0, 170, 255);
+			m_cpuColor = RGB(170, 170, 0);
+			m_memColor = RGB(170, 0, 255);
+			m_driveColor = RGB(170, 65, 55);
+			m_activeProcColor = RGB(0, 170, 170);
+		}
+		else
+		{
+			m_textColor = GetSysColor(COLOR_INFOTEXT);
+			m_textWarningColor = RGB(170, 130, 0);
+			m_textErrorColor = RGB(190, 0, 0);
+
+			m_processBrushes[0].inProgress = CreateSolidBrush(RGB(150, 150, 150));
+			m_processBrushes[1].inProgress = CreateSolidBrush(RGB(180, 180, 180));
+			m_processBrushes[0].inProgressRace = CreateSolidBrush(RGB(130, 130, 130));
+			m_processBrushes[1].inProgressRace = CreateSolidBrush(RGB(170, 170, 170));
+
+			m_processBrushes[0].error = CreateSolidBrush(RGB(255, 70, 70));
+			m_processBrushes[1].error = CreateSolidBrush(RGB(255, 100, 70));
+
+			m_processBrushes[0].returned = CreateSolidBrush(RGB(150, 150, 200));
+			m_processBrushes[1].returned = CreateSolidBrush(RGB(170, 170, 200));
+
+			m_processBrushes[0].recv = CreateSolidBrush(RGB(10, 190, 10));
+			m_processBrushes[1].recv = CreateSolidBrush(RGB(20, 210, 20));
+			m_processBrushes[0].success = CreateSolidBrush(RGB(10, 200, 10));
+			m_processBrushes[1].success = CreateSolidBrush(RGB(20, 220, 20));
+			m_processBrushes[0].send = CreateSolidBrush(RGB(80, 210, 80));
+			m_processBrushes[1].send = CreateSolidBrush(RGB(90, 250, 90));
+			m_processBrushes[0].cacheFetchTest = CreateSolidBrush(RGB(150, 150, 200));
+			m_processBrushes[1].cacheFetchTest = CreateSolidBrush(RGB(170, 170, 200));
+			m_processBrushes[0].cacheFetchDownload = CreateSolidBrush(RGB(150, 150, 200));
+			m_processBrushes[1].cacheFetchDownload = CreateSolidBrush(RGB(170, 170, 200));
+			m_processBrushes[0].warning = CreateSolidBrush(RGB(170, 170, 0));
+			m_processBrushes[1].warning = CreateSolidBrush(RGB(170, 170, 0));
+
+			m_backgroundBrush = GetSysColorBrush(0);
+			m_separatorPen = CreatePen(PS_SOLID, 1, RGB(180, 180, 180));
+			m_tooltipBackgroundBrush = GetSysColorBrush(COLOR_INFOBK);
+			m_checkboxPen = CreatePen(PS_SOLID, 1, RGB(130, 130, 130));
+
+			m_sendColor = RGB(0, 170, 0); // Green
+			m_recvColor = RGB(63, 72, 204); // Blue
+			m_cpuColor = RGB(200, 130, 0); // Orange
+			m_memColor = RGB(170, 0, 255); // Purple
+			m_driveColor = RGB(255, 115, 96);
+			m_activeProcColor = RGB(0, 170, 170);
+		}
+
+		m_textPen = CreatePen(PS_SOLID, 1, m_textColor);
+		m_sendPen = CreatePen(PS_SOLID, 1, m_sendColor);
+		m_recvPen = CreatePen(PS_SOLID, 1, m_recvColor);
+		m_cpuPen = CreatePen(PS_SOLID, 1, m_cpuColor);
+		m_memPen = CreatePen(PS_SOLID, 1, m_memColor);
+		m_drivePen = CreatePen(PS_SOLID, 1, m_driveColor);
+		m_activeProcPen = CreatePen(PS_SOLID, 1, m_activeProcColor);
+	}
+
+	void VisualizerWin::ThreadLoop()
+	{
+		if (m_config.parent)
+		{
+			SetProcessDpiAwareness(PROCESS_SYSTEM_DPI_AWARE);
+		}
+
+		InitBrushes();
+
+		LOGBRUSH br = { 0 };
+		GetObject(m_backgroundBrush, sizeof(br), &br);
+		m_processUpdatePen = CreatePen(PS_SOLID, 2, RGB(GetRValue(br.lbColor), GetGValue(br.lbColor), GetBValue(br.lbColor)));
+
+		HINSTANCE hInstance = GetModuleHandle(NULL);
+		int winPosX = m_config.x;
+		int winPosY = m_config.y;
+		int winWidth = m_config.width;
+		int winHeight = m_config.height;
+
+		RECT rectCombined;
+		SetRectEmpty(&rectCombined);
+		EnumDisplayMonitors(0, 0, [](HMONITOR hMon,HDC hdc,LPRECT lprcMonitor,LPARAM pData)
+			{
+				RECT* rectCombined = reinterpret_cast<RECT*>(pData);
+				UnionRect(rectCombined, rectCombined, lprcMonitor);
+				return TRUE;
+			}, (LPARAM)&rectCombined);
+
+		winPosX = Max(int(rectCombined.left), winPosX);
+		winPosY = Max(int(rectCombined.top), winPosY);
+		winPosX = Min(int(rectCombined.right) - winWidth, winPosX);
+		winPosY = Min(int(rectCombined.bottom) - winHeight, winPosY);
+
+		WNDCLASSEX wndClassEx;
+		ZeroMemory(&wndClassEx, sizeof(wndClassEx));
+		wndClassEx.cbSize = sizeof(wndClassEx);
+		wndClassEx.style = CS_HREDRAW | CS_VREDRAW;
+		wndClassEx.lpfnWndProc = &StaticWinProc;
+		wndClassEx.hIcon = LoadIcon(hInstance, MAKEINTRESOURCE(123));
+		wndClassEx.hCursor = 0;
+		wndClassEx.hInstance = hInstance;
+		wndClassEx.hbrBackground = NULL;
+		wndClassEx.lpszClassName = TEXT("UbaVisualizer");
+		ATOM wndClassAtom = RegisterClassEx(&wndClassEx);
+		const TCHAR* windowClassName = MAKEINTATOM(wndClassAtom);
+
+		auto unreg = MakeGuard([&]() { UnregisterClass(windowClassName, hInstance); });
+
+		UpdateDefaultFont();
+
+		UpdateProcessFont();
+
+		const TCHAR* fontName = TEXT("Consolas");
+		m_popupFont.handle = CreateFontW(-12, 0, 0, 0, FW_NORMAL, false, false, false, ANSI_CHARSET, OUT_DEFAULT_PRECIS, CLIP_DEFAULT_PRECIS, CLEARTYPE_QUALITY, FIXED_PITCH | FF_MODERN, fontName);
+		m_popupFont.height = 14;
+		//m_popupFont = (HFONT)GetStockObject(SYSTEM_FIXED_FONT);
+
+		DWORD scrollbarFlags = 0;
+		m_verticalScrollBarEnabled = !ActiveProcessesShouldFillHeight();
+		if (m_verticalScrollBarEnabled)
+			scrollbarFlags |= WS_VSCROLL;
+		m_horizontalScrollBarEnabled = !m_config.AutoScaleHorizontal;
+		if (m_horizontalScrollBarEnabled)
+			scrollbarFlags |= WS_HSCROLL;
+
+		DWORD windowStyle = WS_CAPTION | WS_SYSMENU | WS_THICKFRAME | WS_MAXIMIZEBOX | WS_MINIMIZEBOX | WS_CLIPCHILDREN | scrollbarFlags;
+
+		DWORD exStyle = 0;
+		if (m_config.parent)
+			windowStyle = WS_POPUP | scrollbarFlags;
+
+		StringBuffer<> title;
+		GetTitlePrefix(title).Append(L"Initializing...");
+
+		HWND hwnd = CreateWindowEx(exStyle, windowClassName, title.data, windowStyle, winPosX, winPosY, winWidth, winHeight, NULL, NULL, hInstance, this);
+		SetWindowLongPtr(hwnd, GWLP_USERDATA, (LONG_PTR)this);
+		m_hwnd = hwnd;
+
+		auto destroyWin = [&]()
+			{
+				if (m_hwnd)
+				{
+					if (m_config.AutoSaveSettings)
+						SaveSettings();
+					DestroyWindow(m_hwnd);
+					m_hwnd = 0;
+				}
+			};
+
+		BOOL cloak = TRUE;
+		DwmSetWindowAttribute(hwnd, DWMWA_CLOAK, &cloak, sizeof(cloak));
+		auto exitCloak = MakeGuard([&]() { cloak = FALSE; DwmSetWindowAttribute(hwnd, DWMWA_CLOAK, &cloak, sizeof(cloak)); });
+
+		if (m_config.DarkMode)
+			UpdateTheme();
+
+		HitTestResult res;
+		HitTest(res, { -1, -1 });
+
+		if (m_config.parent)
+		{
+			exitCloak.Execute();
+
+			// If not child it will not propagate keyboard presses etc to parent
+			SetWindowLong(hwnd, GWL_STYLE, WS_CHILD | scrollbarFlags);
+
+			m_parentHwnd = (HWND)(uintptr_t)m_config.parent;
+			if (!SetParent(hwnd, m_parentHwnd))
+				m_logger.Error(L"SetParent failed using parentHwnd 0x%llx", m_parentHwnd);
+
+			UpdateWindow(m_hwnd);
+			UpdateScrollbars(true);
+			PostMessage(m_parentHwnd, 0x0444, 0, (LPARAM)hwnd);
+		}
+		else
+		{
+			ShowWindow(hwnd, SW_SHOW);
+			UpdateWindow(m_hwnd);
+			UpdateScrollbars(true);
+			exitCloak.Execute();
+		}
+
+		m_startTime = GetTime();
+
+		while (m_looping)
+		{
+			DWORD timeoutMs = 2000;
+			DWORD result = MsgWaitForMultipleObjects(0, NULL, FALSE, timeoutMs, QS_ALLINPUT);
+			if (result == WAIT_TIMEOUT)
+				continue;
+			if (result != WAIT_OBJECT_0)
+				break;
+			MSG msg;
+			while (PeekMessage(&msg, NULL, 0, 0, PM_REMOVE))
+			{
+				TranslateMessage(&msg);
+				DispatchMessage(&msg);
+
+				// It may happen that we receive the WM_DESTOY message from within the DistpachMessage above and handle it directly in WndProc.
+				// So before trying to call GetMessage again, we just need to validate that m_looping is still true otherwise we could end
+				// up waiting forever for this thread to exit.
+				if (!m_looping || msg.message == WM_QUIT || msg.message == WM_DESTROY || msg.message == WM_CLOSE)
+				{
+					destroyWin();
+					m_looping = false;
+					if (m_listenTimeout.IsCreated())
+						m_listenTimeout.Set();
+					break;
+				}
+			}
+		}
+
+		destroyWin();
+	}
+
+	void VisualizerWin::Pause(bool pause)
+	{
+		if (m_paused == pause)
+			return;
+
+		m_paused = pause;
+		if (pause)
+		{
+			m_pauseStart = GetTime();
+		}
+		else
+		{
+			m_replay = 1;
+			m_pauseTime += GetTime() - m_pauseStart;
+			m_traceView.finished = false;
+			SetTimer(m_hwnd, 0, 200, NULL);
+		}
+	}
+
+	void VisualizerWin::StartDragToScroll(const POINT& anchor)
+	{
+		// Uses reference-counter method since multiple input events (left and middle mouse button) can trigger the drag-to-scroll mechansim
+		if (m_dragToScrollCounter == 0)
+		{
+			m_processSelected = false;
+			m_sessionSelectedIndex = ~0u;
+			m_statsSelected = false;
+			m_activeProcessGraphSelected = false;
+			m_buttonSelected = ~0u;
+			m_timelineSelected = 0;
+			m_fetchedFilesSelected = ~0u;
+			m_workSelected = false;
+			m_hyperLinkSelected.clear();
+			m_autoScroll = false;
+			m_mouseAnchor = { anchor.x, anchor.y };
+			m_scrollAtAnchorX = m_scrollPosX;
+			m_scrollAtAnchorY = m_scrollPosY;
+			SetCapture(m_hwnd);
+			Redraw(false);
+		}
+		++m_dragToScrollCounter;
+	}
+
+	void VisualizerWin::StopDragToScroll()
+	{
+		if (m_dragToScrollCounter > 0)
+			--m_dragToScrollCounter;
+		if (m_dragToScrollCounter != 0)
+			return;
+
+		ReleaseCapture();
+		if (UpdateSelection())
+			Redraw(false);
+	}
+
+	void VisualizerWin::SaveSettings()
+	{
+		RECT rect;
+		GetWindowRect(m_hwnd, &rect);
+
+		m_config.x = rect.left;
+		m_config.y = rect.top;
+		m_config.width = rect.right - rect.left;
+		m_config.height = rect.bottom - rect.top;
+		m_config.Save(m_logger);
+	}
+
+	void VisualizerWin::DirtyBitmaps(bool full)
+	{
+		for (auto& session : m_traceView.sessions)
+			for (auto& processor : session.processors)
+				for (auto& process : processor.processes)
+				{
+					process.bitmapDirty = true;
+					if (full)
+						process.bitmap = 0;
+				}
+
+		for (auto& workTrack : m_traceView.workTracks)
+			for (auto& work : workTrack.records)
+			{
+				work.bitmapDirty = true;
+				if (full)
+					work.bitmap = 0;
+			}
+
+		if (!full)
+			return;
+
+		for (HBITMAP bm : m_textBitmaps)
+			DeleteObject(bm);
+		DeleteObject(m_lastBitmap);
+		m_textBitmaps.clear();
+		m_lastBitmapOffset = BitmapCacheHeight;
+		m_lastBitmap = 0;
+	}
+
+	void VisualizerWin::PrintCacheWriteStats(Logger& logger, u32 processId)
+	{
+		auto findIt = m_traceView.cacheWrites.find(processId);
+		if (findIt == m_traceView.cacheWrites.end())
+			return;
+		TraceView::CacheWrite& write = findIt->second;
+		logger.Info(L"");
+		logger.Info(L"  -------- Cache write stats ----------");
+		if (write.end)
+		{
+			logger.Info(L"  Duration                    %9s", TimeToText(write.end - write.start).str);
+			logger.Info(L"  Success                     %9s", write.success ? L"true" : L"false");
+
+			BinaryReader reader(write.stats.data(), 0, write.stats.size());
+			CacheSendStats stats;
+			if (m_traceView.version < 45)
+			{
+				stats.sendFile.bytes = reader.Read7BitEncoded();
+				stats.sendFile.count = 1;
+			}
+			else
+				stats.Read(reader, m_traceView.version);
+			stats.Print(logger, m_traceView.frequency);
+		}
+		else
+		{
+			logger.Info(L"  In progress...");
+		}
+	}
+
+	void VisualizerWin::UpdateFont(Font& font, int height, bool createUnderline)
+	{
+		font.height = height;
+
+		int fh = height;
+		font.offset = 0;
+		if (height <= 13)
+		{
+			++fh;
+			--font.offset;
+		}
+		if (height <= 11)
+			++fh;
+		if (height <= 9)
+			++fh;
+		if (height <= 8)
+			++fh;
+		if (height <= 6)
+			++fh;
+		if (height <= 4)
+			--font.offset;
+
+		if (font.handle)
+			DeleteObject(font.handle);
+		if (font.handleUnderlined)
+			DeleteObject(font.handleUnderlined);
+
+		//NONCLIENTMETRICS nonClientMetrics;
+		//nonClientMetrics.cbSize = sizeof(nonClientMetrics);
+		//SystemParametersInfo(SPI_GETNONCLIENTMETRICS, sizeof(nonClientMetrics), &nonClientMetrics, 0);
+		//m_font = (HFONT)CreateFontIndirect(&nonClientMetrics.lfMessageFont);
+
+		font.handle = CreateFontW(4 - fh, 0, 0, 0, FW_NORMAL, 0, 0, 0, ANSI_CHARSET, OUT_DEFAULT_PRECIS, CLIP_DEFAULT_PRECIS, CLEARTYPE_QUALITY, DEFAULT_PITCH, m_config.fontName.c_str());
+		if (createUnderline)
+			font.handleUnderlined = CreateFontW(4 - fh, 0, 0, 0, FW_NORMAL, 0, 1, 0, ANSI_CHARSET, OUT_DEFAULT_PRECIS, CLIP_DEFAULT_PRECIS, CLEARTYPE_QUALITY, DEFAULT_PITCH, m_config.fontName.c_str());
+	}
+
+	void VisualizerWin::UpdateDefaultFont()
+	{
+		UpdateFont(m_defaultFont, m_config.fontSize, true);
+		m_sessionStepY = m_defaultFont.height + 4;
+
+		m_timelineFont = m_defaultFont;
+	}
+
+	void VisualizerWin::UpdateProcessFont()
+	{
+		m_zoomValue = 1.0f + float(m_config.boxHeight) / 30.0f;
+		int fontHeight = Max(int(m_config.boxHeight) - 2, 1);
+		UpdateFont(m_processFont, fontHeight, false);
+		m_progressRectLeft = int(13 + float(m_processFont.height) * 1.5f);
+		DirtyBitmaps(true);
+	}
+
+	void VisualizerWin::ChangeFontSize(int offset)
+	{
+		m_config.fontSize += offset;
+		m_config.fontSize = Max(m_config.fontSize, 10u);
+		UpdateDefaultFont();
+		Redraw(true);
+	}
+
+	void VisualizerWin::Redraw(bool now)
+	{
+		u32 flags = RDW_INVALIDATE;
+		if (now)
+			flags |= RDW_UPDATENOW;
+		RedrawWindow(m_hwnd, NULL, NULL, flags);
+
+		u32 activeProcessCount = u32(m_trace.m_activeProcesses.size());
+		for (u32 i=0; i!=sizeof_array(m_activeProcessCountHistory); ++i)
+			m_activeProcessCountHistory[i] = activeProcessCount;
+	}
+
+	void VisualizerWin::PaintClient(const Function<void(HDC hdc, HDC memDC, RECT& clientRect)>& paintFunc)
+	{
+		HDC hdc = GetDC(m_hwnd);
+
+		RECT rect;
+		GetClientRect(m_hwnd, &rect);
+
+		HDC memDC = CreateCompatibleDC(hdc);
+
+		if (!EqualRect(&m_cachedBitmapRect, &rect))
+		{
+			if (m_cachedBitmap)
+				DeleteObject(m_cachedBitmap);
+			m_cachedBitmap = CreateCompatibleBitmap(hdc, rect.right - rect.left, rect.bottom - rect.top);
+			m_cachedBitmapRect = rect;
+		}
+		HGDIOBJ oldBmp = SelectObject(memDC, m_cachedBitmap);
+
+		if (!m_isInPaint)
+		{
+			m_isInPaint = true;
+			paintFunc(hdc, memDC, rect);
+			m_isInPaint = false;
+		}
+
+		SelectObject(memDC, oldBmp);
+		DeleteDC(memDC);
+
+		ReleaseDC(m_hwnd, hdc);
+	}
+
+	constexpr int GraphHeight = 30;
+
+	struct SessionRec
+	{
+		TraceView::Session* session;
+		u32 index;
+	};
+	void Populate(SessionRec* recs, TraceView& traceView, bool sort)
+	{
+		u32 count = u32(traceView.sessions.size());
+		for (u32 i = 0, e = count; i != e; ++i)
+			recs[i] = { &traceView.sessions[i], i };
+		if (count <= 1 || !sort)
+			return;
+		std::sort(recs + 1, recs + traceView.sessions.size(), [](SessionRec& a, SessionRec& b)
+			{
+				auto& as = *a.session;
+				auto& bs = *b.session;
+				if ((as.processActiveCount != 0) != (bs.processActiveCount != 0))
+					return as.processActiveCount > bs.processActiveCount;
+				if (as.processActiveCount && as.proxyCreated != bs.proxyCreated)
+					return int(as.proxyCreated) > int(bs.proxyCreated);
+				return a.index < b.index;
+			});
+	}
+
+	void VisualizerWin::PaintAll(HDC hdc, const RECT& clientRect)
+	{
+		u64 playTime = GetPlayTime();
+
+		SetBkMode(hdc, TRANSPARENT);
+		SetTextColor(hdc, m_textColor);
+		SetBkColor(hdc, m_config.DarkMode ? RGB(70, 70, 70) : RGB(180, 180, 180));
+
+		HBRUSH lastSelectedBrush = 0;
+		HBRUSH textLastSelectedBrush = 0;
+		auto fillRect = [&](const RECT& r, HBRUSH b)
+			{
+				if (lastSelectedBrush != b)
+				{
+					SelectObject(hdc, b);
+					lastSelectedBrush = b;
+				}
+				PatBlt(hdc, r.left, r.top, r.right - r.left, r.bottom - r.top, PATCOPY);
+			};
+
+		auto DrawCenteredText = [&](const StringBuffer<>* lines, u32 lineCount, bool drawBackground)
+			{
+				for (u32 i=0; i!=lineCount; ++i)
+				{
+					RECT rect = { 0, 0, 0, 0 };
+					DrawTextW(hdc, lines[i].data, lines[i].count, &rect, DT_SINGLELINE|DT_NOPREFIX|DT_NOCLIP|DT_CALCRECT);
+					int top = (clientRect.bottom - rect.bottom)/2 + i*20;
+					rect.left = (clientRect.right - rect.right)/2;
+					rect.right = (clientRect.right + rect.right)/2;
+					rect.bottom = top + (rect.bottom - rect.top);
+					rect.top = top;
+
+					if (drawBackground)
+					{
+						RECT r = rect;
+						InflateRect(&r, 4, 2);
+						fillRect(r, m_tooltipBackgroundBrush);
+					}
+					DrawTextW(hdc, lines[i].data, lines[i].count, &rect, DT_SINGLELINE|DT_NOPREFIX|DT_NOCLIP);
+				}
+			};
+
+
+		if (m_traceView.sessions.empty() && !m_parentHwnd)
+		{
+			StringBuffer<> str[2];
+			if (!m_fileName.IsEmpty())
+			{
+				str[0].Append(TCV("Loading trace file"));
+			}
+			else
+			{
+				str[1].Append(TCV("Click here to open trace file"));
+				if (m_listenChannel.count)
+				{
+					str[0].Appendf(TC("Listening for new sessions on channel '%s'"), m_listenChannel.data);
+					str[1].Append(TCV(" instead"));
+				}
+				else
+					str[0].Append(TCV("No trace active"));
+			}
+			SetActiveFont(m_popupFont);
+			DrawCenteredText(str, 2, false);
+			return;
+		}
+
+
+		int posY = int(m_scrollPosY);
+		float scaleX = 50.0f*m_zoomValue*m_config.horizontalScaleValue;
+
+		RECT progressRect = clientRect;
+		progressRect.left += m_progressRectLeft;
+
+		if (m_config.showTimeline)
+		{
+			progressRect.bottom -= m_defaultFont.height + 10;
+		}
+
+		HDC textDC = CreateCompatibleDC(hdc);
+		SetTextColor(textDC, m_textColor);
+		SelectObject(textDC, m_processFont.handle);
+		SelectObject(textDC, GetStockObject(NULL_BRUSH));
+		SetBkMode(textDC, TRANSPARENT);
+
+		//TEXTMETRIC metric;
+		//GetTextMetrics(textDC, &metric);
+
+		HBITMAP nullBmp = CreateCompatibleBitmap(hdc, 1, 1);
+		HBITMAP oldBmp = (HBITMAP)SelectObject(textDC, nullBmp);
+		HBITMAP lastSelectedBitmap = 0;
+
+		u64 lastStop = 0;
+
+		SetActiveFont(m_defaultFont);
+
+		auto drawStatusText = [&](StringView text, LogEntryType type, int posX, int endX, bool moveY, bool underlined = false, bool centered = false)
+			{
+				RECT rect;
+				rect.left = posX;
+				rect.right = endX;
+				rect.top = posY + m_activeFont.offset;
+				rect.bottom = posY + m_activeFont.height + 2;
+				SetTextColor(hdc, type == LogEntryType_Info ? m_textColor : (type == LogEntryType_Error ? m_textErrorColor : m_textWarningColor));
+
+				u32 align = 0;
+				if (centered)
+				{
+					posX += (endX - posX)/2;
+					align = SetTextAlign(hdc, TA_CENTER);
+				}
+				if (underlined)
+					SelectObject(m_activeHdc, m_activeFont.handleUnderlined);
+				ExtTextOutW(hdc, posX, posY, ETO_CLIPPED, &rect, text.data, text.count, NULL);
+				if (underlined)
+					SelectObject(m_activeHdc, m_activeFont.handle);
+				if (centered)
+					SetTextAlign(hdc, align);
+				if (moveY)
+					posY = rect.bottom;
+			};
+
+		auto drawIndentedText = [&](const StringView& text, LogEntryType type, int indent, bool moveY, bool underlined = false)
+			{
+				int posX = 5 + indent*m_defaultFont.height;
+				drawStatusText(text, type, posX, clientRect.right, moveY, underlined);
+			};
+
+		if (m_config.showProgress && m_traceView.progressProcessesTotal)
+		{
+			drawIndentedText(AsView(L"Progress"), LogEntryType_Info, 1, false);
+
+
+			float progress = float(m_traceView.progressProcessesDone) / float(m_traceView.progressProcessesTotal);
+			u32 start = 3 + 6*m_activeFont.height;
+			u32 width = m_activeFont.height * 18;
+			RECT rect;
+			rect.left = start;
+			rect.right = start + width;
+			rect.top = posY;
+			rect.bottom = posY + m_activeFont.height;
+			fillRect(rect, m_processBrushes[0].inProgress);
+
+			rect.right = rect.left + int(progress*float(width));//durationMs/100;
+			fillRect(rect, m_traceView.progressErrorCount ? m_processBrushes[0].error : m_processBrushes[0].success);
+
+			StringBuffer<> str;
+			str.Appendf(TC("%u%%    %u / %u"), u32(progress*100.0f), m_traceView.progressProcessesDone, m_traceView.progressProcessesTotal);
+
+			const tchar* remoteDisabled = TC(""); // TODO: Only show this in "advanced mode" m_traceView.remoteExecutionDisabled ? TC("   (Remote spawn disabled)") : TC("");
+			str.Append(remoteDisabled);
+
+			if (u32 active = m_traceView.totalProcessActiveCount)
+				str.Appendf(TC("   (%u active)"), active);
+
+			if (m_traceView.lastKillProcessTime && playTime >= m_traceView.lastKillProcessTime && playTime < m_traceView.lastKillProcessTime + MsToTime(800))
+			{
+				rect.left = start + width + 20;
+				rect.right = rect.left + m_activeFont.height * 7;
+				fillRect(rect, m_processBrushes[0].error);
+				drawStatusText(TCV("Killing (mem)"), LogEntryType_Info, rect.left, rect.right, 0, false, true);
+			}
+			else if (m_traceView.lastSpawningDelayStartTime && playTime >= m_traceView.lastSpawningDelayStartTime && playTime < m_traceView.lastSpawningDelayEndTime + MsToTime(500))
+			{
+				rect.left = start + width + 20;
+				rect.right = rect.left + m_activeFont.height * 7;
+				fillRect(rect, m_processBrushes[0].warning);
+				drawStatusText(TCV("Throttling (mem)"), LogEntryType_Info, rect.left, rect.right, 0, false, true);
+			}
+
+
+			drawIndentedText(str, LogEntryType_Info, 6, true);
+		}
+
+		if (m_traceView.version && (m_traceView.version < TraceReadCompatibilityVersion || m_traceView.version > TraceVersion))
+		{
+			if (!m_traceView.finished)
+			{
+				SetTextColor(hdc, m_textWarningColor);
+				StringBuffer<> str;
+				str.Appendf(TC("Unsupported trace version %u (Versions supported are %u to %u)"), m_traceView.version, TraceReadCompatibilityVersion, TraceVersion);
+				HFONT handle = CreateFontW(-20, 0, 0, 0, FW_NORMAL, 0, 0, 0, ANSI_CHARSET, OUT_DEFAULT_PRECIS, CLIP_DEFAULT_PRECIS, CLEARTYPE_QUALITY, DEFAULT_PITCH, m_config.fontName.c_str());
+				SelectObject(hdc, handle);
+				DrawTextW(hdc, str.data, str.count, (RECT*)&clientRect, DT_CENTER|DT_NOPREFIX|DT_SINGLELINE|DT_VCENTER);
+				DeleteObject(handle);
+			}
+			else
+			{
+				m_traceView.Clear();
+				StringBuffer<> title;
+				SetWindowTextW(m_hwnd, GetTitlePrefix(title).Appendf(L" (Listening for new sessions on channel '%s')", m_listenChannel.data).data);
+			}
+			return;
+		}
+
+		if (m_config.showStatus && !m_traceView.statusMap.empty())
+		{
+			u32 lastRow = ~0u;
+			u32 row = ~0u;
+			for (auto& kv : m_traceView.statusMap)
+			{
+				auto& status = kv.second;
+				if (status.text.empty())
+					continue;
+				row = u32(kv.first >> 32);
+				if (lastRow != ~0u && lastRow != row)
+					posY += m_activeFont.height + 2;
+				lastRow = row;
+				u32 column = u32(kv.first & ~0u);
+				drawIndentedText(status.text, status.type, column, false, !status.link.empty());
+			}
+			if (row != ~0u)
+				posY += m_activeFont.height + 2;
+
+			SetTextColor(hdc, m_textColor);
+			posY += 3;
+		}
+
+		if (m_config.showActiveProcessGraph)
+		{
+			if (posY + GraphHeight >= progressRect.top && posY + GraphHeight - 5 < progressRect.bottom)
+			{
+				int graphBaseY = posY + GraphHeight - 4;
+				double graphHeight = double(GraphHeight - 2);
+
+				Vector<POINT> line;
+
+				SelectObject(hdc, m_activeProcPen);
+				bool isFirstUpdate = true;
+				bool isFirstDraw = true;
+				u64 prevValue = 0;
+				int prevX = 0;
+				int prevY = 0;
+
+				float timeScale = (m_config.horizontalScaleValue*m_zoomValue)*50.0f;
+				float startOffset = ((m_scrollPosX/timeScale) - float(int(m_scrollPosX/timeScale))) * timeScale;
+				float time = -startOffset/timeScale;
+				float playTimeS = TimeToS(playTime);
+				u32 activeProcessCountIndex = 0;
+				u16 count = 0;
+				while (time < playTimeS && activeProcessCountIndex < m_traceView.activeProcessCounts.size())
+				{
+					int posX = progressRect.left + int(startOffset + time*timeScale);
+					if (posX >= clientRect.right)
+						break;
+
+					while (activeProcessCountIndex < m_traceView.activeProcessCounts.size())
+					{
+						auto& entry = m_traceView.activeProcessCounts[activeProcessCountIndex];
+						count = entry.count;
+						if (float(TimeToMs(entry.time))/1000.0f > time)
+							break;
+						++activeProcessCountIndex;
+					}
+
+					int x = posX;
+					int y = graphBaseY;
+
+					double div = double(m_traceView.maxActiveProcessCount);
+					y -= int(double(count) * graphHeight / div);
+
+					if (prevValue <= 0)
+						prevY = y;
+
+					if (x > clientRect.left)
+					{
+						if (isFirstUpdate)
+						{
+							line.push_back({ x, y });
+							isFirstDraw = false;
+						}
+						else
+						{
+							if (isFirstDraw)
+								line.push_back({ prevX, prevY });
+							line.push_back({ x, y });
+							isFirstDraw = false;
+						}
+					}
+					if (x > clientRect.right)
+						break;
+					isFirstUpdate = false;
+					prevX = x;
+					prevY = y;
+					prevValue = count;
+
+					time += 0.25f;
+				}
+				if (line.size() > 1)
+					Polyline(hdc, line.data(), int(line.size()));
+			}
+			posY += GraphHeight;
+		}
+
+		if (m_config.showActiveProcesses && !m_trace.m_activeProcesses.empty())
+		{
+			auto drawBox = [&](u64 start, u64 stop, int height, bool selected, bool inProgress)
+				{
+					//int left = 3 + 6*m_activeFont.height;
+					//int right = rect.left + durationMs/100;
+					int posX = int(m_scrollPosX) + progressRect.left;
+					bool done = stop != ~u64(0);
+					if (!done)
+						stop = playTime;
+					int left = posX + int(float(TimeToS(start)) * scaleX);
+					int right = posX + int(float(TimeToS(stop)) * scaleX) - 1;
+
+					RECT rect;
+					rect.left = left;
+					rect.right = right;
+					rect.top = posY;
+					rect.bottom = posY + height;
+					fillRect(rect, inProgress ? m_processBrushes[selected].inProgress : m_processBrushes[selected].success);
+					return rect;
+				};
+
+			m_activeProcessCountHistory[m_activeProcessCountHistoryIterator++ % sizeof_array(m_activeProcessCountHistory)] = u32(m_trace.m_activeProcesses.size());
+
+			PaintActiveProcesses(posY, clientRect, [&](TraceView::ProcessLocation& processLocation, u32 boxHeight, bool firstWithHeight)
+				{
+					auto& session = m_trace.GetSession(m_traceView, processLocation.sessionIndex);
+					TraceView::Process& process = session.processors[processLocation.processorIndex].processes[processLocation.processIndex];
+
+					bool selected = m_processSelected && m_processSelectedLocation == processLocation;
+
+					if (m_config.showFinishedProcesses)
+					{
+						u32 index = processLocation.processIndex;
+						while (index > 0)
+						{
+							--index;
+							TraceView::Process& process2 = session.processors[processLocation.processorIndex].processes[index];
+							drawBox(process2.start, process2.stop, boxHeight, false, false);
+						}
+					}
+
+					RECT boxRect = drawBox(process.start, process.stop, boxHeight, selected, true);
+
+					u32 v = boxHeight - 1;
+
+					if (v > 4)
+					{
+						u32 fontIndex = Min(v, u32(sizeof_array(m_activeProcessFont) - 1));
+						if (!m_activeProcessFont[fontIndex].handle)
+						{
+							UpdateFont(m_activeProcessFont[fontIndex], fontIndex - 1, false);
+							m_activeProcessFont[fontIndex].offset += 1;
+						}
+						if (firstWithHeight)
+							SetActiveFont(m_activeProcessFont[fontIndex]);
+
+						StringBuffer<> str;
+						auto firstPar = -1;//process.description.find_first_of('(');
+						if (firstPar != -1 && *process.description.rbegin() == ')')
+						{
+							str.Append(process.description.c_str() + firstPar + 1).Resize(str.count-1);
+							str.Append(L"  ").Append(process.description.c_str(), firstPar);
+						}
+						else
+							str.Append(process.description);
+
+						if (process.isRemote)
+							str.Append(L" [").Append(session.name).Append(']');
+						else if (process.type == TraceView::ProcessType_CacheFetchTest || process.type == TraceView::ProcessType_CacheFetchDownload)
+							str.Append(L" [cache]");
+						if (boxRect.left < 0)
+							str.Appendf(L"   %s", TimeToText(playTime - process.start, true).str);
+						//drawIndentedText(str, LogEntryType_Info, 1, true);
+						drawStatusText(str, LogEntryType_Info, Max(int(boxRect.left+1), 1), boxRect.right, false);
+					}
+				});
+		}
+
+		int boxHeight = m_config.boxHeight;
+		int stepY = int(boxHeight) + 2;
+		int processStepY = boxHeight + 1;
+
+		TraceView::WorkRecord selectedWork;
+
+		SessionRec sortedSessions[1024];
+		Populate(sortedSessions, m_traceView, m_config.SortActiveRemoteSessions);
+
+		u32 visibleBoxes = 0;
+
+		TraceView::ProcessLocation processLocation { 0, 0, 0 };
+		for (u64 sessionIt = 0, sessionEnd = m_traceView.sessions.size(); sessionIt != sessionEnd; ++sessionIt)
+		{
+			bool isFirst = sessionIt == 0;
+			auto& session = *sortedSessions[sessionIt].session;
+			bool hasUpdates = !session.updates.empty();
+			if (!isFirst)
+			{
+				if (!hasUpdates && session.processors.empty())
+					continue;
+
+				if (!m_config.showFinishedProcesses && session.disconnectTime != ~u64(0))
+					continue;
+			}
+
+			processLocation.sessionIndex = sortedSessions[sessionIt].index;
+			if (!isFirst)
+				posY += 3;
+
+			if (m_config.showTitleBars)
+			{
+				if (posY + stepY >= progressRect.top && posY <= progressRect.bottom)
+				{
+					SelectObject(hdc, m_separatorPen);
+					MoveToEx(hdc, 0, posY, NULL);
+					LineTo(hdc, clientRect.right, posY);
+
+					StringBuffer<> text;
+					session.GetFullName(text);
+
+					if (hasUpdates && session.disconnectTime == ~u64(0))
+					{
+						u64 ping = session.ping.back();
+						//u64 memAvail = session.updates.back().memAvail;
+						//float cpuLoad = session.updates.back().cpuLoad;
+
+						//text.Appendf(L" - Cpu: %.1f%%", cpuLoad * 100.0f);
+						//if (memAvail)
+						//	text.Appendf(L" Mem: %ls/%ls", BytesToText(session.memTotal - memAvail).str, BytesToText(session.memTotal).str);
+						if (ping)
+							text.Appendf(L" Ping: %ls", TimeToText(ping, false, m_traceView.frequency).str);
+						if (!session.notification.empty())
+							text.Append(L" - ").Append(session.notification);
+					}
+					else if (!isFirst)
+					{
+						text.Append(L" - Disconnected");
+						if (!session.notification.empty())
+							text.Append(L" (").Append(session.notification).Append(')');
+					}
+
+					bool selected = m_sessionSelectedIndex == processLocation.sessionIndex;
+
+					int textBottom = Min(posY + m_sessionStepY, int(progressRect.bottom));
+
+					RECT rect;
+					rect.left = 5;
+					rect.right = clientRect.right;
+					rect.top = posY;
+					rect.bottom = textBottom;
+
+					if (selected)
+						SetBkMode(hdc, OPAQUE);
+
+					SIZE textSize;
+					if (GetTextExtentPoint32W(hdc, text.data, text.count, &textSize))
+						session.fullNameWidth = u32(textSize.cx);
+					ExtTextOutW(hdc, 5, posY+2, ETO_CLIPPED, &rect, text.data, text.count, NULL);
+					if (selected)
+						SetBkMode(hdc, TRANSPARENT);
+				}
+				posY += m_sessionStepY;
+			}
+
+			bool showGraph = m_config.showNetworkStats || m_config.showCpuMemStats || m_config.showDriveStats;
+			if (showGraph && hasUpdates)
+			{
+				if (posY + GraphHeight >= progressRect.top && posY + GraphHeight - 5 < progressRect.bottom)
+				{
+					int posX = int(m_scrollPosX) + progressRect.left;
+					int graphBaseY = posY + GraphHeight - 4;
+					double graphHeight = double(GraphHeight - 2);
+
+					Vector<POINT> line;
+
+					auto drawGraph = [&](auto& values, auto maxValue, double scale, HPEN pen, bool accumulatingValue, int offsetY = 0)
+						{
+							bool loop = true;
+							u32 reconnectIndex = 0;
+							while (loop)
+							{
+								SelectObject(hdc, pen);
+								bool isFirstUpdate = true;
+								bool isFirstDraw = true;
+								decltype(maxValue) prevValue = 0;
+								u64 prevTime = 0;
+								int prevX = 0;
+								int prevPrevX = 0;
+								int prevY = 0;
+
+								u64 i = 0;
+								u64 e = 0;
+
+								if (reconnectIndex)
+									i = session.reconnectIndices[reconnectIndex - 1];
+
+								if (reconnectIndex < session.reconnectIndices.size())
+									e = session.reconnectIndices[reconnectIndex];
+								else
+								{
+									e = Min(session.updates.size(), values.size());
+									loop = false;
+								}
+
+								line.clear();
+								for (; i!=e; ++i)
+								{
+									u64 updateTime = session.updates[i];
+									auto value = values[i];
+									int x = posX + int(float(TimeToS(updateTime)) * scaleX);
+									int y = graphBaseY;
+
+									double duration = TimeToS(updateTime - prevTime);
+									if (updateTime == 0)
+										isFirstUpdate = true;
+
+									if (accumulatingValue)
+									{
+										double invScaleY = duration * scale;
+										if (invScaleY != 0 && prevValue != 0)
+											y -= int(double(value - prevValue) / invScaleY) + offsetY;
+									}
+									else
+									{
+										double div = 1.0f;
+										if (maxValue != 0)
+											div = double(maxValue);
+										y -= int(double((double)maxValue - (double)value*scale)*graphHeight/div) + offsetY;
+									}
+
+									if (prevValue <= 0)
+										prevY = y;
+
+									if (x > clientRect.left)
+									{
+										if (isFirstUpdate)
+										{
+											isFirstDraw = false;
+										}
+										else
+										{
+											if (isFirstDraw)
+												line.push_back({prevPrevX, prevY});
+											line.push_back({prevX, y});
+											isFirstDraw = false;
+										}
+									}
+									if (prevX > clientRect.right)
+										break;
+									isFirstUpdate = false;
+									prevPrevX = prevX;
+									prevX = x;
+									prevY = y;
+									prevValue = value;
+									prevTime = updateTime;
+								}
+								if (line.size() > 1)
+									Polyline(hdc, line.data(), int(line.size()));
+								++reconnectIndex;
+							}
+						};
+
+					if (m_config.showNetworkStats && session.highestSendPerS != 0 && session.highestRecvPerS != 0)
+					{
+						drawGraph(session.networkSend, 100000000000000ull, double(session.highestSendPerS) / graphHeight, m_sendPen, true);
+						drawGraph(session.networkRecv, 100000000000000ull, double(session.highestRecvPerS) / graphHeight, m_recvPen, true, 1);
+					}
+
+					if (m_config.showCpuMemStats)
+					{
+						drawGraph(session.cpuLoad, 0.0f, -1.0, m_cpuPen, false);
+						drawGraph(session.memAvail, session.memTotal, 1.0, m_memPen, false);
+					}
+
+					if (m_config.showDriveStats)
+						for (auto& pair : session.drives)
+							if (pair.second.busyHighest)
+								drawGraph(pair.second.busyPercent, 0, -0.01, m_drivePen, false);
+
+				}
+				posY += GraphHeight;
+			}
+
+			if (m_config.showDetailedData)
+			{
+				auto drawText = [&](const StringBufferBase& text, RECT& rect, u32* outWidth)
+					{
+						if (rect.top > progressRect.bottom)
+							return;
+						bool selected = m_fetchedFilesSelected == processLocation.sessionIndex && text.StartsWith(TC("Fetched Files"));
+						if (selected)
+							SetBkMode(hdc, OPAQUE);
+						if (rect.bottom > progressRect.bottom)
+						{
+							rect.bottom = progressRect.bottom;
+							ExtTextOutW(hdc, rect.left, rect.top, ETO_CLIPPED, &rect, text.data, text.count, NULL);
+						}
+						else
+							ExtTextOutW(hdc, rect.left, rect.top, 0, NULL, text.data, text.count, NULL);
+						if (selected)
+							SetBkMode(hdc, TRANSPARENT);
+						if (outWidth)
+						{
+							SIZE s;
+							GetTextExtentPoint32W(hdc, text.data, text.count, &s);
+							*outWidth = u32(s.cx);
+						}
+					};
+				bool isRemote = processLocation.sessionIndex != 0;
+				PaintDetailedStats(posY, progressRect, session, isRemote, playTime, drawText);
+			}
+
+			SetActiveFont(m_processFont);
+
+			bool shouldDrawText = m_processFont.height > 4;
+
+			if (m_config.showProcessBars)
+			{
+				processLocation.processorIndex = 0;
+				for (auto& processor : session.processors)
+				{
+					bool drawProcessorIndex = m_config.showFinishedProcesses;
+
+					if (posY + m_sessionStepY >= progressRect.top && posY < progressRect.bottom)
+					{
+						int barHeight = boxHeight;
+						int textOffsetY = 0;
+						if (posY + boxHeight > progressRect.bottom)
+						{
+							int newBarHeight = Min(barHeight, int(progressRect.bottom - posY));
+							textOffsetY = int(barHeight - newBarHeight);
+							barHeight = newBarHeight;
+						}
+
+						const int textHeight = barHeight;
+						const int rectBottom = posY + textHeight;
+						const int offsetY = (textHeight - m_processFont.height + textOffsetY) / 2;
+
+						processLocation.processIndex = 0;
+						int posX = int(m_scrollPosX) + progressRect.left;
+						for (auto& process : processor.processes)
+						{
+							int left = posX + int(float(TimeToS(process.start)) * scaleX);
+
+							auto pig = MakeGuard([&]() { ++processLocation.processIndex; });
+
+							if (left >= progressRect.right)
+								continue;
+
+							u64 stop = process.stop;
+							bool done = stop != ~u64(0);
+							if (!done)
+								stop = playTime;
+							else if (!m_config.showFinishedProcesses)
+								continue;
+
+							drawProcessorIndex = true;
+
+							RECT rect;
+							rect.left = left;
+							rect.right = Max(posX + int(float(TimeToS(stop)) * scaleX) - 1, left + 1);
+							rect.top = posY;
+							rect.bottom = rectBottom - 1;
+
+							if (rect.right <= progressRect.left)
+								continue;
+
+							if (!m_filterString.empty() && !Contains(process.description.c_str(), m_filterString.c_str()) && !Contains(process.breadcrumbs.c_str(), m_filterString.c_str()))
+								continue;
+							++visibleBoxes;
+
+							bool selected = m_processSelected && m_processSelectedLocation == processLocation;
+							if (selected)
+								process.bitmapDirty = true;
+
+							--rect.top;
+							PaintProcessRect(process, hdc, rect, progressRect, selected, false, lastSelectedBrush);
+							++rect.top;
+
+							int processWidth = rect.right - rect.left;
+							if (shouldDrawText && m_config.ShowProcessText && processWidth > 3)
+							{
+								if (!process.bitmap || process.bitmapDirty)
+								{
+									if (!process.bitmap)
+									{
+										if (m_lastBitmapOffset == BitmapCacheHeight)
+										{
+											if (m_lastBitmap)
+												m_textBitmaps.push_back(m_lastBitmap);
+
+											m_lastBitmapOffset = 0;
+											m_lastBitmap = CreateCompatibleBitmap(hdc, 256, BitmapCacheHeight);
+										}
+										process.bitmap = m_lastBitmap;
+										process.bitmapOffset = m_lastBitmapOffset;
+										m_lastBitmapOffset += m_processFont.height;
+									}
+									if (lastSelectedBitmap != process.bitmap)
+									{
+										SelectObject(textDC, process.bitmap);
+										lastSelectedBitmap = process.bitmap;
+									}
+
+									RECT rect2{ 0, int(process.bitmapOffset), 256, int(process.bitmapOffset) + m_processFont.height };
+									RECT rect3{ 0, int(process.bitmapOffset), processWidth, int(process.bitmapOffset) + m_processFont.height };
+									if (!done)
+										rect3.right = 256;
+
+									PaintProcessRect(process, textDC, rect3, rect2, selected, true, textLastSelectedBrush);
+
+									rect2.left += 3; // Move in text a bit
+
+									int textY = rect2.top + m_processFont.offset;
+
+									bool dropShadow = m_config.DarkMode;
+									if (dropShadow)
+									{
+										SetTextColor(textDC, RGB(5, 60, 5));
+										++rect2.left;
+										++textY;
+										ExtTextOutW(textDC, rect2.left, textY, ETO_CLIPPED, &rect2, process.description.c_str(), int(process.description.size()), NULL);
+										--rect2.left;
+										--textY;
+									}
+									SetTextColor(textDC, m_textColor);
+									ExtTextOutW(textDC, rect2.left, textY, ETO_CLIPPED, &rect2, process.description.c_str(), int(process.description.size()), NULL);
+
+									if (!selected)
+										process.bitmapDirty = false;
+								}
+
+								if (lastSelectedBitmap != process.bitmap)
+								{
+									SelectObject(textDC, process.bitmap);
+									lastSelectedBitmap = process.bitmap;
+								}
+
+								int width = Min(processWidth, 256);
+								int bitmapOffsetY = process.bitmapOffset;
+								int bltOffsetY = offsetY;
+								if (bltOffsetY < 0)
+								{
+									bitmapOffsetY -= bltOffsetY;
+									bltOffsetY = 0;
+								}
+								int height = Min(textHeight, m_processFont.height);
+								if (bltOffsetY + height > textHeight)
+									height = textHeight - bltOffsetY;
+
+								if (left > -256 && height >= 0)
+								{
+									int bitmapOffsetX = rect.left - left;
+
+									if (left < progressRect.left)
+									{
+										int diff = progressRect.left - left;
+										rect.left = progressRect.left;
+										width -= diff;
+										bitmapOffsetX += diff;
+									}
+									BitBlt(hdc, rect.left, rect.top + bltOffsetY, width, height, textDC, bitmapOffsetX, bitmapOffsetY, SRCCOPY);
+								}
+							}
+						}
+
+						if (drawProcessorIndex)
+						{
+							RECT rect;
+							rect.left = 5;
+							rect.right = progressRect.left - 2;
+							rect.top = posY;
+							rect.bottom = rectBottom;
+
+							StringBuffer<> buf;
+							buf.AppendValue(u64(processLocation.processorIndex) + 1);
+							ExtTextOutW(hdc, 5, posY + offsetY, ETO_CLIPPED, &rect, buf.data, buf.count, NULL);
+						}
+					}
+
+					lastStop = Max(lastStop, processor.processes.rbegin()->stop);
+
+					++processLocation.processorIndex;
+
+					if (drawProcessorIndex)
+						posY += processStepY;
+				}
+			}
+			else
+			{
+				for (auto& processor : session.processors)
+					if (!processor.processes.empty())
+						lastStop = Max(lastStop, processor.processes.rbegin()->stop);
+			}
+
+			if (m_config.showWorkers && isFirst)
+			{
+				u32 trackIndex = 0;
+				for (auto& workTrack : m_traceView.workTracks)
+				{
+					if (posY + m_sessionStepY >= progressRect.top && posY <= progressRect.bottom)
+					{
+						int textOffsetY = 0;
+						int barHeight = boxHeight;
+						if (posY + int(boxHeight) > progressRect.bottom)
+						{
+							int newBarHeight = Min(barHeight, int(progressRect.bottom - posY));
+							textOffsetY = barHeight - newBarHeight;
+							barHeight = newBarHeight;
+						}
+
+						const int textHeight = barHeight;
+						const int rectBottom = posY + textHeight;
+						const int offsetY = (textHeight - m_processFont.height + textOffsetY) / 2;
+
+						if (shouldDrawText)
+						{
+							RECT rect;
+							rect.left = 5;
+							rect.right = progressRect.left - 5;
+							rect.top = posY;
+							rect.bottom = rectBottom;
+
+							StringBuffer<> buf;
+							buf.AppendValue(u64(trackIndex) + 1);
+							ExtTextOutW(hdc, 5, posY + offsetY, ETO_CLIPPED, &rect, buf.data, buf.count, NULL);
+						}
+
+						int lastDrawnRight = 0;
+
+						u32 workIndex = 0;
+						int posX = int(m_scrollPosX) + progressRect.left;
+						for (auto& work : workTrack.records)
+						{
+							auto inc = MakeGuard([&](){ ++workIndex; });
+
+							if (work.start == work.stop)
+								continue;
+
+							float startTime = TimeToS(work.start);
+
+							int left = posX + int(startTime * scaleX);
+
+							if (left >= progressRect.right)
+								continue;
+
+							if (!m_filterString.empty())
+							{
+								bool keep = Contains(work.description, m_filterString.c_str());
+								if (!keep)
+									for (auto& en : work.entries)
+										keep |= Contains(en.text, m_filterString.c_str());
+								if (!keep)
+									continue;
+							}
+
+							HBRUSH brush;
+
+							u64 stop = work.stop;
+
+							float stopTime = TimeToS(stop);
+
+							RECT rect;
+							rect.left = left;
+							rect.right = posX + int(stopTime * scaleX) - 1;
+							rect.top = posY;
+							rect.bottom = rectBottom - 1;
+
+							if (rect.right <= progressRect.left)
+								continue;
+
+							++visibleBoxes;
+
+							rect.right = Max(int(rect.right), left + 1);
+
+							Color color = work.color;
+							bool selected = m_workSelected && m_workTrack == trackIndex && m_workIndex == workIndex;
+							if (selected)
+							{
+								selectedWork = work;
+								work.bitmapDirty = true;
+								auto c = (u8*)&color;
+								for (u32 j=0; j!=3; ++j)
+									c[j] = u8(Min(int(c[j]) + 40, 255));
+							}
+							else
+							{
+								if (rect.left + 1 == rect.right)
+									if (lastDrawnRight == rect.right)
+										continue;
+							}
+
+							lastDrawnRight = rect.right;
+
+							bool done = stop != ~u64(0);
+							if (done)
+							{
+								auto insres = m_coloredBrushes.try_emplace(color);
+								if (insres.second)
+								{
+									auto c = (u8*)&color;
+									insres.first->second = CreateSolidBrush(RGB(c[2], c[1], c[0]));
+								}
+								brush = insres.first->second;
+							}
+							else
+							{
+								//stop = playTime;
+								brush = m_processBrushes[0].inProgress;
+							}
+
+							--rect.top;
+
+							auto clampRect = [&](RECT& r) { r.left = Min(Max(r.left, progressRect.left), progressRect.right); r.right = Max(Min(r.right, progressRect.right), progressRect.left); };
+
+							clampRect(rect);
+							fillRect(rect, brush);
+							++rect.top;
+
+							int processWidth = rect.right - rect.left;
+							if (shouldDrawText && m_config.ShowProcessText && processWidth > 3)
+							{
+								if (!work.bitmap || work.bitmapDirty)
+								{
+									if (!work.bitmap)
+									{
+										if (m_lastBitmapOffset == BitmapCacheHeight)
+										{
+											if (m_lastBitmap)
+												m_textBitmaps.push_back(m_lastBitmap);
+
+											m_lastBitmapOffset = 0;
+											m_lastBitmap = CreateCompatibleBitmap(hdc, 256, BitmapCacheHeight);
+										}
+										work.bitmap = m_lastBitmap;
+										work.bitmapOffset = m_lastBitmapOffset;
+										m_lastBitmapOffset += m_processFont.height;
+									}
+									if (lastSelectedBitmap != work.bitmap)
+									{
+										SelectObject(textDC, work.bitmap);
+										lastSelectedBitmap = work.bitmap;
+									}
+
+									RECT rect2{ 0,int(work.bitmapOffset), 256, int(work.bitmapOffset) + m_processFont.height };
+
+									SelectObject(textDC, brush);
+									PatBlt(textDC, rect2.left, rect2.top, rect2.right - rect2.left, rect2.bottom - rect2.top, PATCOPY);
+									ExtTextOutW(textDC, rect2.left, rect2.top, ETO_CLIPPED, &rect2, work.description, int(wcslen(work.description)), NULL);
+									if (!selected)
+										work.bitmapDirty = false;
+								}
+								else if (lastSelectedBitmap != work.bitmap)
+								{
+									SelectObject(textDC, work.bitmap);
+									lastSelectedBitmap = work.bitmap;
+								}
+
+								int width = Min(processWidth, 256);
+								int bitmapOffsetY = work.bitmapOffset;
+								int bltOffsetY = offsetY;
+								if (bltOffsetY < 0)
+								{
+									bitmapOffsetY -= bltOffsetY;
+									bltOffsetY = 0;
+								}
+								int height = Min(textHeight, m_processFont.height);
+								if (bltOffsetY + height > textHeight)
+									height = textHeight - bltOffsetY;
+
+								if (left > -256 && height >= 0)
+								{
+									int bitmapOffsetX = rect.left - left;
+
+									if (left < progressRect.left)
+									{
+										int diff = progressRect.left - left;
+										rect.left = progressRect.left;
+										width -= diff;
+										bitmapOffsetX += diff;
+									}
+									BitBlt(hdc, rect.left, rect.top + bltOffsetY, width, height, textDC, bitmapOffsetX, bitmapOffsetY, SRCCOPY);
+								}
+							}
+						}
+					}
+					++trackIndex;
+					posY += processStepY;
+				}
+			}
+
+			SetActiveFont(m_defaultFont);
+		}
+
+		if (true)
+		{
+			for (auto& kv : m_traceView.tables)
+			{
+				auto& table = kv.second;
+				u32 tableOffsetX = m_activeFont.height*40;
+				u32 tableOffsetY = 30;
+
+				u32 rowOffset = m_activeFont.height + 2;
+				u32 columnOffset = table.rowCharWidth * m_activeFont.height;
+
+				u32 columnIndex = 0;
+				u32 rowIndex = 0;
+				for (const tchar* item : table.items)
+				{
+					u32 y = tableOffsetY + rowIndex*rowOffset;
+					u32 x = tableOffsetX + columnIndex*columnOffset;
+					if (int(x) >= progressRect.right)
+						break;
+
+					ExtTextOutW(hdc, x, y, ETO_CLIPPED, nullptr, item, TStrlen(item), NULL);
+
+					++rowIndex;
+					if (rowIndex > table.rowMaxCount || int(y) >= progressRect.bottom)
+					{
+						rowIndex = 0;
+						++columnIndex;
+					}
+				}
+			}
+		}
+
+		SelectObject(textDC, oldBmp);
+		DeleteObject(nullBmp);
+		DeleteDC(textDC);
+
+		m_contentWidth = m_progressRectLeft + Max(0, int(TimeToS((lastStop != 0 && lastStop != ~u64(0)) ? lastStop : playTime) * scaleX));
+
+
+		m_contentHeight = posY - int(m_scrollPosY) + stepY + 14;
+
+		float timelineSelected = m_timelineSelected;
+
+		if (m_config.showTimeline && !m_traceView.sessions.empty())
+			PaintTimeline(hdc, clientRect);
+
+		if (!m_filterString.empty())
+		{
+			StringBuffer<> str[2];
+			str[0].Append(TCV("Box Filter: ")).Append(m_filterString);
+			str[1].Appendf(TC("Box Count: %u"), visibleBoxes);
+
+			SetActiveFont(m_popupFont);
+			DrawCenteredText(str, 2, true);
+		}
+
+		if (m_config.showCursorLine && m_mouseOverWindow)
+		{
+			float timeScale = (m_config.horizontalScaleValue * m_zoomValue)*50.0f;
+			float startOffset = -(m_scrollPosX / timeScale);
+			POINT pos;
+			GetCursorPos(&pos);
+			ScreenToClient(m_hwnd, &pos);
+			timelineSelected = startOffset + float(pos.x - m_progressRectLeft) / timeScale;
+		}
+
+		if (timelineSelected != 0)
+		{
+			int posX = int(m_scrollPosX) + progressRect.left;
+			int left = posX + int(timelineSelected * scaleX);
+			int timelineTop = GetTimelineTop(clientRect);
+
+			// TODO: Draw line up
+			MoveToEx(hdc, left, 2, NULL);
+			LineTo(hdc, left, timelineTop);
+
+			if (timelineSelected >= 0)
+			{
+				StringBuffer<> b;
+				u32 milliseconds = u32(timelineSelected * 1000.0f);
+				u32 seconds = milliseconds / 1000;
+				milliseconds -= seconds * 1000;
+				u32 minutes = seconds / 60;
+				seconds -= minutes * 60;
+				u32 hours = minutes / 60;
+				minutes -= hours * 60;
+				if (hours)
+				{
+					b.AppendValue(hours).Append('h');
+					if (minutes < 10)
+						b.Append('0');
+				}
+				if (minutes || hours)
+				{
+					b.AppendValue(minutes).Append('m');
+					if (seconds < 10)
+						b.Append('0');
+				}
+				b.AppendValue(seconds).Append('.');
+				if (milliseconds < 100)
+					b.Append('0');
+				if (milliseconds < 10)
+					b.Append('0');
+				b.AppendValue(milliseconds);
+
+				StringBuffer<128> wt = GetWorldTime(timelineSelected);
+
+				SetActiveFont(m_popupFont);
+				DrawTextLogger logger(m_hwnd, hdc, m_popupFont.height, m_tooltipBackgroundBrush);
+
+				logger.Info(L"%s (%s)", b.data, wt.data);
+				logger.DrawAtPos(left + 4, timelineTop - 20);
+			}
+		}
+
+		{
+			int boxSide = 8;
+			int boxStride = boxSide + 2;
+			int top = 5;
+			int bottom = top + boxSide;
+			int left = progressRect.right - 7 - boxSide;
+			int right = progressRect.right - 7;
+			bool* values = &m_config.showProgress;
+			for (int i = VisualizerFlag_Count - 1; i >= 0; --i)
+			{
+				SelectObject(hdc, m_buttonSelected == u32(i) ? m_textPen : m_checkboxPen);
+				SelectObject(hdc, GetStockObject(NULL_BRUSH));
+				Rectangle(hdc, left, top, right, bottom);
+
+				if (values[i])
+				{
+					MoveToEx(hdc, left + 2, top + 2, NULL);
+					LineTo(hdc, right - 2, bottom - 2);
+					MoveToEx(hdc, right - 3, top + 2, NULL);
+					LineTo(hdc, left + 1, bottom - 2);
+				}
+
+				left -= boxStride;
+				right -= boxStride;
+			}
+
+			top -= 2;
+
+			auto drawText = [&](const tchar* text, COLORREF color)
+				{
+					SetTextColor(hdc, color);
+					RECT r{left, top, left+200, top+200};
+					auto strLen = TStrlen(text);
+					DrawTextW(hdc, text, strLen, &r, DT_SINGLELINE|DT_NOCLIP|DT_NOPREFIX|DT_CALCRECT);
+					left -= (r.right - r.left) + 5;
+					r.left = left;
+					DrawTextW(hdc, text, strLen, &r, DT_SINGLELINE|DT_NOCLIP|DT_NOPREFIX);
+				};
+
+			if (m_config.showDriveStats)
+			{
+				SetActiveFont(m_defaultFont);
+				drawText(L"DRV", m_driveColor);
+			}
+
+			if (m_config.showNetworkStats)
+			{
+				SetActiveFont(m_defaultFont);
+				drawText(L"SND", m_sendColor);
+				drawText(L"RCV", m_recvColor);
+			}
+
+			if (m_config.showCpuMemStats)
+			{
+				SetActiveFont(m_defaultFont);
+				drawText(L"CPU", m_cpuColor);
+				drawText(L"MEM", m_memColor);
+			}
+
+			SetTextColor(hdc, m_textColor);
+		}
+
+		if (m_processSelected)
+		{
+			const TraceView::Process& process = m_traceView.GetProcess(m_processSelectedLocation);
+			u64 duration = 0;
+
+			Vector<TString> logLines;
+			u32 maxCharCount = 50u;
+
+			bool hasExited = process.stop != ~u64(0);
+			if (hasExited)
+			{
+				duration = process.stop - process.start;
+
+				if (!process.logLines.empty())
+				{
+					u32 lineMaxCount = 0;
+					for (auto& line : process.logLines)
+					{
+						u32 offset = 0;
+						u32 left = u32(line.text.size());
+						while (left)
+						{
+							u32 toCopy = Min(left, maxCharCount);
+							lineMaxCount = Max(lineMaxCount, toCopy);
+							logLines.push_back(line.text.substr(offset, toCopy));
+							offset += toCopy;
+							left -= toCopy;
+						}
+					}
+				}
+			}
+			else
+			{
+				duration = playTime - process.start;
+			}
+
+			SetActiveFont(m_popupFont);
+			DrawTextLogger logger(m_hwnd, hdc, m_popupFont.height, m_tooltipBackgroundBrush);
+
+			logger.AddTextOffset(-10); // Remove spaces in the front
+			logger.AddWidth(3);
+
+			logger.AddSpace(2);
+			logger.Info(L"  %ls", process.description.c_str());
+			if (process.type != TraceView::ProcessType_Task)
+			{
+				logger.Info(L"  Host:        %ls", m_processSelectedLocation.sessionIndex == 0 ? TC("local") : m_traceView.GetSession(m_processSelectedLocation).name.c_str());
+				logger.Info(L"  ProcessId:  %6u", process.id);
+			}
+			logger.Info(L"  Start:     %7ls (%ls)", TimeToText(process.start, true).str, GetWorldTime(process.start).data);
+			logger.Info(L"  Duration:  %7ls", TimeToText(duration, true).str);
+			if (!process.returnedReason.empty())
+				logger.Info(L"  Returned: %7s", process.returnedReason.data());
+			if (hasExited && process.exitCode != 0)
+			{
+				if (process.exitCode == ProcessCancelExitCode)
+					logger.Info(L"  ExitCode: Cancelled");
+				else if (process.exitCode >= 0xC0000000u)
+					logger.Info(L"  ExitCode: 0x%x", process.exitCode);
+				else
+					logger.Info(L"  ExitCode: %7u", process.exitCode);
+			}
+
+			const auto& breadcrumbs = process.breadcrumbs;
+			if (!breadcrumbs.empty())
+			{
+				constexpr TString::size_type maxLineLen = 50;
+				logger.Info(L"");
+				logger.Info(L"  ------------ Breadcrumbs ------------");
+				for (TString::size_type lineStart = 0, lineEnd = 0; lineEnd < breadcrumbs.size(); lineStart = lineEnd + 1)
+				{
+					// Log each individual line
+					lineEnd = breadcrumbs.find(L'\n', lineStart);
+					TString line = (lineEnd == TString::npos ? breadcrumbs.substr(lineStart) : breadcrumbs.substr(lineStart, lineEnd - lineStart));
+
+					// Break each line down into smaller section if they are longer than the maximum allowed length
+					if (line.size() > maxLineLen)
+					{
+						for (TString::size_type sectionStart = 0, sectionEnd = 0; sectionStart < line.size(); sectionStart = sectionEnd)
+						{
+							const TString::size_type maxSectionLen = sectionStart == 0 ? maxLineLen : maxLineLen - 2;
+							sectionEnd = std::min<TString::size_type>(sectionEnd + maxSectionLen, line.size());
+							TString section = (sectionStart == 0 ? L"  " : L"    ") + line.substr(sectionStart, sectionEnd - sectionStart);
+							logger.Info(section.c_str());
+						}
+					}
+					else
+					{
+						line = L"  " + line;
+						logger.Info(line.c_str());
+					}
+				}
+			}
+
+			if (process.stop != ~u64(0) && !process.stats.empty())
+			{
+				BinaryReader reader(process.stats.data(), 0, process.stats.size());
+				ProcessStats processStats;
+				SessionStats sessionStats;
+				StorageStats storageStats;
+				KernelStats kernelStats;
+				CacheFetchStats cacheStats;
+
+				if (process.type == TraceView::ProcessType_CacheFetchTest || process.type == TraceView::ProcessType_CacheFetchDownload)
+				{
+					if (!process.returnedReason.empty())
+						logger.Info(L"  Cache:       Miss");
+					else
+						logger.Info(L"  Cache:        Hit");
+					cacheStats.Read(reader, m_traceView.version);
+					if (reader.GetLeft())
+					{
+						storageStats.Read(reader, m_traceView.version);
+						kernelStats.Read(reader, m_traceView.version);
+					}
+				}
+				else
+				{
+					processStats.Read(reader, m_traceView.version);
+
+					if (reader.GetLeft())
+					{
+						if (process.isRemote || (m_traceView.version >= 36 && !process.isReuse))
+							sessionStats.Read(reader, m_traceView.version);
+						storageStats.Read(reader, m_traceView.version);
+						kernelStats.Read(reader, m_traceView.version);
+					}
+				}
+
+				if (processStats.hostTotalTime)
+				{
+					logger.Info(L"");
+					logger.Info(L"  ----------- Detours stats -----------");
+					processStats.Print(logger, m_traceView.frequency);
+				}
+				else if (processStats.peakMemory)
+				{
+					logger.Info(L"");
+					logger.Info(L"  ----------- Process stats -----------");
+					processStats.Print(logger, m_traceView.frequency);
+				}
+
+				if (!sessionStats.IsEmpty())
+				{
+					logger.Info(L"");
+					logger.Info(L"  ----------- Session stats -----------");
+					sessionStats.Print(logger, m_traceView.frequency);
+				}
+
+				if (!cacheStats.IsEmpty())
+				{
+					logger.Info(L"");
+					logger.Info(L"  ------------ Cache stats ------------");
+					cacheStats.Print(logger, m_traceView.frequency);
+				}
+
+				if (!storageStats.IsEmpty())
+				{
+					logger.Info(L"");
+					logger.Info(L"  ----------- Storage stats -----------");
+					storageStats.Print(logger, m_traceView.frequency);
+				}
+
+				if (!kernelStats.IsEmpty())
+				{
+					logger.Info(L"");
+					logger.Info(L"  ----------- Kernel stats ------------");
+					kernelStats.Print(logger, false, m_traceView.frequency);
+				}
+
+				PrintCacheWriteStats(logger, process.id);
+
+				if (!logLines.empty())
+				{
+					logger.Info(L"");
+					logger.Info(L"  ---------------- Log ----------------");
+					logger.AddTextOffset(14);
+					for (auto& line : logLines)
+						logger.Log(LogEntryType_Info, line);
+				}
+			}
+			logger.AddSpace(3);
+			logger.DrawAtCursor();
+		}
+		else if (m_workSelected && selectedWork.description)
+		{
+			u64 duration;
+			if (selectedWork.stop != ~u64(0))
+				duration = selectedWork.stop - selectedWork.start;
+			else
+				duration = playTime - selectedWork.start;
+
+			SetActiveFont(m_popupFont);
+			DrawTextLogger logger(m_hwnd, hdc, m_popupFont.height, m_tooltipBackgroundBrush);
+
+			logger.AddSpace();
+			logger.Info(L"  %ls", selectedWork.description);
+			logger.Info(L"  Start:     %ls (%ls)", TimeToText(selectedWork.start, true).str, GetWorldTime(selectedWork.start).data);
+			logger.Info(L"  Duration:  %ls", TimeToText(duration, true).str);
+			if (!selectedWork.entries.empty())
+				logger.AddSpace();
+			for (u32 i=0,e=u32(selectedWork.entries.size()); i!=e; ++i)
+			{
+				auto& entry = selectedWork.entries[i];
+				u64 time = 0;
+				if (i == 0)
+				{
+					u64 start = entry.startTime ? entry.startTime : entry.time;
+					if (TimeToMs(start - selectedWork.start) > 1)
+						logger.Info(L"   Start (%ls)", TimeToText(start - selectedWork.start).str);
+				}
+
+				if (entry.startTime)
+					time = entry.time - entry.startTime;
+
+				if (!time)
+					for (u32 j=i+1;j<e;++j)
+					{
+						auto& next = selectedWork.entries[j];
+						if (next.startTime)
+							continue;
+						time = next.time - entry.time;
+						break;
+					}
+
+				if (!time && selectedWork.stop != ~0u)
+					time = selectedWork.stop - entry.time;
+				if (entry.count == 1)
+					logger.Info(L"%s  %ls (%ls)", (entry.startTime ? L" " : L""), entry.text, TimeToText(time).str);
+				else
+					logger.Info(L"%s  %ls (%ls %u)", (entry.startTime ? L" " : L""), entry.text, TimeToText(time).str, entry.count);
+			}
+			logger.AddSpace();
+			logger.DrawAtCursor();
+		}
+		else if (m_sessionSelectedIndex != ~0u)
+		{
+			SetActiveFont(m_popupFont);
+			DrawTextLogger logger(m_hwnd, hdc, m_popupFont.height, m_tooltipBackgroundBrush);
+			logger.AddWidth(3);
+			logger.AddSpace(2);
+
+			auto& session = m_traceView.sessions[m_sessionSelectedIndex];
+			Vector<TString>& summary = session.summary;
+			if (summary.empty())
+			{
+				if (m_traceView.finished)
+					logger.Info(L" Session summary was never received");
+				else
+					logger.Info(L" Session summary not available until session is done");
+			}
+
+			logger.AddTextOffset(-10); // Remove spaces in the front
+			for (auto& line : summary)
+				logger.Log(LogEntryType_Info, line);
+			logger.AddTextOffset(0);
+
+			if (m_sessionSelectedIndex == 0)
+			{
+				for (auto& cacheSummary : m_traceView.cacheSummaries)
+					for (auto& line : cacheSummary.lines)
+						logger.Log(LogEntryType_Info, line);
+			}
+
+			for (auto& pair : session.drives)
+			{
+				auto& d = pair.second;
+				logger.Info(L"  %c: Read: %s (%u) Write: %s (%u)", pair.first, BytesToText(d.totalReadBytes).str, d.totalReadCount, BytesToText(d.totalWriteBytes).str, d.totalWriteCount);
+			}
+
+			logger.AddSpace(3);
+			logger.DrawAtCursor();
+		}
+		else if (m_statsSelected)
+		{
+			SetActiveFont(m_popupFont);
+			DrawTextLogger logger(m_hwnd, hdc, m_popupFont.height, m_tooltipBackgroundBrush);
+			logger.AddSpace(3);
+			logger.SetColor(m_cpuColor).Info(L"  Cpu: %.1f%%", m_stats.cpuLoad * 100.0f);
+			logger.SetColor(m_memColor).Info(L"  Mem: %ls/%ls", BytesToText(m_stats.memTotal - m_stats.memAvail).str, BytesToText(m_stats.memTotal).str);
+			if (m_stats.recvBytes || m_stats.sendBytes)
+			{
+				logger.SetColor(m_recvColor).Info(L"  Recv: %lsit/s", BytesToText(m_stats.recvBytesPerSecond*8).str);
+				logger.SetColor(m_sendColor).Info(L"  Send: %lsit/s", BytesToText(m_stats.sendBytesPerSecond*8).str);
+			}
+			logger.SetColor(m_textColor);
+
+			if (m_stats.ping)
+				logger.Info(L"  Ping: %ls", TimeToText(m_stats.ping, false, m_traceView.frequency).str);
+			if (m_stats.connectionCount)
+				logger.Info(L"  Conn: %u", u32(m_stats.connectionCount));
+
+			for (auto& pair : m_stats.drives)
+				logger.SetColor(m_driveColor).Info(L"  %c: %u%% R:%s/s W:%s/s", pair.first, pair.second.busyPercent, BytesToText(pair.second.readPerSecond).str, BytesToText(pair.second.writePerSecond).str);
+
+			logger.AddSpace(3);
+			logger.DrawAtCursor();
+		}
+		else if (m_activeProcessGraphSelected)
+		{
+			SetActiveFont(m_popupFont);
+			DrawTextLogger logger(m_hwnd, hdc, m_popupFont.height, m_tooltipBackgroundBrush);
+			logger.AddSpace(3);
+			logger.AddSpace(3);
+			logger.SetColor(m_activeProcColor).Info(L"  Active Processes: %u", m_activeProcessCount);
+			logger.DrawAtCursor();
+		}
+		else if (m_buttonSelected != ~0u)
+		{
+			const wchar_t* tooltip[] =
+			{
+#define UBA_VISUALIZER_FLAG(name, defaultValue, desc) desc,
+				UBA_VISUALIZER_FLAGS1
+#undef UBA_VISUALIZER_FLAG
+			};
+
+			SetActiveFont(m_popupFont);
+			DrawTextLogger logger(m_hwnd, hdc, m_popupFont.height, m_tooltipBackgroundBrush);
+			logger.Info(L"%ls %ls", L"Show", tooltip[m_buttonSelected]);
+			logger.DrawAtCursor();
+		}
+		else if (m_fetchedFilesSelected != ~0u)
+		{
+			auto& session = m_traceView.sessions[m_fetchedFilesSelected];
+			auto& fetchedFiles = session.fetchedFiles;
+			if (!fetchedFiles.empty() && !fetchedFiles[0].hint.empty())
+			{
+				/*
+				int colWidth = 500;
+				int width = colWidth * 2;
+				int height = Min(int(clientRect.bottom), int(fetchedFiles.size() * m_popupFont.height));
+
+				SetActiveFont(m_defaultFont);
+				DrawTextLogger logger(m_hwnd, hdc, r, m_font.height, m_tooltipBackgroundBrush);
+				for (auto& f : fetchedFiles)
+				{
+				if (f.hint == TC("KnownInput"))
+				continue;
+				if (logger.rect.top >= r.bottom - m_font.height)
+				{
+				if (logger.rect.left + colWidth >= r.right)
+				{
+				logger.Info(L"...");
+				break;
+				}
+				logger.rect.top = r.top;
+				logger.rect.left += colWidth;
+				}
+				logger.Info(L"%s", f.hint.c_str());
+				}
+				logger.DrawAtCursor();
+				*/
+			}
+		}
+	}
+
+	u64 ConvertTime(const TraceView& view, u64 time);
+
+	void VisualizerWin::PaintActiveProcesses(int& posY, const RECT& clientRect, const Function<void(TraceView::ProcessLocation&, u32, bool)>& drawProcess)
+	{
+		SetActiveFont(m_processFont);
+		int startPosY = posY;
+
+		Map<u64, TraceView::ProcessLocation*> activeProcesses;
+		u32 remoteCount = 0;
+		for (auto& kv : m_trace.m_activeProcesses)
+		{
+			auto& active = kv.second;
+			TraceView::Process& process = m_trace.GetSession(m_traceView, active.sessionIndex).processors[active.processorIndex].processes[active.processIndex];
+			u64 start = process.start; //~0llu - process.start;
+			activeProcesses.try_emplace(start, &active);
+			if (process.isRemote)
+				++remoteCount;
+		}
+
+		u32 maxHeight = clientRect.bottom;
+
+		bool fillHeight = ActiveProcessesShouldFillHeight();
+		if (fillHeight)
+		{
+			maxHeight = clientRect.bottom - posY;
+			if (m_config.showTimeline)
+				maxHeight -= GetTimelineHeight();
+		}
+		else
+		{
+			u32 maxHeight2 = m_config.maxActiveVisible * (m_activeFont.height + 2);
+			maxHeight = Min(maxHeight, maxHeight2);
+		}
+
+		u32 maxSize = Min(m_config.maxActiveProcessHeight, 32u);
+		u32 maxSizeMinusOne = maxSize - 1;
+
+		u32 counts[128] = { 0 };
+
+		u32 highestHistoryCount = 0;
+		for (u32 i=0; i!= sizeof_array(m_activeProcessCountHistory) - 1; ++i)
+			highestHistoryCount = Max(highestHistoryCount, m_activeProcessCountHistory[i]);
+
+		u32 activeProcessCount = highestHistoryCount;
+		counts[0] = activeProcessCount;
+		u32 totalHeight = counts[0]*2;
+		while (totalHeight < maxHeight && counts[maxSizeMinusOne] != activeProcessCount)
+		{
+			bool changed = false;
+			for (u32 i=0; i!=maxSizeMinusOne; ++i)
+			{
+				if (counts[i] && counts[i] > counts[i+1]*2 + 1)
+				{
+					--counts[i];
+					++counts[i+1];
+					++totalHeight;
+					changed = true;
+				}
+			}
+			if (!changed)
+			{
+				for (u32 j=0; j!=maxSizeMinusOne; ++j)
+				{
+					if (!counts[j])
+						continue;
+					++counts[j + 1];
+					--counts[j];
+					++totalHeight;
+					break;
+				}
+			}
+		}
+
+		auto it = activeProcesses.begin();
+		auto itEnd = activeProcesses.end();
+		int endY = int(startPosY + maxHeight);
+		for (u32 i=0;i!=maxSize;++i)
+		{
+			u32 v = maxSizeMinusOne - i;
+			u32 boxHeight = v + 1;
+
+			for (u32 j=0;j!=counts[v];++j)
+			{
+				if (it == itEnd || posY >= endY)
+					break;
+				auto& active = *it->second;
+				++it;
+				drawProcess(active, boxHeight, j == 0);
+				posY += boxHeight + 1;
+			}
+		}
+
+		//if (fillHeight)
+		//	posY = clientRect.bottom-1;//startPosY; // To prevent vertical scrollbar
+		if (fillHeight || counts[maxSizeMinusOne] != activeProcessCount)
+			posY = startPosY + maxHeight;
+		else
+			posY += 3;
+
+		SetActiveFont(m_defaultFont);
+	}
+
+	void VisualizerWin::PaintProcessRect(TraceView::Process& process, HDC hdc, RECT rect, const RECT& progressRect, bool selected, bool writingBitmap, HBRUSH& lastSelectedBrush)
+	{
+		auto clampRect = [&](RECT& r) { r.left = Min(Max(r.left, progressRect.left), progressRect.right); r.right = Max(Min(r.right, progressRect.right), progressRect.left); };
+		bool done = process.stop != ~u64(0);
+
+		HBRUSH brush = m_processBrushes[selected].success;
+		if (!process.returnedReason.empty())
+			brush = m_processBrushes[selected].returned;
+		else if (!done)
+			brush = process.isRacing ? m_processBrushes[selected].inProgressRace : m_processBrushes[selected].inProgress;
+		else if (process.type == TraceView::ProcessType_CacheFetchTest)
+			brush = m_processBrushes[selected].cacheFetchTest;
+		else if (process.type == TraceView::ProcessType_CacheFetchDownload)
+			brush = m_processBrushes[selected].cacheFetchDownload;
+		else if (process.exitCode == ProcessCancelExitCode)
+			brush = m_processBrushes[selected].returned;
+		else if (process.exitCode != 0)
+			brush = m_processBrushes[selected].error;
+
+		u64 writeFilesTime = process.writeFilesTime;
+
+		auto fillRect = [&](const RECT& r, HBRUSH b)
+			{
+				if (lastSelectedBrush != b)
+				{
+					SelectObject(hdc, b);
+					lastSelectedBrush = b;
+				}
+				PatBlt(hdc, r.left, r.top, r.right - r.left, r.bottom - r.top, PATCOPY);
+			};
+
+		if (!done || process.exitCode != 0 || !m_config.ShowReadWriteColors)
+		{
+			if (writingBitmap)
+				rect.right = 256;
+			clampRect(rect);
+			fillRect(rect, brush);
+			return;
+		}
+
+		double duration = double(process.stop - process.start);
+
+		RECT main = rect;
+		int width = rect.right - rect.left;
+
+		double recvPart = (double(ConvertTime(m_traceView, process.createFilesTime)) / duration);
+		if (int headSize = int(recvPart * width))
+		{
+			main.left += headSize;
+			RECT r2 = rect;
+			r2.right = r2.left + headSize;
+			clampRect(r2);
+			if (r2.left != r2.right)
+				fillRect(r2, m_processBrushes[selected].recv);
+		}
+
+		double sendPart = (double(ConvertTime(m_traceView, writeFilesTime)) / duration);
+		if (int tailSize = int(sendPart * width))
+		{
+			main.right -= tailSize;
+			RECT r2 = rect;
+			r2.left = r2.right - tailSize;
+			clampRect(r2);
+			if (r2.left != r2.right)
+				fillRect(r2, m_processBrushes[selected].send);
+		}
+
+		clampRect(main);
+		if (main.left != main.right)
+			fillRect(main, brush);
+
+		//clampRect(rect);
+		/*
+		if (!process.updates.empty())
+		{
+		shouldDrawText = false;
+		int prevUpdateX = rect.left;
+		for (auto& update : process.updates)
+		{
+		int updateX = int(posX + TimeToS(update.time) * scaleX) - 1;
+
+		RECT textRect{ prevUpdateX + 5, rect.top, updateX, rect.bottom };
+		DrawTextW(hdc, update.reason.c_str(), u32(update.reason.size()), &textRect, DT_SINGLELINE);
+
+		SelectObject(hdc, m_processUpdatePen);
+		MoveToEx(hdc, updateX, rect.top, NULL);
+		LineTo(hdc, updateX, rect.bottom - 1);
+		prevUpdateX = updateX;
+		}
+		}
+		*/
+	}
+
+	void VisualizerWin::PaintTimeline(HDC hdc, const RECT& clientRect)
+	{
+		SetActiveFont(m_timelineFont);
+		int top = GetTimelineTop(clientRect);
+		float timeScale = (m_config.horizontalScaleValue*m_zoomValue)*50.0f;
+		float startOffset = ((m_scrollPosX/timeScale) - float(int(m_scrollPosX/timeScale))) * timeScale;
+		int index = -int(startOffset/timeScale);
+
+		int number = -int(float(m_scrollPosX)/timeScale);
+
+		int textStepSize = int((5.0f / timeScale) + 1) * 5;
+		if (textStepSize > 150)
+			textStepSize = 600;
+		else if (textStepSize > 120)
+			textStepSize = 300;
+		else if (textStepSize > 90)
+			textStepSize = 240;
+		else if (textStepSize > 45)
+			textStepSize = 120;
+		else if (textStepSize > 30)
+			textStepSize = 60;
+		else if (textStepSize > 10)
+			textStepSize = 30;
+
+		int lineStepSize = textStepSize / 5;
+
+		RECT progressRect = clientRect;
+		progressRect.left += m_progressRectLeft;
+
+		SelectObject(hdc, m_textPen);
+		Vector<DWORD> lines;
+		Vector<POINT> points;
+
+		while (true)
+		{
+			int pos = progressRect.left + int(startOffset + float(index)*timeScale);
+			if (pos >= clientRect.right)
+				break;
+
+			int lineBottom = top + 5;
+			if (!(number % textStepSize))
+			{
+				bool shouldDraw = true;
+				int seconds = number;
+				StringBuffer<> buffer;
+				if (seconds >= 60)
+				{
+					int min = seconds / 60;
+					seconds -= min * 60;
+					if (!seconds)
+					{
+						buffer.Appendf(L"%um", min);
+						lineBottom += 4;
+					}
+				}
+				if (!number || seconds)
+					buffer.Appendf(L"%u", seconds);
+
+				if (shouldDraw)
+				{
+					SIZE s;
+					GetTextExtentPoint32W(hdc, buffer.data, buffer.count, &s);
+					RECT textRect;
+					textRect.top = top + 8;
+					textRect.bottom = textRect.top + m_activeFont.height;
+					textRect.right = pos + s.cx/2;
+					textRect.left = pos - s.cx/2;
+					ExtTextOutW(hdc, textRect.left, textRect.top, 0, NULL, buffer.data, buffer.count, NULL);
+				}
+			}
+			if (!(number % lineStepSize))
+			{
+				points.push_back({pos, top});
+				points.push_back({pos, lineBottom});
+				lines.push_back(2);
+			}
+
+			++number;
+			++index;
+		}
+
+		points.push_back({m_contentWidth, top - 25});
+		points.push_back({m_contentWidth, top});
+		lines.push_back(2);
+
+		PolyPolyline(hdc, points.data(), lines.data(), int(lines.size()));
+
+		/*
+		POINT cursorPos;
+		GetCursorPos(&cursorPos);
+		ScreenToClient(m_hwnd, &cursorPos);
+		RECT lineRect;
+		lineRect.left = cursorPos.x;
+		lineRect.right= cursorPos.x + 1;
+		lineRect.top = top;
+		lineRect.bottom = top + 15;
+		FillRect(hdc, &lineRect, m_lineBrush);
+		*/
+	}
+
+	void VisualizerWin::PaintDetailedStats(int& posY, const RECT& progressRect, TraceView::Session& session, bool isRemote, u64 playTime, const DrawTextFunc& drawTextFunc)
+	{
+		int stepY = m_activeFont.height;
+		int startPosY = posY;
+		int posX = progressRect.left + 5;
+		int maxY = posY + stepY;
+		u32 maxTextWidth = 0;
+
+		RECT textRect;
+		auto drawText = [&](const wchar_t* format, ...)
+			{
+				textRect.top = posY;
+				textRect.bottom = posY + 20;
+				textRect.left = posX;
+				textRect.right = posX + 1000;
+				posY += stepY;
+				StringBuffer<> str;
+				va_list arg;
+				va_start(arg, format);
+				str.Append(format, arg);
+				va_end(arg);
+				u32 lastTextWidth = 0;
+				drawTextFunc(str, textRect, &lastTextWidth);
+				maxTextWidth = Max(maxTextWidth, lastTextWidth);
+				maxY = Max(maxY, posY);
+			};
+
+		if (isRemote)
+		{
+			drawText(L"Finished Processes: %u", session.processExitedCount);
+			drawText(L"Active Processes: %u", session.processActiveCount);
+			drawText(L"ClientId: %u  TcpCount: %u", session.clientUid.data1, session.connectionCount.empty() ? 1 : session.connectionCount.back());
+
+			if (session.disconnectTime == ~u64(0))
+			{
+				if (session.proxyCreated)
+					drawText(L"Proxy(HOSTED): %ls", session.proxyName.c_str());
+				else if (!session.proxyName.empty())
+					drawText(L"Proxy: %ls", session.proxyName.c_str());
+				else
+					drawText(L"Proxy: None");
+			}
+
+			bool hasFileDetails = !session.fetchedFiles.empty() || !session.storedFiles.empty();
+
+			if (!hasFileDetails)
+			{
+				posY = startPosY;
+				posX += maxTextWidth + 15;
+			}
+
+			if (!session.updates.empty())
+			{
+				auto updateTime = session.updates.back();
+				auto updateSend = session.networkSend.back();
+				auto updateRecv = session.networkRecv.back();
+				u64 sendPerS = 0;
+				u64 recvPerS = 0;
+				float duration = TimeToS(updateTime - session.prevUpdateTime);
+				if (duration != 0)
+				{
+					sendPerS = u64(float(updateSend - session.prevSend) / duration);
+					recvPerS = u64(float(updateRecv - session.prevRecv) / duration);
+				}
+				drawText(L"Recv: %ls (%sit/s)", BytesToText(updateRecv), BytesToText(recvPerS*8));
+				drawText(L"Send: %ls (%sit/s)", BytesToText(updateSend), BytesToText(sendPerS*8));
+			}
+
+			int fileWidth = 700;
+
+			auto drawFiles = [&](const wchar_t* fileType, Vector<TraceView::FileTransfer>& files, u64 count, u64 bytes, u64 activeCount, u32& maxVisibleFiles)
+				{
+					textRect.left = posX;
+					textRect.right = posX + fileWidth;
+					drawText(L"%ls Files: %u (%u) %s", fileType, u32(count), u32(activeCount), BytesToText(bytes));
+					u32 fileCount = 0;
+					for (auto rit = files.rbegin(), rend = files.rend(); rit != rend; ++rit)
+					{
+						TraceView::FileTransfer& file = *rit;
+						if (file.stop != ~u64(0))
+							continue;
+						u64 time = 0;
+						if (file.start < playTime)
+							time = playTime - file.start;
+
+						if ((((u8*)&file.key)[19] & IsProxyMask) != 0)
+							drawText(L"%s (proxy) %s", file.hint.c_str(), TimeToText(time, true).str);
+						else if (file.size == 0)
+							drawText(L"%s (calc) %s", file.hint.c_str(), TimeToText(time, true).str);
+						else
+							drawText(L"%s (%s) %s", file.hint.c_str(), BytesToText(file.size).str, TimeToText(time, true).str);
+
+						if (fileCount++ > 5)
+							break;
+					}
+					posY += stepY * (maxVisibleFiles - fileCount);
+					maxVisibleFiles = Max(maxVisibleFiles, fileCount);
+				};
+
+			Vector<TraceView::FileTransfer> fetchedFiles;
+			for (auto& kv : session.fetchedFilesActive)
+				fetchedFiles.push_back(session.fetchedFiles[kv.second]);
+			std::sort(fetchedFiles.begin(), fetchedFiles.end(), [](const TraceView::FileTransfer& a, const TraceView::FileTransfer& b) { return a.start > b.start; });
+
+			if (hasFileDetails)
+			{
+				posY = startPosY;
+				posX += maxTextWidth + 15;
+			}
+			drawFiles(L"Fetched", fetchedFiles, session.fetchedFilesCount, session.fetchedFilesBytes, session.fetchedFilesActive.size(), session.maxVisibleFiles);
+			if (hasFileDetails)
+			{
+				posY = startPosY;
+				posX += fileWidth;
+			}
+			drawFiles(L"Stored", session.storedFiles, session.storedFilesCount, session.storedFilesBytes, session.storedFilesActive.size(), session.maxVisibleFiles);
+		}
+		else
+		{
+			if (m_traceView.totalProcessActiveCount || m_traceView.totalProcessExitedCount)
+			{
+				drawText(L"Finished Processes: %u (local: %u)", m_traceView.totalProcessExitedCount, session.processExitedCount);
+				drawText(L"Active Processes: %u (local: %u)", m_traceView.totalProcessActiveCount, session.processActiveCount);
+				drawText(L"Active Helpers: %u", Max(1u, m_traceView.activeSessionCount) - 1);
+			}
+
+			if (session.highestSendPerS > 0.0f || session.highestRecvPerS > 0.0f)
+			{
+				auto updateTime = session.updates.back();
+				auto updateSend = session.networkSend.back();
+				auto updateRecv = session.networkRecv.back();
+				if (updateSend || updateRecv)
+				{
+					u64 sendPerS = 0;
+					u64 recvPerS = 0;
+					float duration = TimeToS(updateTime - session.prevUpdateTime);
+					if (duration != 0)
+					{
+						sendPerS = u64(float(updateSend - session.prevSend) / duration);
+						recvPerS = u64(float(updateRecv - session.prevRecv) / duration);
+					}
+
+					drawText(L"Recv: %ls (%sit/s)", BytesToText(updateRecv), BytesToText(recvPerS*8));
+					drawText(L"Send: %ls (%sit/s)", BytesToText(updateSend), BytesToText(sendPerS*8));
+				}
+			}
+
+			if (!session.updates.empty())
+			{
+				posY = startPosY;
+				posX += maxTextWidth + 10;
+
+				for (auto& pair : session.drives)
+				{
+					auto& drive = pair.second;
+
+					u64 readPerS = 0;
+					u64 writePerS = 0;
+
+					u64 updateCount = session.updates.size();
+					float duration = TimeToS(session.updates.back() - session.prevUpdateTime);
+					if (duration != 0 && updateCount > 1)
+					{
+						readPerS = u64(float(drive.readBytes.back()) / duration);
+						writePerS = u64(float(drive.writeBytes.back()) / duration);
+					}
+
+					drawText(L"%c: Rd %ls (%s/s) Wr %ls (%s/s) ", pair.first, BytesToText(drive.totalReadBytes).str, BytesToText(readPerS).str, BytesToText(drive.totalWriteBytes).str, BytesToText(writePerS).str);
+				}
+			}
+		}
+
+		posY = maxY;
+	}
+
+	u64 VisualizerWin::GetPlayTime()
+	{
+		u64 currentTime = m_paused ? m_pauseStart : GetTime();
+		u64 playTime = 0;
+		if (m_traceView.startTime && currentTime > (m_traceView.startTime + m_pauseTime))
+			playTime = currentTime - m_traceView.startTime - m_pauseTime;
+		if (m_replay)
+			playTime *= m_replay;
+		return playTime;
+	}
+
+	int VisualizerWin::GetTimelineHeight()
+	{
+		return m_timelineFont.height + 8;
+	}
+
+	int VisualizerWin::GetTimelineTop(const RECT& clientRect)
+	{
+		int timelineHeight = GetTimelineHeight();
+		int posY = m_contentHeight - timelineHeight;
+		int maxY = int(clientRect.bottom - timelineHeight);
+		return m_config.LockTimelineToBottom ? maxY : Min(posY, maxY);
+	}
+
+	void VisualizerWin::HitTest(HitTestResult& outResult, const POINT& pos)
+	{
+		SetActiveFont(m_defaultFont);
+
+		u64 playTime = GetPlayTime();
+
+		RECT clientRect;
+		GetClientRect(m_hwnd, &clientRect);
+
+		int posY = int(m_scrollPosY);
+		int boxHeight = m_config.boxHeight;
+		int processStepY = boxHeight + 1;
+		float scaleX = 50.0f*m_zoomValue*m_config.horizontalScaleValue;
+
+		RECT progressRect = clientRect;
+		progressRect.left += m_progressRectLeft;
+		progressRect.bottom -= 30;
+
+		{
+			int boxSide = 8;
+			int boxStride = boxSide + 2;
+			int top = 5;
+			int bottom = top + boxSide;
+			int left = progressRect.right - 7 - boxSide;
+			int right = progressRect.right - 7;
+			for (int i = VisualizerFlag_Count - 1; i >= 0; --i)
+			{
+				if (pos.x >= left && pos.x <= right && pos.y >= top && pos.y <= bottom)
+				{
+					outResult.buttonSelected = i;
+					return;
+				}
+				left -= boxStride;
+				right -= boxStride;
+			}
+		}
+
+		outResult.section = 0;
+
+		if (m_config.showTimeline && !m_traceView.sessions.empty())
+		{
+			int timelineTop = GetTimelineTop(clientRect);
+			if (pos.y >= timelineTop)// && pos.y < timelineTop + 40)
+			{
+				outResult.section = 3;
+				float timeScale = (m_config.horizontalScaleValue * m_zoomValue)*50.0f;
+				float startOffset = -(m_scrollPosX / timeScale);
+				outResult.timelineSelected = startOffset + float(pos.x - m_progressRectLeft) / timeScale;
+				return;
+			}
+		}
+
+		u64 lastStop = 0;
+
+		if (m_config.showProgress && m_traceView.progressProcessesTotal)
+			posY += m_activeFont.height + 2;
+
+		if (m_config.showStatus && !m_traceView.statusMap.empty())
+		{
+			u32 lastRow = ~0u;
+			u32 row = ~0u;
+			for (auto& kv : m_traceView.statusMap)
+			{
+				if (kv.second.text.empty())
+					continue;
+
+				row = u32(kv.first >> 32);
+				if (lastRow != ~0u && lastRow != row)
+					posY += m_activeFont.height + 2;
+				lastRow = row;
+
+				if (!kv.second.link.empty())
+				{
+					if (pos.y >= posY && pos.y < posY + m_activeFont.height && pos.x > 20 && pos.x < 80) // TODO: x is just hard coded to fit horde for now
+					{
+						outResult.hyperLink = kv.second.link;
+						return;
+					}
+				}
+			}
+			if (row != ~0u)
+				posY += m_activeFont.height + 2;
+			posY += 3;
+		}
+
+		if (pos.y < posY)
+			return;
+
+		outResult.section = 1;
+
+		if (m_config.showActiveProcessGraph)
+		{
+			if (pos.y > posY && pos.y < posY + GraphHeight)
+			{
+				float timeScale = (m_config.horizontalScaleValue * m_zoomValue)*50.0f;
+				float startOffset = -(m_scrollPosX / timeScale);
+				u64 selectedTimeMs = u64(1000.0f*(startOffset + float(pos.x - m_progressRectLeft) / timeScale));
+
+				u64 lastTime = m_traceView.activeProcessCounts.empty() ? 0 : m_traceView.activeProcessCounts[m_traceView.activeProcessCounts.size()-1].time;
+				if (selectedTimeMs < TimeToMs(lastTime))
+				{
+					u16 count = 0;
+					for (auto& entry : m_traceView.activeProcessCounts)
+					{
+						count = entry.count;
+						if (TimeToMs(entry.time) > selectedTimeMs)
+							break;
+					}
+
+					outResult.activeProcessGraphSelected = true;
+					outResult.activeProcessCount = count;
+				}
+			}
+			posY += GraphHeight;
+		}
+
+		TraceView::ProcessLocation& outLocation = outResult.processLocation;
+
+		if (m_config.showActiveProcesses && !m_trace.m_activeProcesses.empty())
+		{
+			PaintActiveProcesses(posY, clientRect, [&](TraceView::ProcessLocation& processLocation, u32 boxHeight, bool firstWithHeight)
+				{
+					if (pos.y < posY || pos.y > posY + int(boxHeight))
+						return;
+					auto& session = m_trace.GetSession(m_traceView, processLocation.sessionIndex);
+					TraceView::Process& process = session.processors[processLocation.processorIndex].processes[processLocation.processIndex];
+
+					int posX = int(m_scrollPosX) + progressRect.left;
+					u64 stop = process.stop;
+					bool done = stop != ~u64(0);
+					if (!done)
+						stop = playTime;
+					int left = posX + int(float(TimeToS(process.start)) * scaleX);
+					int right = posX + int(float(TimeToS(stop)) * scaleX) - 1;
+
+					if (pos.x >= left && pos.x <= right)
+					{
+						outLocation.sessionIndex = processLocation.sessionIndex;
+						outLocation.processorIndex = processLocation.processorIndex;
+						outLocation.processIndex = processLocation.processIndex;
+						outResult.processSelected = true;
+						return;
+					}
+				});
+			if (outResult.processSelected)
+				return;
+		}
+
+		if (pos.y < posY)
+			return;
+
+		outResult.section = 2;
+
+		SessionRec sortedSessions[1024];
+		Populate(sortedSessions, m_traceView, m_config.SortActiveRemoteSessions);
+
+		for (u64 sessionIt = 0, sessionEnd = m_traceView.sessions.size(); sessionIt != sessionEnd; ++sessionIt)
+		{
+			bool isFirst = sessionIt == 0;
+			auto& session = *sortedSessions[sessionIt].session;
+			bool hasUpdates = !session.updates.empty();
+			if (!isFirst)
+			{
+				if (!hasUpdates && session.processors.empty())
+					continue;
+				if (!m_config.showFinishedProcesses && session.disconnectTime != ~u64(0))
+					continue;
+			}
+
+			u32 sessionIndex = sortedSessions[sessionIt].index;
+			if (!isFirst)
+				posY += 3;
+
+			if (m_config.showTitleBars)
+			{
+				if (pos.y >= posY && pos.y < posY + m_sessionStepY)
+				{
+					if (pos.x < int(session.fullNameWidth) + 5)
+					{
+						outResult.sessionSelectedIndex = sessionIndex;
+						return;
+					}
+				}
+
+				posY += m_sessionStepY;
+			}
+
+			bool showGraph = m_config.showNetworkStats || m_config.showCpuMemStats || m_config.showDriveStats;
+			if (showGraph && !session.updates.empty())
+			{
+				int posX = int(m_scrollPosX) + progressRect.left;
+
+				if (pos.y >= posY && pos.y < posY + GraphHeight && pos.x >= posX)
+				{
+					bool loop = true;
+					u32 reconnectIndex = 0;
+					while (loop)
+					{
+						u64 i = 0;
+						u64 e = 0;
+
+						if (reconnectIndex)
+							i = session.reconnectIndices[reconnectIndex - 1];
+
+						if (reconnectIndex < session.reconnectIndices.size())
+							e = session.reconnectIndices[reconnectIndex];
+						else
+						{
+							e = session.updates.size();
+							loop = false;
+						}
+
+						int prevX = -1;
+						int prevCenterX = -1;
+						for (; i!=e; ++i)
+						{
+							u64 updateTime = session.updates[i];
+							auto updateSend = session.networkSend[i];
+							auto updateRecv = session.networkRecv[i];
+
+							int x = posX + int(float(TimeToS(updateTime)) * scaleX);
+							int centerX = prevX != -1 ? (prevX + (x - prevX)/2) : x;
+
+							if (prevCenterX != -1 && pos.x >= prevCenterX && pos.x < centerX)
+							{
+								u64 prevTime = 0;
+								u64 prevSend = 0;
+								u64 prevRecv = 0;
+								if (i > 0)
+								{
+									prevTime = session.updates[i - 1];
+									prevSend = Min(session.networkSend[i - 1], updateSend);
+									prevRecv = Min(session.networkRecv[i - 1], updateRecv);
+								}
+
+								double duration = TimeToS(updateTime - prevTime);
+								outResult.stats.recvBytes = updateRecv;
+								outResult.stats.sendBytes = updateSend;
+								outResult.stats.recvBytesPerSecond = u64(float(updateRecv - prevRecv) / duration);
+								outResult.stats.sendBytesPerSecond = u64(float(updateSend - prevSend) / duration);
+								outResult.stats.ping = session.ping[i];
+								outResult.stats.memAvail = session.memAvail[i];
+								outResult.stats.cpuLoad = session.cpuLoad[i];
+								outResult.stats.memTotal = session.memTotal;
+								outResult.stats.connectionCount = session.connectionCount[i];
+								outResult.statsSelected = true;
+
+								for (auto& pair : session.drives)
+								{
+									auto& updateBusy = pair.second.busyPercent[i];
+									u64 updateRead = pair.second.readBytes[i];
+									u64 updateWrite = pair.second.writeBytes[i];
+									auto& statsDrive = outResult.stats.drives[pair.first];
+									statsDrive.busyPercent = updateBusy;
+									statsDrive.readPerSecond = u64(float(updateRead) / duration);
+									statsDrive.writePerSecond = u64(float(updateWrite) / duration);
+								}
+								return;
+							}
+
+							prevCenterX = centerX;
+							prevX = x;
+						}
+						++reconnectIndex;
+					}
+					posY += GraphHeight;
+				}
+
+				posY += GraphHeight;
+			}
+
+			if (m_config.showDetailedData)
+			{
+				auto drawText = [&](const StringBufferBase& text, RECT& rect, u32* outWidth)
+					{
+						if (pos.x >= rect.left && pos.x < rect.right && pos.y >= rect.top && pos.y < rect.bottom && text.StartsWith(TC("Fetched Files")))
+							outResult.fetchedFilesSelected = sessionIndex;
+					};
+				PaintDetailedStats(posY, progressRect, session, sessionIt != 0, playTime, drawText);
+			}
+
+			if (m_config.showProcessBars)
+			{
+				u32 processorIndex = 0;
+				for (auto& processor : session.processors)
+				{
+					bool drawProcessorIndex = m_config.showFinishedProcesses;
+
+					if (pos.y < progressRect.bottom && posY + processStepY >= progressRect.top && posY <= progressRect.bottom && pos.y >= posY-1 && pos.y < posY-1 + processStepY)
+					{
+						u32 processIndex = 0;
+						int posX = int(m_scrollPosX) + progressRect.left;
+						for (auto& process : processor.processes)
+						{
+							int left = posX + int(float(TimeToS(process.start)) * scaleX);
+
+							auto pig = MakeGuard([&]() { ++processIndex; });
+
+							if (left >= progressRect.right)
+								continue;
+
+							if (left < progressRect.left)
+								left = progressRect.left;
+
+							u64 stopTime = process.stop;
+							bool done = stopTime != ~u64(0);
+							if (!done)
+								stopTime = playTime;
+							else if (!m_config.showFinishedProcesses)
+								continue;
+
+							drawProcessorIndex = true;
+
+							RECT rect;
+							rect.left = left;
+							rect.right = posX + int(float(TimeToS(stopTime)) * scaleX);
+
+							if (rect.right <= progressRect.left)
+								continue;
+
+							if (!m_filterString.empty() && !Contains(process.description.c_str(), m_filterString.c_str()) && !Contains(process.breadcrumbs.c_str(), m_filterString.c_str()))
+								continue;
+
+							rect.right = Max(int(rect.right), left + 1);
+							rect.top = posY;
+							rect.bottom = posY + int(float(18) * m_zoomValue);
+
+							if (pos.x >= rect.left && pos.x <= rect.right)
+							{
+								outLocation.sessionIndex = sessionIndex;
+								outLocation.processorIndex = processorIndex;
+								outLocation.processIndex = processIndex;
+								outResult.processSelected = true;
+								return;
+							}
+						}
+					}
+
+					if (!processor.processes.empty())
+						lastStop = Max(lastStop, processor.processes.rbegin()->stop);
+
+					if (drawProcessorIndex)
+						posY += processStepY;
+
+					++processorIndex;
+				}
+			}
+			else
+			{
+				for (auto& processor : session.processors)
+					if (!processor.processes.empty())
+						lastStop = Max(lastStop, processor.processes.rbegin()->stop);
+			}
+
+			if (m_config.showWorkers && isFirst)
+			{
+				int trackIndex = 0;
+				for (auto& workTrack : m_traceView.workTracks)
+				{
+					if (pos.y < progressRect.bottom && posY + processStepY >= progressRect.top && posY <= progressRect.bottom && pos.y >= posY-1 && pos.y < posY-1 + processStepY)
+					{
+						u32 workIndex = 0;
+						int posX = int(m_scrollPosX) + progressRect.left;
+						for (auto& work : workTrack.records)
+						{
+							auto inc = MakeGuard([&](){ ++workIndex; });
+
+							int left = posX + int(float(TimeToS(work.start)) * scaleX);
+
+							if (left >= progressRect.right)
+								continue;
+
+							if (!m_filterString.empty())
+							{
+								bool keep = Contains(work.description, m_filterString.c_str());
+								if (!keep)
+									for (auto& en : work.entries)
+										keep |= Contains(en.text, m_filterString.c_str());
+								if (!keep)
+									continue;
+							}
+
+							if (left < progressRect.left)
+								left = progressRect.left;
+
+							u64 stopTime = work.stop;
+							bool done = stopTime != ~u64(0);
+							if (!done)
+								stopTime = playTime;
+
+							RECT rect;
+							rect.left = left;
+							rect.right = posX + int(float(TimeToS(stopTime)) * scaleX);
+
+							if (rect.right <= progressRect.left)
+								continue;
+
+							rect.right = Max(int(rect.right), left + 1);
+							rect.top = posY;
+							rect.bottom = posY + int(float(18) * m_zoomValue);
+
+							if (pos.x >= rect.left && pos.x <= rect.right)
+							{
+								outResult.workTrack = trackIndex;
+								outResult.workIndex = workIndex;
+								outResult.workSelected = true;
+								return;
+							}
+						}
+					}
+					++trackIndex;
+					posY += processStepY;
+				}
+			}
+		}
+
+		m_contentWidth = m_progressRectLeft + Max(0, int(TimeToS((lastStop != 0 && lastStop != ~u64(0)) ? lastStop : playTime) * scaleX));
+		m_contentHeight = posY - int(m_scrollPosY) + processStepY + 14;
+	}
+
+	void VisualizerWin::WriteProcessStats(Logger& out, const TraceView::Process& process)
+	{
+		bool hasExited = process.stop != ~u64(0);
+		out.Info(L"  %ls", process.description.c_str());
+		out.Info(L"  ProcessId: %u", process.id);
+		out.Info(L"  Start:     %ls", TimeToText(process.start, true).str);
+		if (hasExited)
+			out.Info(L"  Duration:  %ls", TimeToText(process.stop - process.start, true).str);
+		if (hasExited && process.exitCode != 0)
+			out.Info(L"  ExitCode:  %u", process.exitCode);
+
+		if (process.stop != ~u64(0))
+		{
+			out.Info(L"");
+
+			BinaryReader reader(process.stats.data(), 0, process.stats.size());
+			ProcessStats processStats;
+			SessionStats sessionStats;
+			StorageStats storageStats;
+			KernelStats kernelStats;
+
+			processStats.Read(reader, m_traceView.version);
+			if (reader.GetLeft())
+			{
+				if (process.isRemote || (m_traceView.version >= 36 && !process.isReuse))
+					sessionStats.Read(reader, m_traceView.version);
+				storageStats.Read(reader, m_traceView.version);
+				kernelStats.Read(reader, m_traceView.version);
+			}
+
+			out.Info(L"  ----------- Detours stats -----------");
+			processStats.Print(out, m_traceView.frequency);
+
+			if (!sessionStats.IsEmpty())
+			{
+				out.Info(L"");
+				out.Info(L"  ----------- Session stats -----------");
+				sessionStats.Print(out, m_traceView.frequency);
+			}
+
+			if (!storageStats.IsEmpty())
+			{
+				out.Info(L"");
+				out.Info(L"  ----------- Storage stats -----------");
+				storageStats.Print(out, m_traceView.frequency);
+			}
+
+			if (!kernelStats.IsEmpty())
+			{
+				out.Info(L"");
+				out.Info(L"  ----------- Kernel stats ------------");
+				kernelStats.Print(out, false, m_traceView.frequency);
+			}
+
+			PrintCacheWriteStats(out, process.id);
+		}
+	}
+
+	void VisualizerWin::WriteWorkStats(Logger& out, const TraceView::WorkRecord& record)
+	{
+		out.Info(L"  %ls", record.description);
+		out.Info(L"  Start:     %ls", TimeToText(record.start, true).str);
+		if (record.stop != ~u64(0))
+			out.Info(L"  Duration:  %ls", TimeToText(record.stop - record.start, true).str);
+		for (auto& e : record.entries)
+			out.Info(L"   %ls (%ls)", e.text, TimeToText(e.time - record.start).str);
+	}
+
+	void VisualizerWin::CopyTextToClipboard(const TString& str)
+	{
+		if (!OpenClipboard(m_hwnd))
+			return;
+		if (auto hglbCopy = GlobalAlloc(GMEM_MOVEABLE, (str.size() + 1) * sizeof(TCHAR)))
+		{
+			if (auto lptstrCopy = GlobalLock(hglbCopy))
+			{
+				memcpy(lptstrCopy, str.data(), (str.size() + 1) * sizeof(TCHAR));
+				GlobalUnlock(hglbCopy);
+				EmptyClipboard();
+				SetClipboardData(CF_UNICODETEXT, hglbCopy);
+			}
+		}
+		CloseClipboard();
+	}
+
+	void VisualizerWin::UnselectAndRedraw()
+	{
+		if (Unselect() || m_config.showCursorLine)
+			Redraw(false);
+	}
+
+	bool VisualizerWin::UpdateAutoscroll()
+	{
+		if (!m_autoScroll)
+			return false;
+
+		u64 playTime = GetPlayTime();
+
+
+		RECT rect;
+		GetClientRect(m_hwnd, &rect);
+		if (rect.right == 0)
+			return false;
+
+		float timeS = TimeToS(playTime);
+
+		if (m_config.AutoScaleHorizontal)
+		{
+			m_scrollPosX = 0;
+			timeS = Max(timeS, 20.0f/m_zoomValue);
+			m_config.horizontalScaleValue = Max(float(rect.right - m_progressRectLeft - 2)/(m_zoomValue*timeS*50.0f), 0.001f);
+			return true;
+		}
+		else
+		{
+			float oldScrollPosX = m_scrollPosX;
+			m_scrollPosX = Min(0.0f, (float)rect.right - timeS*50.0f*m_config.horizontalScaleValue*m_zoomValue - float(m_progressRectLeft));
+			return oldScrollPosX != m_scrollPosX;
+		}
+	}
+
+	bool VisualizerWin::UpdateSelection()
+	{
+		if (!m_mouseOverWindow || m_dragToScrollCounter > 0)
+			return false;
+		POINT pos;
+		GetCursorPos(&pos);
+		ScreenToClient(m_hwnd, &pos);
+
+		HitTestResult res;
+		HitTest(res, pos);
+
+		m_activeSection = res.section;
+
+		if (res.processSelected == m_processSelected && res.processLocation == m_processSelectedLocation &&
+			res.sessionSelectedIndex == m_sessionSelectedIndex &&
+			res.statsSelected == m_statsSelected && memcmp(&res.stats, &m_stats, sizeof(Stats)) == 0 &&
+			res.buttonSelected == m_buttonSelected && res.timelineSelected == m_timelineSelected &&
+			res.activeProcessGraphSelected == m_activeProcessGraphSelected && 
+			res.activeProcessCount == m_activeProcessCount && 
+			res.fetchedFilesSelected == m_fetchedFilesSelected &&
+			res.workSelected == m_workSelected && res.workTrack == m_workTrack && res.workIndex == m_workIndex &&
+			res.hyperLink == m_hyperLinkSelected)
+			return false;
+		m_processSelected = res.processSelected;
+		m_processSelectedLocation = res.processLocation;
+		m_sessionSelectedIndex = res.sessionSelectedIndex;
+		m_statsSelected = res.statsSelected;
+		m_stats = res.stats;
+		m_activeProcessGraphSelected = res.activeProcessGraphSelected;
+		m_activeProcessCount = res.activeProcessCount;
+		m_buttonSelected = res.buttonSelected;
+		m_timelineSelected = res.timelineSelected;
+		m_fetchedFilesSelected = res.fetchedFilesSelected;
+		m_workSelected = res.workSelected;
+		m_workTrack = res.workTrack;
+		m_workIndex = res.workIndex;
+		m_hyperLinkSelected = res.hyperLink;
+		return true;
+	}
+
+	void VisualizerWin::UpdateScrollbars(bool redraw)
+	{
+		RECT rect;
+		GetClientRect(m_hwnd, &rect);
+
+		SCROLLINFO si;
+		si.cbSize = sizeof(SCROLLINFO);
+		si.fMask = SIF_ALL | SIF_DISABLENOSCROLL;
+		si.nMin = 0;
+		si.nMax = m_contentHeight;
+		si.nPage = int(rect.bottom);
+		si.nPos = -int(m_scrollPosY);
+		si.nTrackPos = 0;
+
+		bool updateFrame = false;
+		if (ActiveProcessesShouldFillHeight())
+		{
+			if (m_verticalScrollBarEnabled)
+			{
+				SetWindowLong(m_hwnd, GWL_STYLE, GetWindowLong(m_hwnd, GWL_STYLE) & ~WS_VSCROLL);
+				m_verticalScrollBarEnabled = false;
+				updateFrame = true;
+			}
+		}
+		else
+		{
+			if (!m_verticalScrollBarEnabled)
+			{
+				SetWindowLong(m_hwnd, GWL_STYLE, GetWindowLong(m_hwnd, GWL_STYLE) | WS_VSCROLL);
+				m_verticalScrollBarEnabled = true;
+				updateFrame = true;
+			}
+			SetScrollInfo(m_hwnd, SB_VERT, &si, redraw);
+		}
+
+		if (m_config.AutoScaleHorizontal)
+		{
+			if (m_horizontalScrollBarEnabled)
+			{
+				SetWindowLong(m_hwnd, GWL_STYLE, GetWindowLong(m_hwnd, GWL_STYLE) & ~WS_HSCROLL);
+				m_horizontalScrollBarEnabled = false;
+				updateFrame = true;
+			}
+		}
+		else
+		{
+			if (!m_horizontalScrollBarEnabled)
+			{
+				SetWindowLong(m_hwnd, GWL_STYLE, GetWindowLong(m_hwnd, GWL_STYLE) | WS_HSCROLL);
+				m_horizontalScrollBarEnabled = true;
+				updateFrame = true;
+			}
+
+			si.nMax = m_contentWidth;
+			si.nPage = rect.right;
+			si.nPos = -int(m_scrollPosX);
+			SetScrollInfo(m_hwnd, SB_HORZ, &si, redraw);
+		}
+
+		if (updateFrame)
+			SetWindowPos(m_hwnd, NULL, 0, 0, 0, 0, SWP_NOSIZE | SWP_NOMOVE | SWP_NOZORDER | SWP_FRAMECHANGED);
+	}
+
+	void VisualizerWin::SetActiveFont(const Font& font)
+	{
+		m_activeFont = font;
+		if (m_activeHdc)
+			SelectObject(m_activeHdc, font.handle);
+	}
+
+	void VisualizerWin::UpdateTheme()
+	{
+		SetWindowTheme(m_hwnd, m_config.DarkMode ? L"DarkMode_Explorer" : L"Explorer", NULL);
+		SendMessageW(m_hwnd, WM_THEMECHANGED, 0, 0);
+		BOOL useDarkMode = m_config.DarkMode;
+		u32 attribute = 20; // DWMWA_USE_IMMERSIVE_DARK_MODE
+		DwmSetWindowAttribute(m_hwnd, attribute, &useDarkMode, sizeof(useDarkMode));
+	}
+
+	bool VisualizerWin::ActiveProcessesShouldFillHeight()
+	{
+		return !m_config.showDetailedData && !m_config.showTitleBars && !m_config.showCpuMemStats && !m_config.showNetworkStats && !m_config.showDriveStats && !m_config.showProcessBars;
+	}
+
+	StringBuffer<128> VisualizerWin::GetWorldTime(u64 time)
+	{
+		return GetWorldTime(TimeToS(time));
+	}
+
+	StringBuffer<128> VisualizerWin::GetWorldTime(float seconds)
+	{
+		return StringBuffer<128>(DateToText(float(m_traceView.traceSystemStartTimeUs / 1'000'000) + seconds).str);
+	}
+
+	void VisualizerWin::PostNewTrace(u32 replay, bool paused)
+	{
+		KillTimer(m_hwnd, 0);
+		PostMessage(m_hwnd, WM_NEWTRACE, replay, paused);
+	}
+
+	void VisualizerWin::PostNewTitle(const StringView& title)
+	{
+#if !defined( __clang_analyzer__ )
+		PostMessage(m_hwnd, WM_SETTITLE, 0, (LPARAM)_wcsdup(title.data));
+#endif
+	}
+
+	void VisualizerWin::PostQuit()
+	{
+		m_looping = false;
+		PostMessage(m_hwnd, WM_USER+666, 0, 0);
+	}
+
+	LRESULT VisualizerWin::WinProc(HWND hWnd, UINT Msg, WPARAM wParam, LPARAM lParam)
+	{
+		switch (Msg)
+		{
+		case WM_SETTITLE:
+		{
+			auto title = (wchar_t*)lParam;
+			SetWindowTextW(hWnd, title);
+			free(title);
+			break;
+		}
+		case WM_NEWTRACE:
+		{
+			m_replay = u32(wParam);
+			m_paused = lParam;
+			m_autoScroll = true;
+			m_scrollPosX = 0;
+			m_scrollPosY = 0;
+			Reset();
+			StringBuffer<> title;
+			GetTitlePrefix(title);
+
+			auto g = MakeGuard([&]()
+				{
+					Redraw(true);
+					UpdateScrollbars(true);
+				});
+
+			if (m_client)
+			{
+				if (!m_trace.StartReadClient(m_traceView, *m_client))
+				{
+					m_clientDisconnect.Set();
+					return false;
+				}
+				m_namedTrace.Clear().Append(m_newTraceName);
+				m_traceView.finished = false;
+			}
+			else if (!m_fileName.IsEmpty())
+			{
+				m_trace.ReadFile(m_traceView, m_fileName.data, m_replay != 0);
+				m_traceView.finished = m_replay == 0;
+				PostNewTitle(GetTitlePrefix(title).Appendf(L"%s (v%u) - %s", m_fileName.data, m_traceView.version, GetWorldTime((u64)0).data));
+			}
+			else if (m_usingNamed)
+			{
+				if (!m_trace.StartReadNamed(m_traceView, m_newTraceName.data, true, m_replay != 0))
+					return false;
+				m_namedTrace.Clear().Append(m_newTraceName);
+				m_traceView.finished = false;
+				PostNewTitle(GetTitlePrefix(title).Appendf(L"%s (Listening for new sessions on channel '%s')", m_namedTrace.data, m_listenChannel.data));
+			}
+
+			SetTimer(m_hwnd, 0, 200, NULL);
+			return 0;
+		}
+
+		case WM_SYSCOMMAND:
+			// Don't send this message through DefWindowProc as it will destroy the window 
+			// and GetMessage will stay stuck indefinitely.
+			if (wParam == SC_CLOSE)
+			{
+				PostQuit();
+				return 0;
+			}
+			break;
+
+		case WM_DESTROY:
+			PostQuit();
+			return 0;
+
+		case WM_ERASEBKGND:
+			return 1;
+
+		case WM_PAINT:
+		{
+			u64 start = GetTime();
+			PaintClient([&](HDC hdc, HDC memDC, RECT& rect)
+				{
+					FillRect(memDC, &rect, m_backgroundBrush);
+					m_activeHdc = memDC;
+					PaintAll(memDC, rect);
+					m_activeHdc = 0;
+					BitBlt(hdc, 0, 0, rect.right - rect.left, rect.bottom - rect.top, memDC, 0, 0, SRCCOPY);
+				});
+			m_lastPaintTimeMs = TimeToMs(GetTime() - start);
+			break;
+		}
+		case WM_SIZE:
+		{
+			int height = HIWORD(lParam);
+			if (m_contentHeight && m_contentHeight + int(m_scrollPosY) < height)
+				m_scrollPosY = float(Min(0, height - m_contentHeight));
+			int width = LOWORD(lParam);
+			if (m_contentWidth && m_contentWidth + int(m_scrollPosX) < width)
+				m_scrollPosX = float(Min(0, width - m_contentWidth));
+			UpdateScrollbars(true);
+			break;
+		}
+		case WM_TIMER:
+		{
+			if (m_handlingTimer) // This is to prevent hangs caused by message boxes from asserts
+				break;
+			m_handlingTimer = true;
+
+			bool changed = false;
+			if (!m_paused)
+			{
+				u64 timeOffset = (GetTime() - m_startTime - m_pauseTime) * m_replay;
+				if (!m_fileName.IsEmpty())
+				{
+					if (m_replay)
+						m_trace.UpdateReadFile(m_traceView, timeOffset, changed);
+				}
+				else if (m_client)
+				{
+					if (!m_trace.UpdateReadClient(m_traceView, *m_client, changed))
+						m_clientDisconnect.Set();
+				}
+				else if (m_usingNamed)
+				{
+					if (!m_trace.UpdateReadNamed(m_traceView, m_replay ? timeOffset : ~u64(0), changed))
+						if (m_listenTimeout.IsCreated())
+							m_listenTimeout.Set();
+				}
+			}
+
+			if (m_traceView.finished)
+			{
+				m_autoScroll = false;
+				KillTimer(m_hwnd, 0);
+				changed = true;
+			}
+
+			changed = UpdateAutoscroll() || changed;
+			changed = UpdateSelection() || changed;
+			if (changed && !IsIconic(m_hwnd))
+			{
+				UpdateScrollbars(true);
+
+				RedrawWindow(m_hwnd, NULL, NULL, RDW_INVALIDATE); // Don't use RDW_UPDATENOW.. it will cause stutter
+				u32 waitTime = 60u;//u32(Min(m_lastPaintTimeMs * 5, 200ull));
+				if (!m_traceView.finished)
+					SetTimer(m_hwnd, 0, waitTime, NULL);
+			}
+
+			m_handlingTimer = false;
+			break;
+		}
+		case WM_MOUSEWHEEL:
+		{
+			if (m_dragToScrollCounter > 0)
+				break;
+
+			int delta = GET_WHEEL_DELTA_WPARAM(wParam);
+			bool controlDown = GetAsyncKeyState(VK_CONTROL) & (1<<15);
+			bool shiftDown = GetAsyncKeyState(VK_LSHIFT) & (1<<15);
+
+			if (m_config.ScaleHorizontalWithScrollWheel || controlDown || shiftDown)
+			{
+				if (m_activeSection == 2 || !controlDown) // process bars
+				{
+					RECT r;
+					GetClientRect(hWnd, &r);
+
+					// Use mouse cursor as scroll anchor point
+					POINT cursorPos = {};
+					GetCursorPos(&cursorPos);
+					ScreenToClient(m_hwnd, &cursorPos);
+
+					float horizontalScaleValue = m_config.horizontalScaleValue;
+					float newScaleValue = horizontalScaleValue;
+					int newBoxHeight = m_config.boxHeight;
+
+					if (controlDown)
+					{
+						if (delta < 0)
+						{
+							if (newBoxHeight > 1)
+								--newBoxHeight;
+						}
+						else if (delta > 0)
+							++newBoxHeight;
+					}
+					else
+						newScaleValue = Max(horizontalScaleValue + horizontalScaleValue*float(delta)*0.0006f, 0.001f);
+
+					// TODO: m_progressRectLeft changes with zoom so anchor logic is wrong
+					const float scrollAnchorOffsetX = float(cursorPos.x) - float(m_progressRectLeft);
+					const float scrollAnchorOffsetY = 0;//float(cursorPos.y)*m_zoomValue;// - m_progressRectLeft;
+
+					float oldZoomValue = m_zoomValue;
+					if (newBoxHeight != int(m_config.boxHeight))
+					{
+						m_config.boxHeight = newBoxHeight;
+						UpdateProcessFont();
+					}
+
+					m_scrollPosY = Min(0.0f, float(m_scrollPosY - scrollAnchorOffsetY)*(m_zoomValue/oldZoomValue) + scrollAnchorOffsetY);
+					m_scrollPosX = Min(0.0f, float(m_scrollPosX - scrollAnchorOffsetX)*(m_zoomValue/oldZoomValue)*(newScaleValue/horizontalScaleValue) + scrollAnchorOffsetX);//LOWORD(lParam);
+
+					if (horizontalScaleValue != newScaleValue)
+						m_config.horizontalScaleValue = newScaleValue;
+
+
+					UpdateAutoscroll();
+					UpdateSelection();
+
+					int minScroll = r.right - m_contentWidth;
+					m_scrollPosX = Min(0.0f, Max(m_scrollPosX, float(minScroll)));
+					m_scrollPosY = Min(0.0f, Max(m_scrollPosY, float(r.bottom - m_contentHeight)));
+
+					//if (!m_traceView.finished && m_scrollPosX <= minScroll)
+					//	m_autoScroll = true;
+
+
+					if (m_config.ShowReadWriteColors)
+						for (auto& session : m_traceView.sessions)
+							for (auto& processor : session.processors)
+								for (auto& process : processor.processes)
+									//if (TimeToMs(process.writeFilesTime, m_traceView.frequency) >= 300 || TimeToMs(process.createFilesTime, m_traceView.frequency) >= 300)
+									process.bitmapDirty = true;
+				}
+				else if (m_activeSection == 1) // active processes
+				{
+					if (delta < 0)
+						m_config.maxActiveProcessHeight = Max(m_config.maxActiveProcessHeight - 1u, 5u);
+					else if (delta > 0)
+						m_config.maxActiveProcessHeight = Min(m_config.maxActiveProcessHeight + 1u, 32u);
+				}
+				else if (m_activeSection == 0 || m_activeSection == 3) // status/timeline
+				{
+					if (delta < 0)
+						m_config.fontSize -= 1;
+					else if (delta > 0)
+						m_config.fontSize += 1;
+					UpdateDefaultFont();
+				}
+				UpdateScrollbars(true);
+				Redraw(false);
+			}
+			else
+			{
+				RECT r;
+				GetClientRect(hWnd, &r);
+				float oldScrollY = m_scrollPosY;
+				m_scrollPosY = m_scrollPosY + float(delta);
+				m_scrollPosY = Min(Max(m_scrollPosY, float(r.bottom - m_contentHeight)), 0.0f);
+				if (oldScrollY != m_scrollPosY)
+				{
+					UpdateScrollbars(true);
+					Redraw(false);
+				}
+			}
+			break;
+		}
+		case WM_NCHITTEST:
+			if (m_parentHwnd)
+				return HTCLIENT;
+			break;
+		case WM_MOUSEMOVE:
+		{
+			//m_logger.Info(TC("Section: %u"), m_activeSection);
+
+			POINTS p = MAKEPOINTS(lParam);
+			POINT pos{ p.x, p.y };
+			if (m_dragToScrollCounter > 0)
+			{
+				RECT r;
+				GetClientRect(hWnd, &r);
+
+				if (m_contentHeight <= r.bottom)
+					m_scrollPosY = 0;
+				else
+					m_scrollPosY = Max(Min(m_scrollAtAnchorY + float(pos.y - m_mouseAnchor.y), 0.0f), float(r.bottom - m_contentHeight));
+
+				if (m_contentWidth <= r.right)
+					m_scrollPosX = 0;
+				else
+				{
+					int minScroll = r.right - m_contentWidth;
+					m_scrollPosX = Max(Min(m_scrollAtAnchorX + float(pos.x - m_mouseAnchor.x), 0.0f), float(minScroll));
+					if (!m_traceView.finished && m_scrollPosX <= float(minScroll))
+						m_autoScroll = true;
+				}
+				UpdateScrollbars(true);
+				Redraw(false);
+			}
+			else
+			{
+				if (UpdateSelection() || m_config.showCursorLine)
+					Redraw(false);
+				/*
+				else
+				{
+				PaintClient([&](HDC hdc, HDC memDC, RECT& rect)
+				{
+				RECT timelineRect = rect;
+				rect.top = Min(rect.bottom, m_contentHeight) - 28;
+				rect.bottom = Min(rect.bottom, rect.top + 28);
+				FillRect(memDC, &rect, m_backgroundBrush);
+				SetBkMode(memDC, TRANSPARENT);
+				SetTextColor(memDC, m_textColor);
+				SelectObject(memDC, m_font);
+				PaintTimeline(memDC, timelineRect);
+				BitBlt(hdc, 0, rect.top, rect.right - rect.left, rect.bottom - rect.top, memDC, 0, rect.top, SRCCOPY);
+				});
+				}
+				*/
+			}
+
+			TRACKMOUSEEVENT tme;
+			tme.cbSize = sizeof(tme);
+			tme.dwFlags = TME_LEAVE;
+			tme.hwndTrack = hWnd;
+			TrackMouseEvent(&tme);
+			m_mouseOverWindow = true;
+			break;
+		}
+		case WM_MOUSELEAVE:
+			m_mouseOverWindow = false;
+			TRACKMOUSEEVENT tme;
+			tme.cbSize = sizeof(tme);
+			tme.dwFlags = TME_CANCEL;
+			tme.hwndTrack = hWnd;
+			TrackMouseEvent(&tme);
+
+			if (!m_showPopup)
+				UnselectAndRedraw();
+			break;
+
+		case WM_MBUTTONDOWN:
+		{
+			POINTS p = MAKEPOINTS(lParam);
+			StartDragToScroll(POINT{ p.x, p.y });
+			break;
+		}
+
+		case WM_LBUTTONDOWN:
+		{
+			if (!m_parentHwnd && m_traceView.sessions.empty())
+			{
+				RECT r;
+				GetClientRect(hWnd, &r);
+				int centerHrz = r.right/2;
+				int centerVrt = r.bottom/2;
+				r = { centerHrz, centerVrt, centerHrz, centerVrt };
+				InflateRect(&r, 180, 40);
+				POINTS p = MAKEPOINTS(lParam);
+				if (PtInRect(&r, { p.x, p.y }))
+				{
+					OPENFILENAME ofn;
+					tchar szFile[MAX_PATH] = TC("");
+					ZeroMemory(&ofn, sizeof(ofn));
+					ofn.lStructSize = sizeof(ofn);
+					ofn.hwndOwner = m_hwnd;
+					ofn.lpstrFilter = L"Uba Files\0*.uba\0All Files\0*.*\0";
+					ofn.lpstrFile = szFile;
+					ofn.nMaxFile = MAX_PATH;
+					ofn.Flags = OFN_FILEMUSTEXIST | OFN_PATHMUSTEXIST;
+
+					if (GetOpenFileName(&ofn))
+					{
+						m_fileName.Append(szFile);
+						PostNewTrace(0, false);
+					}
+				}
+				break;
+			}
+
+			if (m_buttonSelected != ~0u)
+			{
+				bool* values = &m_config.showProgress;
+				values[m_buttonSelected] = !values[m_buttonSelected];
+				HitTestResult res;
+				HitTest(res, { -1, -1 });
+				UpdateScrollbars(true);
+				Redraw(false);
+			}
+			else if (m_timelineSelected != 0)
+			{
+				if (!m_client) // Does not work for network streams
+				{
+					float timelineSelected = Max(m_timelineSelected, 0.0f);
+					Reset();
+
+					bool changed;
+					u64 time = MsToTime(u64(timelineSelected * 1000.0));
+
+					if (!m_fileName.IsEmpty())
+					{
+						if (!m_trace.ReadFile(m_traceView, m_fileName.data, true))
+							return false;
+					}
+					else
+					{
+						if (!m_trace.StartReadNamed(m_traceView, nullptr, true, true))
+							return false;
+						if (m_traceView.realStartTime + time > m_startTime)
+							time = m_startTime - m_traceView.realStartTime;
+					}
+
+					m_traceView.finished = false;
+					if (!m_fileName.IsEmpty())
+						m_trace.UpdateReadFile(m_traceView, time, changed);
+					else if (m_usingNamed)
+						m_trace.UpdateReadNamed(m_traceView, time, changed);
+
+					m_pauseStart = m_startTime + time;
+					m_pauseTime = m_startTime - m_pauseStart;
+
+					if (!m_paused)
+					{
+						m_autoScroll = true;
+						m_replay = 1;
+						SetTimer(m_hwnd, 0, 200, NULL);
+					}
+					else
+					{
+						m_pauseTime = 0;
+					}
+
+					HitTestResult res;
+					HitTest(res, { -1, -1 });
+
+					RECT r;
+					GetClientRect(hWnd, &r);
+					m_scrollPosX = Min(Max(m_scrollPosX, float(r.right - m_contentWidth)), 0.0f);
+					m_scrollPosY = Min(0.0f, Max(m_scrollPosY, float(r.bottom - m_contentHeight)));
+
+					UpdateScrollbars(true);
+					Redraw(true);
+				}
+			}
+			else if (!m_hyperLinkSelected.empty())
+			{
+				ShellExecuteW(NULL, L"open", m_hyperLinkSelected.c_str(), NULL, NULL, SW_SHOW);
+			}
+			else if (m_sessionSelectedIndex != ~0u)
+			{
+				auto& session = m_traceView.sessions[m_sessionSelectedIndex];
+				for (auto& info : session.infos)
+				{
+					if (info.hyperlink.empty())
+						continue;
+					ShellExecuteW(NULL, L"open", info.hyperlink.c_str(), NULL, NULL, SW_SHOW);
+					break;
+				}
+			}
+			else
+			{
+				POINTS p = MAKEPOINTS(lParam);
+				StartDragToScroll(POINT{ p.x, p.y });
+			}
+			break;
+		}
+
+		case WM_SETCURSOR:
+		{
+			static HCURSOR arrow = LoadCursorW(NULL, IDC_ARROW);
+			static HCURSOR hand = LoadCursorW(NULL, IDC_HAND);
+			bool useHand = false;
+			if (!m_parentHwnd && m_traceView.sessions.empty())
+			{
+				RECT r;
+				GetClientRect(hWnd, &r);
+				int centerHrz = r.right/2;
+				int centerVrt = r.bottom/2;
+				r = { centerHrz, centerVrt, centerHrz, centerVrt };
+				InflateRect(&r, 180, 40);
+				POINT pt;
+				GetCursorPos(&pt);
+				ScreenToClient(hWnd, &pt);
+				useHand = PtInRect(&r, { pt.x, pt.y });
+			}
+			else if (!m_hyperLinkSelected.empty())
+			{
+				useHand = true;
+			}
+			else if (m_sessionSelectedIndex != ~0u)
+			{
+				auto& session = m_traceView.sessions[m_sessionSelectedIndex];
+				for (auto& info : session.infos)
+					if (!info.hyperlink.empty())
+						useHand = true;
+			}
+			SetCursor(useHand ? hand : arrow);
+			break;
+		}
+		case WM_LBUTTONUP:
+		{
+			if (!(m_buttonSelected != ~0u || m_timelineSelected != 0))
+			{
+				StopDragToScroll();
+			}
+			break;
+		}
+
+		case WM_RBUTTONUP:
+		{
+			POINT point;
+			point.x = LOWORD(lParam);
+			point.y = HIWORD(lParam);
+
+			HMENU hMenu = CreatePopupMenu();
+			ClientToScreen(hWnd, &point);
+
+#define UBA_VISUALIZER_FLAG(name, defaultValue, desc) \
+				AppendMenuW(hMenu, MF_STRING | (m_config.name ? MF_CHECKED : 0), Popup_##name, desc);
+			UBA_VISUALIZER_FLAGS2
+#undef UBA_VISUALIZER_FLAG
+
+				AppendMenuW(hMenu, MF_STRING, Popup_IncreaseFontSize, L"&Increase Font Size");
+			AppendMenuW(hMenu, MF_STRING, Popup_DecreaseFontSize, L"&Decrease Font Size");
+
+			AppendMenuW(hMenu, MF_SEPARATOR, 0, NULL);
+
+			if (m_sessionSelectedIndex != ~0u)
+			{
+				AppendMenuW(hMenu, MF_STRING, Popup_CopySessionInfo, L"&Copy Session Info");
+				AppendMenuW(hMenu, MF_SEPARATOR, 0, NULL);
+			}
+			else if (m_processSelected)
+			{
+				const TraceView::Process& process = m_traceView.GetProcess(m_processSelectedLocation);
+
+				AppendMenuW(hMenu, MF_STRING, Popup_CopyProcessInfo, L"&Copy Process Info");
+				if (!process.logLines.empty())
+					AppendMenuW(hMenu, MF_STRING, Popup_CopyProcessLog, L"Copy Process &Log");
+				if (!process.breadcrumbs.empty())
+					AppendMenuW(hMenu, MF_STRING, Popup_CopyProcessBreadcrumbs, L"Copy Process &Breadcrumbs");
+				AppendMenuW(hMenu, MF_SEPARATOR, 0, NULL);
+			}
+			else if (m_workSelected)
+			{
+				AppendMenuW(hMenu, MF_STRING, Popup_CopyWorkInfo, L"&Copy Work Info");
+			}
+
+			if (!m_traceView.sessions.empty())
+			{
+				if (!m_client)
+				{
+					if (!m_replay || m_traceView.finished)
+						AppendMenuW(hMenu, MF_STRING, Popup_Replay, L"&Replay Trace");
+					else
+					{
+						if (m_paused)
+							AppendMenuW(hMenu, MF_STRING, Popup_Play, L"&Play");
+						else
+							AppendMenuW(hMenu, MF_STRING, Popup_Pause, L"&Pause");
+						AppendMenuW(hMenu, MF_STRING, Popup_JumpToEnd, L"&Jump To End");
+					}
+				}
+
+				if (m_fileName.IsEmpty())
+					AppendMenuW(hMenu, MF_STRING, Popup_SaveAs, L"&Save Trace");
+
+				AppendMenuW(hMenu, MF_SEPARATOR, 0, NULL);
+			}
+
+			AppendMenuW(hMenu, MF_STRING, Popup_SaveSettings, L"Save Position/Settings");
+			AppendMenuW(hMenu, MF_STRING, Popup_OpenSettings, L"Open Settings file");
+			AppendMenuW(hMenu, MF_STRING, Popup_Quit, L"&Quit");
+			m_showPopup = true;
+			switch (TrackPopupMenu(hMenu, TPM_RETURNCMD | TPM_RIGHTBUTTON, point.x, point.y, 0, hWnd, NULL))
+			{
+			case Popup_SaveAs:
+			{
+				OPENFILENAME ofn;
+				TCHAR szFile[260] = { 0 };
+
+				// Initialize OPENFILENAME
+				ZeroMemory(&ofn, sizeof(ofn));
+				ofn.lStructSize = sizeof(ofn);
+				ofn.hwndOwner = hWnd;
+				ofn.lpstrFile = szFile;
+				ofn.nMaxFile = sizeof(szFile);
+				ofn.lpstrDefExt = TC("uba");
+				ofn.lpstrFilter = TC("Uba\0*.uba\0All\0*.*\0");
+				ofn.nFilterIndex = 1;
+				ofn.lpstrFileTitle = NULL;
+				ofn.nMaxFileTitle = 0;
+				ofn.lpstrInitialDir = NULL;
+				//ofn.Flags = OFN_PATHMUSTEXIST | OFN_FILEMUSTEXIST;
+				if (GetSaveFileName(&ofn))
+					m_trace.SaveAs(ofn.lpstrFile);
+				break;
+			}
+			case Popup_ShowProcessText:
+				m_config.ShowProcessText = !m_config.ShowProcessText;
+				Redraw(true);
+				break;
+			case Popup_ShowReadWriteColors:
+				m_config.ShowReadWriteColors = !m_config.ShowReadWriteColors;
+				DirtyBitmaps(false);
+				Redraw(true);
+				break;
+			case Popup_ScaleHorizontalWithScrollWheel:
+				m_config.ScaleHorizontalWithScrollWheel = !m_config.ScaleHorizontalWithScrollWheel;
+				break;
+			case Popup_ShowAllTraces:
+				m_config.ShowAllTraces = !m_config.ShowAllTraces;
+				break;
+			case Popup_SortActiveRemoteSessions:
+				m_config.SortActiveRemoteSessions = !m_config.SortActiveRemoteSessions;
+				Redraw(true);
+				break;
+			case Popup_AutoScaleHorizontal:
+				m_config.AutoScaleHorizontal = !m_config.AutoScaleHorizontal;
+				UpdateScrollbars(true);
+				Redraw(true);
+				break;
+			case Popup_LockTimelineToBottom:
+				m_config.LockTimelineToBottom = !m_config.LockTimelineToBottom;
+				Redraw(true);
+				break;
+			case Popup_DarkMode:
+			{
+				m_config.DarkMode = !m_config.DarkMode;
+				DirtyBitmaps(false);
+				InitBrushes();
+				UpdateTheme();
+				Redraw(true);
+				break;
+			}
+			case Popup_AutoSaveSettings:
+				m_config.AutoSaveSettings = !m_config.AutoSaveSettings;
+				break;
+
+			case Popup_Replay:
+				PostNewTrace(1, false);
+				break;
+
+			case Popup_Play:
+				Pause(false);
+				break;
+
+			case Popup_Pause:
+				Pause(true);
+				break;
+
+			case Popup_JumpToEnd:
+				m_traceView.finished = true;
+				PostNewTrace(0, false);
+				break;
+
+			case Popup_SaveSettings:
+				SaveSettings();
+				break;
+
+			case Popup_OpenSettings:
+				ShellExecuteW(NULL, L"open", m_config.filename.c_str(), NULL, NULL, SW_SHOW);
+				break;
+
+			case Popup_Quit: // Quit
+				PostQuit();
+				break;
+
+			case Popup_IncreaseFontSize:
+				ChangeFontSize(1);
+				break;
+			case Popup_DecreaseFontSize:
+				ChangeFontSize(-1);
+				break;
+
+			case Popup_CopySessionInfo:
+			{
+				TString str;
+				auto& session = m_traceView.sessions[m_sessionSelectedIndex];
+				StringBuffer<> fullName;
+				session.GetFullName(fullName);
+				str.append(fullName.data).append(TC("\n"));
+				for (auto& line : session.summary)
+					str.append(line).append(TC("\n"));
+				CopyTextToClipboard(str);
+				break;
+			}
+
+			case Popup_CopyProcessInfo:
+			{
+				TString str;
+				WriteTextLogger logger(str);
+				const TraceView::Process& process = m_traceView.GetProcess(m_processSelectedLocation);
+				WriteProcessStats(logger, process);
+				CopyTextToClipboard(str);
+				break;
+			}
+			case Popup_CopyProcessLog:
+			{
+				TString str;
+				const TraceView::Process& process = m_traceView.GetProcess(m_processSelectedLocation);
+				bool isFirst = true;
+				for (auto line : process.logLines)
+				{
+					if (!isFirst)
+						str += '\n';
+					isFirst = false;
+					str += line.text;
+				}
+				CopyTextToClipboard(str);
+				break;
+			}
+			case Popup_CopyProcessBreadcrumbs:
+			{
+				const TraceView::Process& process = m_traceView.GetProcess(m_processSelectedLocation);
+				CopyTextToClipboard(process.breadcrumbs);
+				break;
+			}
+			case Popup_CopyWorkInfo:
+			{
+				TString str;
+				WriteTextLogger logger(str);
+				auto& record = m_traceView.workTracks[m_workTrack].records[m_workIndex];
+				WriteWorkStats(logger, record);
+				CopyTextToClipboard(str);
+				break;
+			}
+			}
+
+			DestroyMenu(hMenu);
+			m_showPopup = false;
+			UnselectAndRedraw();
+			break;
+		}
+
+		case WM_MBUTTONUP:
+		{
+			StopDragToScroll();
+			//m_processSelected = false;
+			break;
+		}
+		case WM_KEYDOWN:
+		{
+			if (wParam == VK_SPACE)
+				Pause(!m_paused);
+			if (wParam == VK_ADD)
+				++m_replay;
+			if (wParam == VK_SUBTRACT)
+				--m_replay;
+			if (wParam == VK_BACK)
+				if (!m_filterString.empty())
+				{
+					m_filterString.resize(m_filterString.size() -1);
+					Redraw(true);
+				}
+			break;
+		}
+		case WM_CHAR:
+		{
+			tchar c = static_cast<tchar>(wParam);
+			if (c > 32 && c != '\t' && c != '\n')
+				m_filterString += static_cast<tchar>(wParam);
+			Redraw(true);
+		}
+		case WM_VSCROLL:
+		{
+			RECT r;
+			GetClientRect(hWnd, &r);
+			float oldScrollY = m_scrollPosY;
+
+			// HIWORD(wParam) only carries 16-bits, so use GetScrollInfo for larger scroll areas
+			SCROLLINFO scrollInfo = {};
+			scrollInfo.cbSize = sizeof(scrollInfo);
+			scrollInfo.fMask = SIF_TRACKPOS;
+			GetScrollInfo(m_hwnd, SB_VERT, &scrollInfo);
+
+			switch (LOWORD(wParam))
+			{
+			case SB_THUMBTRACK:
+			case SB_THUMBPOSITION:
+				m_scrollPosY = -float(scrollInfo.nTrackPos);
+				break;
+			case SB_PAGEDOWN:
+				m_scrollPosY = m_scrollPosY - float(r.bottom);
+				break;
+			case SB_PAGEUP:
+				m_scrollPosY = m_scrollPosY + float(r.bottom);
+				break;
+			case SB_LINEDOWN:
+				m_scrollPosY -= 30;
+				break;
+			case SB_LINEUP:
+				m_scrollPosY += 30;
+				break;
+			}
+			m_scrollPosY = Min(Max(m_scrollPosY, float(r.bottom - m_contentHeight)), 0.0f);
+
+			if (oldScrollY != m_scrollPosY)
+			{
+				UpdateScrollbars(true);
+				Redraw(false);
+			}
+			return 0;
+		}
+		case WM_HSCROLL:
+		{
+			RECT r;
+			GetClientRect(hWnd, &r);
+			float oldScrollX = m_scrollPosX;
+			bool autoScroll = false;
+
+			// HIWORD(wParam) only carries 16-bits, so use GetScrollInfo for larger scroll areas
+			SCROLLINFO scrollInfo = {};
+			scrollInfo.cbSize = sizeof(scrollInfo);
+			scrollInfo.fMask = SIF_TRACKPOS;
+			GetScrollInfo(m_hwnd, SB_HORZ, &scrollInfo);
+
+			switch (LOWORD(wParam))
+			{
+			case SB_THUMBTRACK:
+				m_scrollPosX = -float(scrollInfo.nTrackPos);
+				if (m_contentWidthWhenThumbTrack == 0)
+					m_contentWidthWhenThumbTrack = m_contentWidth;
+				break;
+			case SB_THUMBPOSITION:
+				autoScroll = m_contentWidthWhenThumbTrack - r.right <= HIWORD(wParam) + 10;
+				m_contentWidthWhenThumbTrack = 0;
+				m_scrollPosX = -float(scrollInfo.nTrackPos);
+				break;
+			case SB_PAGEDOWN:
+				m_scrollPosX = m_scrollPosX - float(r.right);
+				break;
+			case SB_PAGEUP:
+				m_scrollPosX = m_scrollPosX + float(r.right);
+				break;
+			case SB_LINEDOWN:
+				m_scrollPosX -= 30;
+				break;
+			case SB_LINEUP:
+				m_scrollPosX += 30;
+				break;
+			case SB_ENDSCROLL:
+				return 0;
+			}
+
+			int minScroll = r.right - m_contentWidth;
+			m_autoScroll = !m_traceView.finished && (m_scrollPosX <= float(minScroll) || autoScroll);
+			m_scrollPosX = Min(Max(m_scrollPosX, float(r.right - m_contentWidth)), 0.0f);
+
+			if (oldScrollX != m_scrollPosX)
+			{
+				UpdateScrollbars(true);
+				Redraw(false);
+			}
+			return 0;
+		}
+		}
+		return DefWindowProc(hWnd, Msg, wParam, lParam);
+	}
+
+	LRESULT CALLBACK VisualizerWin::StaticWinProc(HWND hWnd, UINT Msg, WPARAM wParam, LPARAM lParam)
+	{
+		auto thisPtr = (VisualizerWin*)GetWindowLongPtr(hWnd, GWLP_USERDATA);
+		if (!thisPtr && Msg == WM_CREATE)
+		{
+			thisPtr = (VisualizerWin*)lParam;
+			SetWindowLongPtr(hWnd, GWLP_USERDATA, lParam);
+			//EnableNonClientDpiScaling(hWnd);
+		}
+		if (thisPtr && hWnd == thisPtr->m_hwnd)
+			return thisPtr->WinProc(hWnd, Msg, wParam, lParam);
+		else
+			return DefWindowProc(hWnd, Msg, wParam, lParam);
+	}
+}
\ No newline at end of file
diff --git a/Engine/Source/Programs/UnrealBuildAccelerator/Visualizer/Private/Windows/UbaVisualizerWin.h b/Engine/Source/Programs/UnrealBuildAccelerator/Visualizer/Private/Windows/UbaVisualizerWin.h
new file mode 100644
--- /dev/null
+++ b/Engine/Source/Programs/UnrealBuildAccelerator/Visualizer/Private/Windows/UbaVisualizerWin.h
@@ -0,0 +1,257 @@
+// Copyright Epic Games, Inc. All Rights Reserved.
+
+#pragma once
+
+#include "UbaNetworkClient.h"
+#include "UbaTraceReader.h"
+#include "UbaVisualizer.h"
+
+namespace uba
+{
+	class VisualizerWin : public Visualizer
+	{
+	public:
+		VisualizerWin(VisualizerConfig& config, Logger& logger);
+		~VisualizerWin();
+
+		bool ShowUsingListener(const wchar_t* channelName);
+		bool ShowUsingNamedTrace(const wchar_t* namedTrace);
+		bool ShowUsingSocket(NetworkBackend& backend, const wchar_t* host, u16 port = DefaultPort);
+		bool ShowUsingFile(const wchar_t* fileName, u32 replay);
+
+		bool HasWindow();
+		HWND GetHwnd();
+		void Lock(bool lock);
+
+	private:
+		bool StartHwndThread();
+		bool Unselect();
+		void Reset();
+		void PaintClient(const Function<void(HDC hdc, HDC memDC, RECT& clientRect)>& paintFunc);
+		void PaintAll(HDC hdc, const RECT& clientRect);
+		void PaintActiveProcesses(int& posY, const RECT& clientRect, const Function<void(TraceView::ProcessLocation&, u32, bool)>& drawProcess);
+		void PaintProcessRect(TraceView::Process& process, HDC hdc, RECT rect, const RECT& progressRect, bool selected, bool writingBitmap, HBRUSH& lastSelectedBrush);
+		void PaintTimeline(HDC hdc, const RECT& clientRect);
+		using DrawTextFunc = Function<void(const StringBufferBase& text, RECT& rect, u32* outWidth)>;
+		void PaintDetailedStats(int& posY, const RECT& progressRect, TraceView::Session& session, bool isRemote, u64 playTime, const DrawTextFunc& drawTextFunc);
+
+		struct Stats
+		{
+			u64 recvBytesPerSecond = 0;
+			u64 sendBytesPerSecond = 0;
+			u64 ping = 0;
+			u64 memAvail = 0;
+			u64 memTotal = 0;
+			u64 recvBytes = 0;
+			u64 sendBytes = 0;
+			u64 procActive = 0;
+			u16 connectionCount = 0;
+			float cpuLoad = 0;
+			struct Drive
+			{
+				u8 busyPercent;
+				u64 readPerSecond = 0;
+				u64 writePerSecond = 0;
+			};
+			Map<char, Drive> drives;
+		};
+		struct HitTestResult
+		{
+			u32 section = ~0u;
+			TraceView::ProcessLocation processLocation;
+			bool processSelected = false;
+			u32 sessionSelectedIndex = ~0u;
+			bool statsSelected = false;
+			Stats stats;
+			u32 buttonSelected = ~0u;
+			float timelineSelected = 0;
+			u32 fetchedFilesSelected = ~0u;
+			bool workSelected = false;
+			u32 workTrack = ~0u;
+			u32 workIndex = ~0u;
+			bool activeProcessGraphSelected = false;
+			u16 activeProcessCount = 0u;
+			TString hyperLink;
+		};
+		u64 GetPlayTime();
+		int GetTimelineHeight();
+		int GetTimelineTop(const RECT& clientRect);
+		void HitTest(HitTestResult& outResult, const POINT& pos);
+
+		void WriteProcessStats(Logger& out, const TraceView::Process& process);
+		void WriteWorkStats(Logger& out, const TraceView::WorkRecord& record);
+		void CopyTextToClipboard(const TString& str);
+		void UnselectAndRedraw();
+		bool UpdateAutoscroll();
+		bool UpdateSelection();
+		void UpdateScrollbars(bool redraw);
+		StringBufferBase& GetTitlePrefix(StringBufferBase& out);
+		void InitBrushes();
+		void ThreadLoop();
+		void Pause(bool pause);
+		void StartDragToScroll(const POINT& anchor);
+		void StopDragToScroll();
+		void SaveSettings();
+		void DirtyBitmaps(bool full);
+		void PrintCacheWriteStats(Logger& logger, u32 processId);
+
+		struct Font
+		{
+			HFONT handle = 0;
+			HFONT handleUnderlined = 0;
+			int height = 0;
+			int offset = 0;
+		};
+
+		void UpdateFont(Font& font, int height, bool createUnderline);
+		void UpdateDefaultFont();
+		void UpdateProcessFont();
+		void ChangeFontSize(int offset);
+		void Redraw(bool now);
+		void SetActiveFont(const Font& font);
+		void UpdateTheme();
+		bool ActiveProcessesShouldFillHeight();
+		StringBuffer<128> GetWorldTime(u64 time);
+		StringBuffer<128> GetWorldTime(float seconds);
+
+		StringBuffer<256> m_namedTrace;
+		StringBuffer<256> m_fileName;
+		u32 m_replay = 0;
+		u64 m_startTime = 0;
+		u64 m_pauseTime = 0;
+
+		u64 m_lastPaintTimeMs;
+
+		struct ProcessBrushes
+		{
+			HBRUSH inProgress = 0;
+			HBRUSH inProgressRace = 0;
+			HBRUSH success = 0;
+			HBRUSH error = 0;
+			HBRUSH returned = 0;
+			HBRUSH recv = 0;
+			HBRUSH send = 0;
+			HBRUSH cacheFetchTest = 0;
+			HBRUSH cacheFetchDownload = 0;
+			HBRUSH warning = 0;
+		};
+
+		ProcessBrushes m_processBrushes[2]; // Non-selected and selected
+
+		Atomic<bool> m_looping = false;
+		HWND m_hwnd = 0;
+		HWND m_parentHwnd = 0;
+		COLORREF m_textColor = {};
+		COLORREF m_textWarningColor = {};
+		COLORREF m_textErrorColor = {};
+		COLORREF m_sendColor = {};
+		COLORREF m_recvColor = {};
+		COLORREF m_cpuColor = {};
+		COLORREF m_memColor = {};
+		COLORREF m_driveColor = {};
+		COLORREF m_activeProcColor = {};
+		HBRUSH m_backgroundBrush = 0;
+		HBRUSH m_tooltipBackgroundBrush = 0;
+		HPEN m_textPen = 0;
+		HPEN m_separatorPen = 0;
+		HPEN m_sendPen = 0;
+		HPEN m_recvPen = 0;
+		HPEN m_cpuPen = 0;
+		HPEN m_memPen = 0;
+		HPEN m_drivePen = 0;
+		HPEN m_activeProcPen = 0;
+		HPEN m_processUpdatePen = 0;
+		HPEN m_checkboxPen = 0;
+		int m_sessionStepY = 0;
+
+		Font m_defaultFont;
+		Font m_processFont;
+		Font m_timelineFont;
+		Font m_popupFont;
+
+		int m_processFontOffsetY = 0;
+
+		Font m_activeProcessFont[32];
+		u32 m_activeProcessCountHistory[5];
+		u32 m_activeProcessCountHistoryIterator = 0;
+
+		HDC m_activeHdc = 0;
+		Font m_activeFont;
+
+		int m_progressRectLeft = 30;
+
+		NetworkClient* m_client = nullptr;
+		Event m_clientDisconnect;
+
+		StringBuffer<256>m_listenChannel;
+		StringBuffer<256> m_newTraceName;
+		Event m_listenTimeout;
+
+		int m_contentWidth = 0;
+		int m_contentHeight = 0;
+
+		int m_contentWidthWhenThumbTrack = 0;
+
+		float m_scrollPosX = 0;
+		float m_scrollPosY = 0;
+		float m_zoomValue = 0.5f;
+		float m_horizontalScaleValue = 0.5f;
+		bool m_autoScroll = true;
+		bool m_paused = false;
+		u64 m_pauseStart = 0;
+
+		bool m_isInPaint = false;
+
+		static constexpr int BitmapCacheHeight = 1024*1024;
+		HBITMAP m_lastBitmap = 0;
+		int m_lastBitmapOffset = BitmapCacheHeight;
+
+		u32 m_activeSection = ~0u;
+		TraceView::ProcessLocation m_processSelectedLocation;
+		bool m_processSelected = false;
+		u32 m_sessionSelectedIndex = ~0u;
+		bool m_statsSelected = false;
+		bool m_activeProcessGraphSelected = false;
+		u64 m_activeProcessCount = 0u;
+		Stats m_stats;
+		u32 m_buttonSelected = ~0u;
+		float m_timelineSelected = 0;
+		u32 m_fetchedFilesSelected = ~0u;
+		TString m_hyperLinkSelected;
+
+		bool m_workSelected = false;
+		u32 m_workTrack = ~0u;
+		u32 m_workIndex = ~0u;
+		
+		bool m_usingNamed = false;
+		bool m_mouseOverWindow = false;
+		bool m_showPopup = false;
+		bool m_locked = false;
+		bool m_handlingTimer = false;
+
+		HBITMAP m_cachedBitmap = 0;
+		RECT m_cachedBitmapRect = { INT_MIN, INT_MIN, INT_MIN, INT_MIN };
+
+		Vector<HBITMAP> m_textBitmaps;
+
+		POINT m_mouseAnchor = {};
+		float m_scrollAtAnchorX = 0;
+		float m_scrollAtAnchorY = 0;
+		int m_dragToScrollCounter = 0;
+
+		bool m_horizontalScrollBarEnabled = true;
+		bool m_verticalScrollBarEnabled = true;
+
+		TString m_filterString;
+
+		Thread m_thread;
+
+		UnorderedMap<Color, HBRUSH> m_coloredBrushes;
+
+		void PostNewTrace(u32 replay, bool paused);
+		void PostNewTitle(const StringView& title);
+		void PostQuit();
+		LRESULT WinProc(HWND hWnd, UINT Msg, WPARAM wParam, LPARAM lParam);
+		static LRESULT CALLBACK StaticWinProc(HWND hWnd, UINT Msg, WPARAM wParam, LPARAM lParam);
+	};
+}
