[UBA]
* Fixed so visualizer can compile on posix (it will just exit straight away)

diff --git a/Engine/Source/Programs/UnrealBuildAccelerator/UbaVisualizer.Target.cs b/Engine/Source/Programs/UnrealBuildAccelerator/UbaVisualizer.Target.cs
--- a/Engine/Source/Programs/UnrealBuildAccelerator/UbaVisualizer.Target.cs
+++ b/Engine/Source/Programs/UnrealBuildAccelerator/UbaVisualizer.Target.cs
@@ -2,7 +2,6 @@
 
 using UnrealBuildTool;
 
-[SupportedPlatforms("Win64")]
 public class UbaVisualizerTarget : TargetRules
 {
 	public UbaVisualizerTarget(TargetInfo Target) : base(Target)
diff --git a/Engine/Source/Programs/UnrealBuildAccelerator/Visualizer/Private/PosixOS/UbaVisualizerApp.cpp b/Engine/Source/Programs/UnrealBuildAccelerator/Visualizer/Private/PosixOS/UbaVisualizerApp.cpp
new file mode 100644
--- /dev/null
+++ b/Engine/Source/Programs/UnrealBuildAccelerator/Visualizer/Private/PosixOS/UbaVisualizerApp.cpp
@@ -0,0 +1,6 @@
+
+
+int main(int argc, char* argv[])
+{
+	return 0;
+}
\ No newline at end of file
diff --git a/Engine/Source/Programs/UnrealBuildAccelerator/Visualizer/Private/UbaTraceReader.cpp b/Engine/Source/Programs/UnrealBuildAccelerator/Visualizer/Private/UbaTraceReader.cpp
--- a/Engine/Source/Programs/UnrealBuildAccelerator/Visualizer/Private/UbaTraceReader.cpp
+++ b/Engine/Source/Programs/UnrealBuildAccelerator/Visualizer/Private/UbaTraceReader.cpp
@@ -93,11 +93,9 @@ namespace uba
 		Unmap();
 	}
 
-#if PLATFORM_WINDOWS
-
 	u64 ConvertTime(const TraceView& view, u64 time)
 	{
-		return u64(time * view.timeMultiplier);
+		return u64(double(time) * view.timeMultiplier);
 	}
 
 	bool TraceReader::ReadFile(TraceView& out, const tchar* fileName, bool replay)
@@ -273,7 +271,10 @@ namespace uba
 
 		if (*namedTrace && m_namedTrace != namedTrace)
 		{
+#if PLATFORM_WINDOWS
 			m_memoryHandle.mh = ::OpenFileMappingW(PAGE_READWRITE, false, namedTrace);
+#endif
+
 			if (!m_memoryHandle.IsValid())
 			{
 				if (!silentFail)
@@ -303,11 +304,14 @@ namespace uba
 		m_memoryEnd = m_memoryBegin + *(u32*)m_memoryBegin;
 		outChanged = !m_activeProcesses.empty() || m_memoryPos != m_memoryEnd;
 		bool keepRunning = ReadMemory(out, true, maxTime);
+
+#if PLATFORM_WINDOWS
 		if (m_hostProcess && WaitForSingleObject(m_hostProcess, 0) != WAIT_TIMEOUT)
 		{
 			StopAllActive(out, GetTime() - out.realStartTime);
 			keepRunning = false;
 		}
+#endif
 
 		if (keepRunning || m_namedTrace.empty())
 			return true;
@@ -354,9 +358,11 @@ namespace uba
 			bool replay = maxTime != ~u64(0);
 
 			u32 hostProcessId = traceReader.ReadU32();
+#if PLATFORM_WINDOWS
 			m_hostProcess = 0;
 			if (trackHost && !replay)
 				m_hostProcess = OpenProcess(SYNCHRONIZE, FALSE, hostProcessId);
+#endif
 			u64 traceSystemStartTimeUs = 0;
 
 			if (version < TraceReadCompatibilityVersion || version > TraceVersion)
@@ -1524,9 +1530,11 @@ namespace uba
 
 	void TraceReader::Unmap()
 	{
+#if PLATFORM_WINDOWS
 		if (m_hostProcess)
 			CloseHandle(m_hostProcess);
 		m_hostProcess = 0;
+#endif
 		if (m_memoryBegin)
 			UnmapViewOfFile(m_logger, m_memoryBegin, 0, TC("TraceReader"));
 		m_memoryBegin = nullptr;
@@ -1672,6 +1680,4 @@ namespace uba
 		}
 		return { info.data, hyperlink };
 	}
-
-#endif // PLATFORM_WINDOWS
 }
diff --git a/Engine/Source/Programs/UnrealBuildAccelerator/Visualizer/Private/UbaTraceReader.h b/Engine/Source/Programs/UnrealBuildAccelerator/Visualizer/Private/UbaTraceReader.h
--- a/Engine/Source/Programs/UnrealBuildAccelerator/Visualizer/Private/UbaTraceReader.h
+++ b/Engine/Source/Programs/UnrealBuildAccelerator/Visualizer/Private/UbaTraceReader.h
@@ -6,6 +6,10 @@
 #include "UbaStats.h"
 #include "UbaTrace.h"
 
+#if !PLATFORM_WINDOWS
+typedef void* HBITMAP;
+#endif
+
 namespace uba
 {
 	class NetworkClient;
@@ -225,8 +229,6 @@ namespace uba
 		TraceReader(Logger& logger);
 		~TraceReader();
 
-		#if PLATFORM_WINDOWS
-
 		// Use for file read
 		bool ReadFile(TraceView& out, const tchar* fileName, bool replay);
 		bool UpdateReadFile(TraceView& out, u64 maxTime, bool& outChanged);
@@ -264,8 +266,6 @@ namespace uba
 
 		Vector<u32> m_sessionIndexToSession;
 
-		#endif
-
 		Logger& m_logger;
 		TraceChannel m_channel;
 		FileMappingHandle m_memoryHandle;
@@ -273,7 +273,10 @@ namespace uba
 		u8* m_memoryBegin = nullptr;
 		u8* m_memoryPos = nullptr;
 		u8* m_memoryEnd = nullptr;
-		HANDLE m_hostProcess = NULL;
 		Futex m_memoryFutex;
+
+		#if PLATFORM_WINDOWS
+		HANDLE m_hostProcess = NULL;
+		#endif
 	};
 }
\ No newline at end of file
diff --git a/Engine/Source/Programs/UnrealBuildAccelerator/Visualizer/Private/UbaVisualizer.h b/Engine/Source/Programs/UnrealBuildAccelerator/Visualizer/Private/UbaVisualizer.h
--- a/Engine/Source/Programs/UnrealBuildAccelerator/Visualizer/Private/UbaVisualizer.h
+++ b/Engine/Source/Programs/UnrealBuildAccelerator/Visualizer/Private/UbaVisualizer.h
@@ -2,8 +2,6 @@
 
 #pragma once
 
-#include "UbaNetworkClient.h"
-#include "UbaThread.h"
 #include "UbaTraceReader.h"
 #include "UbaVisualizerConfig.h"
 
diff --git a/Engine/Source/Programs/UnrealBuildAccelerator/Visualizer/Private/UbaVisualizerApp.cpp b/Engine/Source/Programs/UnrealBuildAccelerator/Visualizer/Private/UbaVisualizerApp.cpp
deleted file mode 100644
diff --git a/Engine/Source/Programs/UnrealBuildAccelerator/Visualizer/Private/UbaVisualizerConfig.cpp b/Engine/Source/Programs/UnrealBuildAccelerator/Visualizer/Private/UbaVisualizerConfig.cpp
--- a/Engine/Source/Programs/UnrealBuildAccelerator/Visualizer/Private/UbaVisualizerConfig.cpp
+++ b/Engine/Source/Programs/UnrealBuildAccelerator/Visualizer/Private/UbaVisualizerConfig.cpp
@@ -24,21 +24,21 @@ namespace uba
 			#endif
 			return false;
 		}
-		config.GetValueAsInt(x, L"X");
-		config.GetValueAsInt(y, L"Y");
-		config.GetValueAsU32(width, L"Width");
-		config.GetValueAsU32(height, L"Height");
-		config.GetValueAsU32(fontSize, L"FontSize");
-		config.GetValueAsString(fontName, L"FontName");
-		config.GetValueAsU32(maxActiveVisible, L"MaxActiveVisible");
-		config.GetValueAsU32(maxActiveProcessHeight, L"MaxActiveProcessHeight");
+		config.GetValueAsInt(x, TC("X"));
+		config.GetValueAsInt(y, TC("Y"));
+		config.GetValueAsU32(width, TC("Width"));
+		config.GetValueAsU32(height, TC("Height"));
+		config.GetValueAsU32(fontSize, TC("FontSize"));
+		config.GetValueAsString(fontName, TC("FontName"));
+		config.GetValueAsU32(maxActiveVisible, TC("MaxActiveVisible"));
+		config.GetValueAsU32(maxActiveProcessHeight, TC("MaxActiveProcessHeight"));
 
-		config.GetValueAsU32(boxHeight, L"BoxHeight");
+		config.GetValueAsU32(boxHeight, TC("BoxHeight"));
 		u64 temp;
-		if (config.GetValueAsU64(temp, L"HorizontalScrollValue"))
+		if (config.GetValueAsU64(temp, TC("HorizontalScrollValue")))
 			horizontalScaleValue = float(double(temp) / 10000000.0);
 
-#define UBA_VISUALIZER_FLAG(name, defaultValue, desc) config.GetValueAsBool(show##name, L"Show" #name);
+#define UBA_VISUALIZER_FLAG(name, defaultValue, desc) config.GetValueAsBool(show##name, TC("Show" #name));
 		UBA_VISUALIZER_FLAGS1
 		#undef UBA_VISUALIZER_FLAG
 		#define UBA_VISUALIZER_FLAG(name, defaultValue, desc) config.GetValueAsBool(name, TC(#name));
@@ -52,19 +52,19 @@ namespace uba
 	bool VisualizerConfig::Save(Logger& logger)
 	{
 		Config config;
-		config.AddValue(L"X", x);
-		config.AddValue(L"Y", y);
-		config.AddValue(L"Width", width);
-		config.AddValue(L"Height", height);
-		config.AddValue(L"FontSize", fontSize);
-		config.AddValue(L"FontName", fontName.c_str());
-		config.AddValue(L"MaxActiveVisible", maxActiveVisible);
-		config.AddValue(L"MaxActiveProcessHeight", maxActiveProcessHeight);
+		config.AddValue(TC("X"), x);
+		config.AddValue(TC("Y"), y);
+		config.AddValue(TC("Width"), width);
+		config.AddValue(TC("Height"), height);
+		config.AddValue(TC("FontSize"), fontSize);
+		config.AddValue(TC("FontName"), fontName.c_str());
+		config.AddValue(TC("MaxActiveVisible"), maxActiveVisible);
+		config.AddValue(TC("MaxActiveProcessHeight"), maxActiveProcessHeight);
 
-		config.AddValue(L"BoxHeight", boxHeight);
-		config.AddValue(L"HorizontalScrollValue", u64(double(horizontalScaleValue)*10000000.0));
+		config.AddValue(TC("BoxHeight"), boxHeight);
+		config.AddValue(TC("HorizontalScrollValue"), u64(double(horizontalScaleValue)*10000000.0));
 
-		#define UBA_VISUALIZER_FLAG(name, defaultValue, desc) config.AddValue(L"Show" #name, show##name);
+		#define UBA_VISUALIZER_FLAG(name, defaultValue, desc) config.AddValue(TC("Show" #name), show##name);
 		UBA_VISUALIZER_FLAGS1
 		#undef UBA_VISUALIZER_FLAG
 		#define UBA_VISUALIZER_FLAG(name, defaultValue, desc) config.AddValue(TC(#name), name);
diff --git a/Engine/Source/Programs/UnrealBuildAccelerator/Visualizer/Private/Windows/UbaVisualizerApp.cpp b/Engine/Source/Programs/UnrealBuildAccelerator/Visualizer/Private/Windows/UbaVisualizerApp.cpp
new file mode 100644
--- /dev/null
+++ b/Engine/Source/Programs/UnrealBuildAccelerator/Visualizer/Private/Windows/UbaVisualizerApp.cpp
@@ -0,0 +1,348 @@
+// Copyright Epic Games, Inc. All Rights Reserved.
+
+#include "UbaConfig.h"
+#include "UbaNetworkBackendTcp.h"
+#include "UbaVersion.h"
+#include "Windows/UbaVisualizerWin.h"
+#include <tlhelp32.h>
+#include <psapi.h>
+
+using namespace uba;
+
+static int PrintHelp(const tchar* message = nullptr)
+{
+	StringBuffer<64*1024> s;
+	if (message && *message)
+		s.Appendf(TC("%s\r\n\r\n"), message);
+
+	s.Appendf(TC("\r\n"));
+	s.Appendf(TC("------------------------\r\n"));
+	s.Appendf(TC("   UbaVisualizer v%s\r\n"), GetVersionString());
+	s.Appendf(TC("------------------------\r\n"));
+	s.Appendf(TC("\r\n"));
+	s.Appendf(TC("  When started UbaVisualizer will keep trying to connect to provided host address or named memory buffer.\r\n"));
+	s.Appendf(TC("  Once connected it will start visualizing. Nothing else is needed :)\r\n"));
+	s.Appendf(TC("\r\n"));
+	s.Appendf(TC("  -host=<host>[:<port>] The ip/name and port (default: %u) of the machine we want to connect to\r\n"), DefaultPort);
+	s.Appendf(TC("  -named=<name>        Name of named memory to connect to\r\n"));
+	s.Appendf(TC("  -file=<name>         Name of file to parse\r\n"));
+	s.Appendf(TC("  -listen[=<channel>]  Listen for announcements of new sessions. Defaults to channel '%s'\r\n"), TC("Default"));
+	s.Appendf(TC("  -replay              Visualize the data as if it was running right now\r\n"));
+	s.Appendf(TC("  -config=<file>       Specify config file to use\r\n"));
+	s.Appendf(TC("  -parent=<hwnd>       Specify hwnd this window should be a child of\r\n"));
+	s.Appendf(TC("  -nocopy              Will prevent UbaVisualizer.exe from being copied to temp and executed from there\r\n"));
+	s.Appendf(TC("\r\n"));
+	MessageBox(NULL, s.data, TC("UbaVisualizer"), 0);
+	//wprintf(s.data);
+	return -1;
+}
+
+struct MessageBoxLogWriter : public LogWriter
+{
+	virtual void BeginScope() override
+	{
+	}
+
+	virtual void EndScope() override
+	{
+	}
+
+	virtual void Log(LogEntryType type, const tchar* str, u32 strLen, const tchar* prefix = nullptr, u32 prefixLen = 0) override
+	{
+		if (type > LogEntryType_Warning)
+		{
+			#if UBA_DEBUG
+			StringBuffer<> buf;
+			buf.Append(str, strLen).Append(TCV("\r\n"));
+			OutputDebugStringW(buf.data);
+			#endif
+			return;
+		}
+
+		HWND hwnd = NULL;
+		if (m_visualizer)
+		{
+			hwnd = m_visualizer->GetHwnd();
+			m_visualizer->Lock(true);
+		}
+
+		UINT flags = type == LogEntryType_Error ? MB_ICONERROR : MB_ICONWARNING;
+		if (!hwnd)
+			flags |= MB_TOPMOST;
+		MessageBox(hwnd, str, TC("UbaVisualizer"), flags);
+		if (type == LogEntryType_Error)
+			ExitProcess(~0u);
+		if (m_visualizer)
+			m_visualizer->Lock(false);
+	}
+	
+	VisualizerWin* m_visualizer = nullptr;
+};
+
+int WINAPI wWinMain(_In_ HINSTANCE hInstance, _In_opt_ HINSTANCE hPrevInstance, _In_ PWSTR pCmdLine, _In_ int nShowCmd)
+{
+	StringBuffer<> host; // 192.168.86.49
+	StringBuffer<> named;
+	StringBuffer<> file;
+	StringBuffer<> channel;
+	StringBuffer<> configPath;
+	u32 port = DefaultPort;
+	u32 replay = 0;
+	u64 parent = 0;
+	bool copyAndLaunch = true;
+
+	int argc;
+	auto argv = CommandLineToArgvW(GetCommandLine(), &argc);
+
+	for (int i=1; i!=argc; ++i)
+	{
+		StringBuffer<> name;
+		StringBuffer<> value;
+
+		if (i == 1 && argv[i][0] != '-')
+		{
+			file.Append(argv[i]);
+			continue;
+		}
+
+		if (const tchar* equals = wcschr(argv[i],'='))
+		{
+			name.Append(argv[i], equals - argv[i]);
+			value.Append(equals+1);
+		}
+		else
+		{
+			name.Append(argv[i]);
+		}
+
+		if (name.Equals(TCV("-help")))
+		{
+			return PrintHelp();
+		}
+		else if (name.Equals(TCV("-host")))
+		{
+			if (const tchar* portIndex = value.First(':'))
+			{
+				StringBuffer<> portStr(portIndex + 1);
+				if (!portStr.Parse(port))
+					return PrintHelp(TC("Invalid value for port in -host"));
+				value.Resize(portIndex - value.data);
+			}
+			if (value.IsEmpty())
+				return PrintHelp(TC("-host needs a name/ip"));
+			host.Append(value);
+		}
+		else if (name.Equals(TCV("-named")))
+		{
+			if (value.IsEmpty())
+				return PrintHelp(TC("-named needs a value"));
+			named.Append(value);
+		}
+		else if (name.Equals(TCV("-file")))
+		{
+			if (value.IsEmpty())
+				return PrintHelp(TC("-file needs a value"));
+			file.Append(value);
+		}
+		else if (name.Equals(TCV("-port")))
+		{
+			if (!value.Parse(port))
+				return PrintHelp(TC("Invalid value for -port"));
+		}
+		else if (name.Equals(TCV("-listen")))
+		{
+			if (!value.IsEmpty())
+				channel.Append(value.data);
+			else
+				channel.Append(TCV("Default"));
+		}
+		else if (name.Equals(TCV("-replay")))
+		{
+			replay = 1;
+			if (!value.IsEmpty())
+				value.Parse(replay);
+		}
+		else if (name.Equals(TCV("-config")))
+		{
+			if (value.IsEmpty())
+				return PrintHelp(TC("-config needs a value"));
+			configPath.Append(value);
+		}
+		else if (name.Equals(TCV("-parent")))
+		{
+			if (value.IsEmpty())
+				return PrintHelp(TC("-parent needs a value"));
+			if (value.count > 8)
+				return PrintHelp(TC("-parent has invalid value"));
+			value.MakeLower();
+			if (value.count & 1) // uneven char, add 0 in the front
+			{
+				memmove(value.data + 1, value.data, (value.count+1)*sizeof(tchar));
+				value.data[0] = '0';
+				++value.count;
+			}
+			parent = StringToValue2(value.data, value.count);
+			//value.ParseHex(parent);
+		}
+		else if (name.Equals(TCV("-nocopy")))
+		{
+			copyAndLaunch = false;
+		}
+		else if (name.Equals(TCV("-ownerPid")))
+		{
+			u32 ownerPid;
+			if (value.Parse(ownerPid))
+				const_cast<OwnerInfo&>(GetOwnerInfo()).pid = ownerPid;
+		}
+		else if (name.Equals(TCV("-ownerId")))
+		{
+			TStrcpy_s(const_cast<tchar*>(GetOwnerInfo().id), 260, value.data);
+		}
+		else
+		{
+			StringBuffer<> msg;
+			msg.Appendf(TC("Unknown argument '%s'"), name.data);
+			return PrintHelp(msg.data);
+		}
+	}
+
+	MessageBoxLogWriter logWriter;
+	LoggerWithWriter logger(logWriter);
+
+	if (copyAndLaunch)
+	{
+		StringBuffer<> tempPath;
+		tempPath.count = GetTempPathW(tempPath.capacity, tempPath.data);
+		if (!tempPath.count)
+		{
+			logger.Error(TC("GetTempPathW failed"));
+			return -1;
+		}
+
+		StringBuffer<> thisExe;
+		thisExe.count = GetModuleFileNameW(NULL, thisExe.data, thisExe.capacity);
+		if (!thisExe.count)
+		{
+			logger.Error(TC("GetModuleFileNameW failed"));
+			return -1;
+		}
+
+		WIN32_FILE_ATTRIBUTE_DATA data;
+		if (!GetFileAttributesExW(thisExe.data, GetFileExInfoStandard, &data))
+		{
+			logger.Error(TC("GetFileAttributesExW failed"));
+			return -1;
+		}
+
+		u64 thisLastWriteTime = *(u64*)&data.ftLastWriteTime;
+
+		StringBuffer<> ubaFileName;
+		for (u32 i=0; i!=10; ++i)
+		{
+			ubaFileName.Append(tempPath).EnsureEndsWithSlash().Appendf(TC("UbaVisualizer%u.exe"), i);
+			if (GetFileAttributesExW(ubaFileName.data, GetFileExInfoStandard, &data))
+			{
+				u64 lastWriteTime = *(u64*)&data.ftLastWriteTime;
+				if (thisLastWriteTime == lastWriteTime)
+					break;
+			}
+			if (CopyFileW(thisExe.data, ubaFileName.data, false))
+				break;
+			ubaFileName.Clear();
+		}
+
+		if (!ubaFileName.count)
+		{
+			logger.Error(TC("Failed to create temporary UbaVisualizer.exe to launch."));
+			return -1;
+		}
+
+		StringBuffer<> args;
+		args.Append(ubaFileName);
+		for (int i = 1; i != argc; ++i)
+		{
+			args.Append(' ');
+			if (TStrchr(argv[i], ' '))
+				args.Append('\"').Append(argv[i]).Append('\"');
+			else
+				args.Append(argv[i]);
+		}
+		args.Append(" -nocopy");
+
+		OwnerInfo ownerInfo = GetOwnerInfo();
+		if (ownerInfo.pid)
+			args.Appendf(TC(" -ownerPid=%u -ownerId=%s"), ownerInfo.pid, ownerInfo.id);
+
+		STARTUPINFOW si;
+		ZeroMemory(&si, sizeof(si));
+		si.cb = sizeof(si);
+		PROCESS_INFORMATION pi;
+		ZeroMemory(&pi, sizeof(pi));
+		if (!CreateProcessW(NULL, args.data, NULL, NULL, false, 0, NULL, NULL, &si, &pi))
+		{
+			logger.Error(TC("Failed to launch process %s"), ubaFileName.data);
+			return -1;
+		}
+		CloseHandle(pi.hProcess);
+		CloseHandle(pi.hThread);
+		return 0;
+	}
+
+	LocalFree(argv);
+
+	if (host.IsEmpty() && named.IsEmpty() && file.IsEmpty() && !channel.count)
+		channel.Append(TCV("Default")); // return PrintHelp(TC("No host/named/file provided. Add -host=<host> or -file=<file> or -named=<name>"));
+
+	bool showAllTraces = true;
+	if (!configPath.count)
+	{
+		configPath.count = ExpandEnvironmentStringsW(L"%PROGRAMDATA%", configPath.data, configPath.capacity) - 1;
+		configPath.Append(L"\\Epic\\UbaVisualizer\\UbaVisualizer");
+
+		OwnerInfo ownerInfo = GetOwnerInfo();
+		if (ownerInfo.pid)
+		{
+			configPath.Append('_').Append(ownerInfo.id);
+			showAllTraces = false;
+		}
+		configPath.Append(L".toml");
+	}
+
+	VisualizerConfig visualizerConfig(configPath.data);
+	visualizerConfig.parent = parent;
+	visualizerConfig.ShowAllTraces = showAllTraces;
+	visualizerConfig.Load(logger);
+
+	NetworkBackendTcp networkBackend(logWriter);
+	VisualizerWin visualizer(visualizerConfig, logger);
+	logWriter.m_visualizer = &visualizer;
+
+	if (channel.count)
+	{
+		if (!visualizer.ShowUsingListener(channel.data))
+			logger.Error(TC("Failed listening to named pipe"));
+	}
+	else if (!named.IsEmpty())
+	{
+		if (!visualizer.ShowUsingNamedTrace(named.data))
+			logger.Error(TC("Failed reading from mapped memory %s"), named.data);
+	}
+	else if (!host.IsEmpty())
+	{
+		if (!visualizer.ShowUsingSocket(networkBackend, host.data, u16(port)))
+			logger.Error(TC("Failed to connect to %s:%u"), host.data, port);
+	}
+	else
+	{
+		if (!visualizer.ShowUsingFile(file.data, replay))
+			logger.Error(TC("Failed to read trace file '%s'"), file.data);
+	}
+
+	while (true)
+	{
+		if (!visualizer.HasWindow())
+			break;
+		uba::Sleep(500);
+	}
+	return 0;
+}
