[UBA]
* Fixed linux compile warnings in scheduler

diff --git a/Engine/Source/Programs/UnrealBuildAccelerator/Common/Private/UbaScheduler.cpp b/Engine/Source/Programs/UnrealBuildAccelerator/Common/Private/UbaScheduler.cpp
--- a/Engine/Source/Programs/UnrealBuildAccelerator/Common/Private/UbaScheduler.cpp
+++ b/Engine/Source/Programs/UnrealBuildAccelerator/Common/Private/UbaScheduler.cpp
@@ -179,7 +179,7 @@ namespace uba
 	Scheduler::Scheduler(const SchedulerCreateInfo& info)
 	:	m_session(info.session)
 	,	m_maxLocalProcessors(info.maxLocalProcessors != ~0u ? info.maxLocalProcessors : GetLogicalProcessorCount())
-	,	m_maxRacingWeight(info.maxRacingPercent*float(m_maxLocalProcessors)/100.0f)
+	,	m_maxRacingWeight(float(info.maxRacingPercent)*float(m_maxLocalProcessors)/100.0f)
 	,	m_maxParallelCacheQueries(info.maxParallelCacheQueries)
 	,	m_updateThreadWakeup(EventResetType_Auto)
 	,	m_enableProcessReuse(info.enableProcessReuse)
@@ -233,7 +233,7 @@ namespace uba
 			while (!m_memThread.Wait(1) && !m_memEstimatedUsage)
 				;
 		}
-		else if (m_maxRacingWeight)
+		else if (m_maxRacingWeight != 0)
 		{
 			m_memThreadEvent.Create(EventResetType_Manual);
 			m_memThread.Start([this]() { ThreadMemoryCheckLoop();  return 0; }, TC("UbaMemTrackLoop"));
@@ -275,7 +275,7 @@ namespace uba
 			m_memThread.Wait();
 			SaveMemTrackTable();
 		}
-		else if (m_maxRacingWeight)
+		else if (m_maxRacingWeight != 0)
 		{
 			m_memThreadEvent.Set();
 			m_memThread.Wait();
@@ -1057,7 +1057,7 @@ namespace uba
 
 	bool Scheduler::RaceRemoteProcess()
 	{
-		if (!m_activeRemoteProcesses || !m_maxRacingWeight)
+		if (!m_activeRemoteProcesses || m_maxRacingWeight == 0)
 			return false;
 
 		float totalWeight = m_activeLocalProcessWeight + m_activeRacingProcessWeight;
