[UBA]
* Added copyright to posix UbaVisualizerApp.cpp

diff --git a/Engine/Source/Programs/UnrealBuildAccelerator/Visualizer/Private/PosixOS/UbaVisualizerApp.cpp b/Engine/Source/Programs/UnrealBuildAccelerator/Visualizer/Private/PosixOS/UbaVisualizerApp.cpp
--- a/Engine/Source/Programs/UnrealBuildAccelerator/Visualizer/Private/PosixOS/UbaVisualizerApp.cpp
+++ b/Engine/Source/Programs/UnrealBuildAccelerator/Visualizer/Private/PosixOS/UbaVisualizerApp.cpp
@@ -1,4 +1,4 @@
-
+// Copyright Epic Games, Inc. All Rights Reserved.
 
 int main(int argc, char* argv[])
 {
