[UBA]
* Removed all usage of HDC and friends inside visualizer's large paint function.

diff --git a/Engine/Source/Programs/UnrealBuildAccelerator/Visualizer/Private/UbaTraceReader.h b/Engine/Source/Programs/UnrealBuildAccelerator/Visualizer/Private/UbaTraceReader.h
--- a/Engine/Source/Programs/UnrealBuildAccelerator/Visualizer/Private/UbaTraceReader.h
+++ b/Engine/Source/Programs/UnrealBuildAccelerator/Visualizer/Private/UbaTraceReader.h
@@ -19,7 +19,14 @@ namespace uba
 	public:
 		enum ProcessType : u8 { ProcessType_Normal, ProcessType_CacheFetchTest, ProcessType_CacheFetchDownload, ProcessType_Task };
 
-		struct Process
+		struct BitmapCache
+		{
+			HBITMAP bitmap = 0;
+			u32 bitmapOffset = 0;
+			bool bitmapDirty = true;
+		};
+
+		struct Process : BitmapCache
 		{
 			u32 id = 0;
 			u32 exitCode = ~0u;
@@ -28,10 +35,7 @@ namespace uba
 			TString description;
 			TString returnedReason;
 			TString breadcrumbs;
-			HBITMAP bitmap = 0;
-			u32 bitmapOffset = 0;
 			ProcessType type = ProcessType_Normal;
-			bool bitmapDirty = true;
 			bool isRemote = false;
 			bool isReuse = false;
 			bool isRacing = false;
@@ -54,16 +58,13 @@ namespace uba
 			u32 count = 1;
 		};
 
-		struct WorkRecord
+		struct WorkRecord : BitmapCache
 		{
 			const tchar* description = nullptr;
 			u64 start = 0;
 			u64 stop = 0;
 			Vector<WorkRecordLogEntry> entries;
-			HBITMAP bitmap = 0;
-			u32 bitmapOffset = 0;
 			u32 color = ColorWork;
-			bool bitmapDirty = true;
 		};
 
 		struct WorkTrack
diff --git a/Engine/Source/Programs/UnrealBuildAccelerator/Visualizer/Private/UbaVisualizer.cpp b/Engine/Source/Programs/UnrealBuildAccelerator/Visualizer/Private/UbaVisualizer.cpp
--- a/Engine/Source/Programs/UnrealBuildAccelerator/Visualizer/Private/UbaVisualizer.cpp
+++ b/Engine/Source/Programs/UnrealBuildAccelerator/Visualizer/Private/UbaVisualizer.cpp
@@ -10,4 +10,15 @@ namespace uba
 	,	m_trace(logger)
 	{
 	}
+
+	u64 Visualizer::GetPlayTime()
+	{
+		u64 currentTime = m_paused ? m_pauseStart : GetTime();
+		u64 playTime = 0;
+		if (m_traceView.startTime && currentTime > (m_traceView.startTime + m_pauseTime))
+			playTime = currentTime - m_traceView.startTime - m_pauseTime;
+		if (m_replay)
+			playTime *= m_replay;
+		return playTime;
+	}
 }
\ No newline at end of file
diff --git a/Engine/Source/Programs/UnrealBuildAccelerator/Visualizer/Private/UbaVisualizer.h b/Engine/Source/Programs/UnrealBuildAccelerator/Visualizer/Private/UbaVisualizer.h
--- a/Engine/Source/Programs/UnrealBuildAccelerator/Visualizer/Private/UbaVisualizer.h
+++ b/Engine/Source/Programs/UnrealBuildAccelerator/Visualizer/Private/UbaVisualizer.h
@@ -12,9 +12,60 @@ namespace uba
 	public:
 		Visualizer(VisualizerConfig& config, Logger& logger);
 
+
+	protected:
+		u64 GetPlayTime();
+
+
+		struct Stats
+		{
+			u64 recvBytesPerSecond = 0;
+			u64 sendBytesPerSecond = 0;
+			u64 ping = 0;
+			u64 memAvail = 0;
+			u64 memTotal = 0;
+			u64 recvBytes = 0;
+			u64 sendBytes = 0;
+			u64 procActive = 0;
+			u16 connectionCount = 0;
+			float cpuLoad = 0;
+			struct Drive
+			{
+				u8 busyPercent;
+				u64 readPerSecond = 0;
+				u64 writePerSecond = 0;
+			};
+			Map<char, Drive> drives;
+		};
+		struct HitTestResult
+		{
+			u32 section = ~0u;
+			TraceView::ProcessLocation processLocation;
+			bool processSelected = false;
+			u32 sessionSelectedIndex = ~0u;
+			bool statsSelected = false;
+			Stats stats;
+			u32 buttonSelected = ~0u;
+			float timelineSelected = 0;
+			u32 fetchedFilesSelected = ~0u;
+			bool workSelected = false;
+			u32 workTrack = ~0u;
+			u32 workIndex = ~0u;
+			bool activeProcessGraphSelected = false;
+			u16 activeProcessCount = 0u;
+			TString hyperLink;
+		};
+
+
 		Logger& m_logger;
 		VisualizerConfig m_config;
 		TraceReader m_trace;
 		TraceView m_traceView;
+
+		bool m_paused = false;
+		u32 m_replay = 0;
+		u64 m_startTime = 0;
+		u64 m_pauseTime = 0;
+		u64 m_pauseStart = 0;
 	};
 }
diff --git a/Engine/Source/Programs/UnrealBuildAccelerator/Visualizer/Private/Windows/UbaVisualizerWin.cpp b/Engine/Source/Programs/UnrealBuildAccelerator/Visualizer/Private/Windows/UbaVisualizerWin.cpp
--- a/Engine/Source/Programs/UnrealBuildAccelerator/Visualizer/Private/Windows/UbaVisualizerWin.cpp
+++ b/Engine/Source/Programs/UnrealBuildAccelerator/Visualizer/Private/Windows/UbaVisualizerWin.cpp
@@ -53,27 +53,182 @@ namespace uba
 		VisualizerFlag_Count
 	};
 
-	class DrawTextLogger : public Logger
+	struct VisualizerWin::PaintContext
 	{
-	public:
-		DrawTextLogger(HWND hw, HDC h, int fh, HBRUSH bb)
-			:	hwnd(hw)
-			,	hdc(h)
+		void SetBkMode(int mode) { ::SetBkMode(hdc, mode); }
+		void SetTextColor(COLORREF color) { ::SetTextColor(hdc, color); }
+		void SetPen(HPEN pen) { ::SelectObject(hdc, pen); }
+		void SetFont(HFONT font) { ::SelectObject(hdc, font); }
+		u32 SetTextAlign(u32 align) { return ::SetTextAlign(hdc, align); }
+		void ExtTextOutW(int x, int y, UINT options, CONST RECT * lprect, StringView str) { ::ExtTextOutW(hdc, x, y, options, lprect, str.data, str.count, NULL); }
+		
+		void SetActiveFont(const Font& font)
+		{
+			owner.m_activeFont = font;
+			::SelectObject(hdc, font.handle);
+		}
+
+		Font& GetActiveFont() { return owner.m_activeFont; }
+
+		void SetBrush(HBRUSH b)
+		{
+			if (lastSelectedBrush != b)
+			{
+				SelectObject(hdc, b);
+				lastSelectedBrush = b;
+			}
+		}
+
+		void FillRect(const RECT& r, HBRUSH b)
+		{
+			SetBrush(b);
+			PatBlt(hdc, r.left, r.top, r.right - r.left, r.bottom - r.top, PATCOPY);
+		}
+
+		void Rectangle(int left, int top, int right, int bottom)
+		{
+			SetBrush((HBRUSH)GetStockObject(NULL_BRUSH));
+			::Rectangle(hdc, left, top, right, bottom);
+		}
+
+		void Line(int fromX, int fromY, int toX, int toY)
+		{
+			::MoveToEx(hdc, fromX, fromY, NULL);
+			::LineTo(hdc, toX, toY);
+		}
+
+		void Polyline(const Vector<POINT>& polyLine)
+		{
+			::Polyline(hdc, polyLine.data(), int(polyLine.size()));
+		}
+
+		void PolyPolyline(const POINT* points, const DWORD* lines, u32 lineCount)
+		{
+			::PolyPolyline(hdc, points, lines, int(lineCount));
+		}
+
+		void DrawCenteredText(const StringBuffer<>* rows, u32 rowCount, bool drawBackground)
+		{
+			for (u32 i=0; i!=rowCount; ++i)
+			{ 
+				SIZE textSize = GetTextSize(rows[i]);
+
+				RECT rect = { 0, 0, 0, 0 };
+				int top = (clientRect.bottom - textSize.cy)/2 + i*20;
+				rect.left = (clientRect.right - textSize.cx)/2;
+				rect.right = (clientRect.right + textSize.cx)/2;
+				rect.bottom = top + textSize.cy;
+				rect.top = top;
+
+				if (drawBackground)
+				{
+					RECT r = rect;
+					InflateRect(&r, 4, 2);
+					FillRect(r, owner.m_tooltipBackgroundBrush);
+				}
+				ExtTextOutW(rect.left, rect.top, 0, NULL, rows[i]);
+			}
+		}
+
+		SIZE GetTextSize(StringView str)
+		{
+			SIZE s;
+			return GetTextExtentPoint32W(hdc, str.data, str.count, &s) ? s : SIZE();
+		}
+
+
+		void DrawCached(TraceView::BitmapCache& process, int processWidth, int offsetY, int textHeight, int left, RECT rect, RECT progressRect, const Function<void(PaintContext& textCacheContext)>& func)
+		{
+			if (!process.bitmap || process.bitmapDirty)
+			{
+				if (!process.bitmap)
+				{
+					if (owner.m_lastBitmapOffset == BitmapCacheHeight)
+					{
+						if (owner.m_lastBitmap)
+							owner.m_textBitmaps.push_back(owner.m_lastBitmap);
+
+						owner.m_lastBitmapOffset = 0;
+						owner.m_lastBitmap = CreateCompatibleBitmap(hdc, 256, BitmapCacheHeight);
+					}
+					process.bitmap = owner.m_lastBitmap;
+					process.bitmapOffset = owner.m_lastBitmapOffset;
+					owner.m_lastBitmapOffset += owner.m_processFont.height;
+				}
+				if (lastSelectedBitmap != process.bitmap)
+				{
+					SelectObject(textDC, process.bitmap);
+					lastSelectedBitmap = process.bitmap;
+				}
+
+				PaintContext textCacheContext{owner, NULL, textDC, clientRect };
+				func(textCacheContext);
+			}
+
+			if (lastSelectedBitmap != process.bitmap)
+			{
+				SelectObject(textDC, process.bitmap);
+				lastSelectedBitmap = process.bitmap;
+			}
+
+			int width = Min(processWidth, 256);
+			int bitmapOffsetY = process.bitmapOffset;
+			int bltOffsetY = offsetY;
+			if (bltOffsetY < 0)
+			{
+				bitmapOffsetY -= bltOffsetY;
+				bltOffsetY = 0;
+			}
+			int height = Min(textHeight, owner.m_processFont.height);
+			if (bltOffsetY + height > textHeight)
+				height = textHeight - bltOffsetY;
+
+			if (left > -256 && height >= 0)
+			{
+				int bitmapOffsetX = rect.left - left;
+
+				if (left < progressRect.left)
+				{
+					int diff = progressRect.left - left;
+					rect.left = progressRect.left;
+					width -= diff;
+					bitmapOffsetX += diff;
+				}
+				BitBlt(hdc, rect.left, rect.top + bltOffsetY, width, height, textDC, bitmapOffsetX, bitmapOffsetY, SRCCOPY);
+			}
+		}
+
+
+
+
+		VisualizerWin& owner;
+		HWND hwnd;
+		HDC hdc;
+		const RECT& clientRect;
+
+		HDC textDC = NULL;
+		HBRUSH lastSelectedBrush = 0;
+		HBITMAP lastSelectedBitmap = 0;
+	};
+
+	struct VisualizerWin::DrawTextLogger : public Logger
+	{
+		DrawTextLogger(PaintContext& c, int fh, HBRUSH bb)
+			:	context(c)
 			,	fontHeight(fh)
 			,	backgroundBrush(bb)
 		{
-			textColor = GetTextColor(hdc);
+			textColor = GetTextColor(c.hdc);
 		}
 
 		virtual void BeginScope() override {}
 		virtual void EndScope() override {}
 		virtual void Log(LogEntryType type, const wchar_t* str, u32 strLen) override
 		{
-			RECT textRect{0,0,0,0};
-			DrawTextW(hdc, str, strLen, &textRect, DT_CALCRECT);
+			u32 textWidth = context.GetTextSize(StringView(str, strLen)).cx;
 
 			lines.emplace_back(TString(str, strLen), textOffset, height, textColor);
-			width = Max(width, int(textRect.right + textOffset));
+			width = Max(width, int(textWidth + textOffset));
 			height += fontHeight;
 		}
 
@@ -101,7 +256,7 @@ namespace uba
 			r.bottom = r.top + height;
 
 			RECT clientRect;
-			GetClientRect(hwnd, &clientRect);
+			GetClientRect(context.hwnd, &clientRect);
 
 			if (r.right > clientRect.right)
 				OffsetRect(&r, -width - 15, 0);
@@ -114,15 +269,15 @@ namespace uba
 
 			RECT fillRect = r;
 			fillRect.right += 2 + extraWidth;
-			FillRect(hdc, &fillRect, backgroundBrush);
+			context.FillRect(fillRect, backgroundBrush);
 
 			for (auto& line : lines)
 			{
 				RECT tr = r;
 				tr.left += line.left;
 				tr.top += line.top;
-				SetTextColor(hdc, line.color);
-				DrawTextW(hdc, line.str.data(), u32(line.str.size()), &tr, DT_SINGLELINE|DT_NOPREFIX);
+				context.SetTextColor(line.color);
+				context.ExtTextOutW(tr.left, tr.top, 0, NULL, line.str);
 			}
 		}
 
@@ -130,7 +285,7 @@ namespace uba
 		{
 			POINT p;
 			GetCursorPos(&p);
-			ScreenToClient(hwnd, &p);
+			ScreenToClient(context.hwnd, &p);
 			p.x += 3;
 			p.y += 3;
 			DrawAtPos(p.x, p.y);
@@ -145,8 +300,7 @@ namespace uba
 		struct Line { TString str; int left; int top; COLORREF color; };
 		Vector<Line> lines;
 
-		HWND hwnd;
-		HDC hdc;
+		PaintContext& context;
 		int fontHeight;
 		HBRUSH backgroundBrush;
 		COLORREF textColor;
@@ -921,48 +1075,12 @@ namespace uba
 			});
 	}
 
-	void VisualizerWin::PaintAll(HDC hdc, const RECT& clientRect)
+	void VisualizerWin::PaintAll(PaintContext& context, const RECT& clientRect)
 	{
 		u64 playTime = GetPlayTime();
 
-		SetBkMode(hdc, TRANSPARENT);
-		SetTextColor(hdc, m_textColor);
-		SetBkColor(hdc, m_config.DarkMode ? RGB(70, 70, 70) : RGB(180, 180, 180));
-
-		HBRUSH lastSelectedBrush = 0;
-		HBRUSH textLastSelectedBrush = 0;
-		auto fillRect = [&](const RECT& r, HBRUSH b)
-			{
-				if (lastSelectedBrush != b)
-				{
-					SelectObject(hdc, b);
-					lastSelectedBrush = b;
-				}
-				PatBlt(hdc, r.left, r.top, r.right - r.left, r.bottom - r.top, PATCOPY);
-			};
-
-		auto DrawCenteredText = [&](const StringBuffer<>* lines, u32 lineCount, bool drawBackground)
-			{
-				for (u32 i=0; i!=lineCount; ++i)
-				{
-					RECT rect = { 0, 0, 0, 0 };
-					DrawTextW(hdc, lines[i].data, lines[i].count, &rect, DT_SINGLELINE|DT_NOPREFIX|DT_NOCLIP|DT_CALCRECT);
-					int top = (clientRect.bottom - rect.bottom)/2 + i*20;
-					rect.left = (clientRect.right - rect.right)/2;
-					rect.right = (clientRect.right + rect.right)/2;
-					rect.bottom = top + (rect.bottom - rect.top);
-					rect.top = top;
-
-					if (drawBackground)
-					{
-						RECT r = rect;
-						InflateRect(&r, 4, 2);
-						fillRect(r, m_tooltipBackgroundBrush);
-					}
-					DrawTextW(hdc, lines[i].data, lines[i].count, &rect, DT_SINGLELINE|DT_NOPREFIX|DT_NOCLIP);
-				}
-			};
-
+		context.SetBkMode(TRANSPARENT);
+		context.SetTextColor(m_textColor);
 
 		if (m_traceView.sessions.empty() && !m_parentHwnd)
 		{
@@ -982,8 +1100,8 @@ namespace uba
 				else
 					str[0].Append(TCV("No trace active"));
 			}
-			SetActiveFont(m_popupFont);
-			DrawCenteredText(str, 2, false);
+			context.SetActiveFont(m_popupFont);
+			context.DrawCenteredText(str, 2, false);
 			return;
 		}
 
@@ -999,45 +1117,33 @@ namespace uba
 			progressRect.bottom -= m_defaultFont.height + 10;
 		}
 
-		HDC textDC = CreateCompatibleDC(hdc);
-		SetTextColor(textDC, m_textColor);
-		SelectObject(textDC, m_processFont.handle);
-		SelectObject(textDC, GetStockObject(NULL_BRUSH));
-		SetBkMode(textDC, TRANSPARENT);
-
-		//TEXTMETRIC metric;
-		//GetTextMetrics(textDC, &metric);
-
-		HBITMAP nullBmp = CreateCompatibleBitmap(hdc, 1, 1);
-		HBITMAP oldBmp = (HBITMAP)SelectObject(textDC, nullBmp);
-		HBITMAP lastSelectedBitmap = 0;
-
 		u64 lastStop = 0;
 
-		SetActiveFont(m_defaultFont);
+		context.SetActiveFont(m_defaultFont);
 
 		auto drawStatusText = [&](StringView text, LogEntryType type, int posX, int endX, bool moveY, bool underlined = false, bool centered = false)
 			{
+				Font& activeFont = context.GetActiveFont();
 				RECT rect;
 				rect.left = posX;
 				rect.right = endX;
-				rect.top = posY + m_activeFont.offset;
-				rect.bottom = posY + m_activeFont.height + 2;
-				SetTextColor(hdc, type == LogEntryType_Info ? m_textColor : (type == LogEntryType_Error ? m_textErrorColor : m_textWarningColor));
+				rect.top = posY + activeFont.offset;
+				rect.bottom = posY + activeFont.height + 2;
+				context.SetTextColor(type == LogEntryType_Info ? m_textColor : (type == LogEntryType_Error ? m_textErrorColor : m_textWarningColor));
 
 				u32 align = 0;
 				if (centered)
 				{
 					posX += (endX - posX)/2;
-					align = SetTextAlign(hdc, TA_CENTER);
+					align = context.SetTextAlign(TA_CENTER);
 				}
 				if (underlined)
-					SelectObject(m_activeHdc, m_activeFont.handleUnderlined);
-				ExtTextOutW(hdc, posX, posY, ETO_CLIPPED, &rect, text.data, text.count, NULL);
+					context.SetFont(activeFont.handleUnderlined);
+				context.ExtTextOutW(posX, posY, ETO_CLIPPED, &rect, text);
 				if (underlined)
-					SelectObject(m_activeHdc, m_activeFont.handle);
+					context.SetFont(activeFont.handle);
 				if (centered)
-					SetTextAlign(hdc, align);
+					context.SetTextAlign(align);
 				if (moveY)
 					posY = rect.bottom;
 			};
@@ -1052,19 +1158,20 @@ namespace uba
 		{
 			drawIndentedText(AsView(L"Progress"), LogEntryType_Info, 1, false);
 
+			Font& activeFont = context.GetActiveFont();
 
 			float progress = float(m_traceView.progressProcessesDone) / float(m_traceView.progressProcessesTotal);
-			u32 start = 3 + 6*m_activeFont.height;
-			u32 width = m_activeFont.height * 18;
+			u32 start = 3 + 6*activeFont.height;
+			u32 width = activeFont.height * 18;
 			RECT rect;
 			rect.left = start;
 			rect.right = start + width;
 			rect.top = posY;
-			rect.bottom = posY + m_activeFont.height;
-			fillRect(rect, m_processBrushes[0].inProgress);
+			rect.bottom = posY + activeFont.height;
+			context.FillRect(rect, m_processBrushes[0].inProgress);
 
 			rect.right = rect.left + int(progress*float(width));//durationMs/100;
-			fillRect(rect, m_traceView.progressErrorCount ? m_processBrushes[0].error : m_processBrushes[0].success);
+			context.FillRect(rect, m_traceView.progressErrorCount ? m_processBrushes[0].error : m_processBrushes[0].success);
 
 			StringBuffer<> str;
 			str.Appendf(TC("%u%%    %u / %u"), u32(progress*100.0f), m_traceView.progressProcessesDone, m_traceView.progressProcessesTotal);
@@ -1078,15 +1185,15 @@ namespace uba
 			if (m_traceView.lastKillProcessTime && playTime >= m_traceView.lastKillProcessTime && playTime < m_traceView.lastKillProcessTime + MsToTime(800))
 			{
 				rect.left = start + width + 20;
-				rect.right = rect.left + m_activeFont.height * 7;
-				fillRect(rect, m_processBrushes[0].error);
+				rect.right = rect.left + activeFont.height * 7;
+				context.FillRect(rect, m_processBrushes[0].error);
 				drawStatusText(TCV("Killing (mem)"), LogEntryType_Info, rect.left, rect.right, 0, false, true);
 			}
 			else if (m_traceView.lastSpawningDelayStartTime && playTime >= m_traceView.lastSpawningDelayStartTime && playTime < m_traceView.lastSpawningDelayEndTime + MsToTime(500))
 			{
 				rect.left = start + width + 20;
-				rect.right = rect.left + m_activeFont.height * 7;
-				fillRect(rect, m_processBrushes[0].warning);
+				rect.right = rect.left + activeFont.height * 7;
+				context.FillRect(rect, m_processBrushes[0].warning);
 				drawStatusText(TCV("Throttling (mem)"), LogEntryType_Info, rect.left, rect.right, 0, false, true);
 			}
 
@@ -1098,13 +1205,10 @@ namespace uba
 		{
 			if (!m_traceView.finished)
 			{
-				SetTextColor(hdc, m_textWarningColor);
 				StringBuffer<> str;
 				str.Appendf(TC("Unsupported trace version %u (Versions supported are %u to %u)"), m_traceView.version, TraceReadCompatibilityVersion, TraceVersion);
-				HFONT handle = CreateFontW(-20, 0, 0, 0, FW_NORMAL, 0, 0, 0, ANSI_CHARSET, OUT_DEFAULT_PRECIS, CLIP_DEFAULT_PRECIS, CLEARTYPE_QUALITY, DEFAULT_PITCH, m_config.fontName.c_str());
-				SelectObject(hdc, handle);
-				DrawTextW(hdc, str.data, str.count, (RECT*)&clientRect, DT_CENTER|DT_NOPREFIX|DT_SINGLELINE|DT_VCENTER);
-				DeleteObject(handle);
+				context.SetTextColor(m_textWarningColor);
+				context.DrawCenteredText(&str, 1, false);
 			}
 			else
 			{
@@ -1119,6 +1223,7 @@ namespace uba
 		{
 			u32 lastRow = ~0u;
 			u32 row = ~0u;
+			Font& activeFont = context.GetActiveFont();
 			for (auto& kv : m_traceView.statusMap)
 			{
 				auto& status = kv.second;
@@ -1126,15 +1231,15 @@ namespace uba
 					continue;
 				row = u32(kv.first >> 32);
 				if (lastRow != ~0u && lastRow != row)
-					posY += m_activeFont.height + 2;
+					posY += activeFont.height + 2;
 				lastRow = row;
 				u32 column = u32(kv.first & ~0u);
 				drawIndentedText(status.text, status.type, column, false, !status.link.empty());
 			}
 			if (row != ~0u)
-				posY += m_activeFont.height + 2;
+				posY += activeFont.height + 2;
 
-			SetTextColor(hdc, m_textColor);
+			context.SetTextColor(m_textColor);
 			posY += 3;
 		}
 
@@ -1147,7 +1252,7 @@ namespace uba
 
 				Vector<POINT> line;
 
-				SelectObject(hdc, m_activeProcPen);
+				context.SetPen(m_activeProcPen);
 				bool isFirstUpdate = true;
 				bool isFirstDraw = true;
 				u64 prevValue = 0;
@@ -1209,7 +1314,7 @@ namespace uba
 					time += 0.25f;
 				}
 				if (line.size() > 1)
-					Polyline(hdc, line.data(), int(line.size()));
+					context.Polyline(line);
 			}
 			posY += GraphHeight;
 		}
@@ -1232,13 +1337,13 @@ namespace uba
 					rect.right = right;
 					rect.top = posY;
 					rect.bottom = posY + height;
-					fillRect(rect, inProgress ? m_processBrushes[selected].inProgress : m_processBrushes[selected].success);
+					context.FillRect(rect, inProgress ? m_processBrushes[selected].inProgress : m_processBrushes[selected].success);
 					return rect;
 				};
 
 			m_activeProcessCountHistory[m_activeProcessCountHistoryIterator++ % sizeof_array(m_activeProcessCountHistory)] = u32(m_trace.m_activeProcesses.size());
 
-			PaintActiveProcesses(posY, clientRect, [&](TraceView::ProcessLocation& processLocation, u32 boxHeight, bool firstWithHeight)
+			PaintActiveProcesses(context, posY, clientRect, [&](TraceView::ProcessLocation& processLocation, u32 boxHeight, bool firstWithHeight)
 				{
 					auto& session = m_trace.GetSession(m_traceView, processLocation.sessionIndex);
 					TraceView::Process& process = session.processors[processLocation.processorIndex].processes[processLocation.processIndex];
@@ -1269,7 +1374,7 @@ namespace uba
 							m_activeProcessFont[fontIndex].offset += 1;
 						}
 						if (firstWithHeight)
-							SetActiveFont(m_activeProcessFont[fontIndex]);
+							context.SetActiveFont(m_activeProcessFont[fontIndex]);
 
 						StringBuffer<> str;
 						auto firstPar = -1;//process.description.find_first_of('(');
@@ -1327,9 +1432,8 @@ namespace uba
 			{
 				if (posY + stepY >= progressRect.top && posY <= progressRect.bottom)
 				{
-					SelectObject(hdc, m_separatorPen);
-					MoveToEx(hdc, 0, posY, NULL);
-					LineTo(hdc, clientRect.right, posY);
+					context.SetPen(m_separatorPen);
+					context.Line(0, posY, clientRect.right, posY);
 
 					StringBuffer<> text;
 					session.GetFullName(text);
@@ -1366,14 +1470,14 @@ namespace uba
 					rect.bottom = textBottom;
 
 					if (selected)
-						SetBkMode(hdc, OPAQUE);
+						context.SetBkMode(OPAQUE);
 
-					SIZE textSize;
-					if (GetTextExtentPoint32W(hdc, text.data, text.count, &textSize))
+					SIZE textSize = context.GetTextSize(text);
+					if (textSize.cx)
 						session.fullNameWidth = u32(textSize.cx);
-					ExtTextOutW(hdc, 5, posY+2, ETO_CLIPPED, &rect, text.data, text.count, NULL);
+					context.ExtTextOutW(5, posY+2, ETO_CLIPPED, &rect, text);
 					if (selected)
-						SetBkMode(hdc, TRANSPARENT);
+						context.SetBkMode(TRANSPARENT);
 				}
 				posY += m_sessionStepY;
 			}
@@ -1395,7 +1499,7 @@ namespace uba
 							u32 reconnectIndex = 0;
 							while (loop)
 							{
-								SelectObject(hdc, pen);
+								context.SetPen(pen);
 								bool isFirstUpdate = true;
 								bool isFirstDraw = true;
 								decltype(maxValue) prevValue = 0;
@@ -1471,7 +1575,7 @@ namespace uba
 									prevTime = updateTime;
 								}
 								if (line.size() > 1)
-									Polyline(hdc, line.data(), int(line.size()));
+									context.Polyline(line);
 								++reconnectIndex;
 							}
 						};
@@ -1505,28 +1609,24 @@ namespace uba
 							return;
 						bool selected = m_fetchedFilesSelected == processLocation.sessionIndex && text.StartsWith(TC("Fetched Files"));
 						if (selected)
-							SetBkMode(hdc, OPAQUE);
+							context.SetBkMode(OPAQUE);
 						if (rect.bottom > progressRect.bottom)
 						{
 							rect.bottom = progressRect.bottom;
-							ExtTextOutW(hdc, rect.left, rect.top, ETO_CLIPPED, &rect, text.data, text.count, NULL);
+							context.ExtTextOutW(rect.left, rect.top, ETO_CLIPPED, &rect, text);
 						}
 						else
-							ExtTextOutW(hdc, rect.left, rect.top, 0, NULL, text.data, text.count, NULL);
+							context.ExtTextOutW(rect.left, rect.top, 0, NULL, text);
 						if (selected)
-							SetBkMode(hdc, TRANSPARENT);
+							context.SetBkMode(TRANSPARENT);
 						if (outWidth)
-						{
-							SIZE s;
-							GetTextExtentPoint32W(hdc, text.data, text.count, &s);
-							*outWidth = u32(s.cx);
-						}
+							*outWidth = u32(context.GetTextSize(text).cx);
 					};
 				bool isRemote = processLocation.sessionIndex != 0;
 				PaintDetailedStats(posY, progressRect, session, isRemote, playTime, drawText);
 			}
 
-			SetActiveFont(m_processFont);
+			context.SetActiveFont(m_processFont);
 
 			bool shouldDrawText = m_processFont.height > 4;
 
@@ -1590,93 +1690,41 @@ namespace uba
 								process.bitmapDirty = true;
 
 							--rect.top;
-							PaintProcessRect(process, hdc, rect, progressRect, selected, false, lastSelectedBrush);
+							PaintProcessRect(context, process, rect, progressRect, selected, false);
 							++rect.top;
 
 							int processWidth = rect.right - rect.left;
 							if (shouldDrawText && m_config.ShowProcessText && processWidth > 3)
 							{
-								if (!process.bitmap || process.bitmapDirty)
-								{
-									if (!process.bitmap)
-									{
-										if (m_lastBitmapOffset == BitmapCacheHeight)
-										{
-											if (m_lastBitmap)
-												m_textBitmaps.push_back(m_lastBitmap);
-
-											m_lastBitmapOffset = 0;
-											m_lastBitmap = CreateCompatibleBitmap(hdc, 256, BitmapCacheHeight);
-										}
-										process.bitmap = m_lastBitmap;
-										process.bitmapOffset = m_lastBitmapOffset;
-										m_lastBitmapOffset += m_processFont.height;
-									}
-									if (lastSelectedBitmap != process.bitmap)
+								context.DrawCached(process, processWidth, offsetY, textHeight, left, rect, progressRect, [&](PaintContext& textCacheContext)
 									{
-										SelectObject(textDC, process.bitmap);
-										lastSelectedBitmap = process.bitmap;
-									}
-
-									RECT rect2{ 0, int(process.bitmapOffset), 256, int(process.bitmapOffset) + m_processFont.height };
-									RECT rect3{ 0, int(process.bitmapOffset), processWidth, int(process.bitmapOffset) + m_processFont.height };
-									if (!done)
-										rect3.right = 256;
-
-									PaintProcessRect(process, textDC, rect3, rect2, selected, true, textLastSelectedBrush);
-
-									rect2.left += 3; // Move in text a bit
-
-									int textY = rect2.top + m_processFont.offset;
-
-									bool dropShadow = m_config.DarkMode;
-									if (dropShadow)
-									{
-										SetTextColor(textDC, RGB(5, 60, 5));
-										++rect2.left;
-										++textY;
-										ExtTextOutW(textDC, rect2.left, textY, ETO_CLIPPED, &rect2, process.description.c_str(), int(process.description.size()), NULL);
-										--rect2.left;
-										--textY;
-									}
-									SetTextColor(textDC, m_textColor);
-									ExtTextOutW(textDC, rect2.left, textY, ETO_CLIPPED, &rect2, process.description.c_str(), int(process.description.size()), NULL);
+										RECT rect2{ 0, int(process.bitmapOffset), 256, int(process.bitmapOffset) + m_processFont.height };
+										RECT rect3{ 0, int(process.bitmapOffset), processWidth, int(process.bitmapOffset) + m_processFont.height };
+										if (!done)
+											rect3.right = 256;
 
-									if (!selected)
-										process.bitmapDirty = false;
-								}
+										PaintProcessRect(textCacheContext, process, rect3, rect2, selected, true);
 
-								if (lastSelectedBitmap != process.bitmap)
-								{
-									SelectObject(textDC, process.bitmap);
-									lastSelectedBitmap = process.bitmap;
-								}
+										rect2.left += 3; // Move in text a bit
 
-								int width = Min(processWidth, 256);
-								int bitmapOffsetY = process.bitmapOffset;
-								int bltOffsetY = offsetY;
-								if (bltOffsetY < 0)
-								{
-									bitmapOffsetY -= bltOffsetY;
-									bltOffsetY = 0;
-								}
-								int height = Min(textHeight, m_processFont.height);
-								if (bltOffsetY + height > textHeight)
-									height = textHeight - bltOffsetY;
+										int textY = rect2.top + m_processFont.offset;
 
-								if (left > -256 && height >= 0)
-								{
-									int bitmapOffsetX = rect.left - left;
+										bool dropShadow = m_config.DarkMode;
+										if (dropShadow)
+										{
+											textCacheContext.SetTextColor(RGB(5, 60, 5));
+											++rect2.left;
+											++textY;
+											textCacheContext.ExtTextOutW(rect2.left, textY, ETO_CLIPPED, &rect2, process.description);
+											--rect2.left;
+											--textY;
+										}
+										textCacheContext.SetTextColor(m_textColor);
+										textCacheContext.ExtTextOutW(rect2.left, textY, ETO_CLIPPED, &rect2, process.description);
 
-									if (left < progressRect.left)
-									{
-										int diff = progressRect.left - left;
-										rect.left = progressRect.left;
-										width -= diff;
-										bitmapOffsetX += diff;
-									}
-									BitBlt(hdc, rect.left, rect.top + bltOffsetY, width, height, textDC, bitmapOffsetX, bitmapOffsetY, SRCCOPY);
-								}
+										if (!selected)
+											process.bitmapDirty = false;
+									});
 							}
 						}
 
@@ -1690,7 +1738,7 @@ namespace uba
 
 							StringBuffer<> buf;
 							buf.AppendValue(u64(processLocation.processorIndex) + 1);
-							ExtTextOutW(hdc, 5, posY + offsetY, ETO_CLIPPED, &rect, buf.data, buf.count, NULL);
+							context.ExtTextOutW(5, posY + offsetY, ETO_CLIPPED, &rect, buf);
 						}
 					}
 
@@ -1739,7 +1787,7 @@ namespace uba
 
 							StringBuffer<> buf;
 							buf.AppendValue(u64(trackIndex) + 1);
-							ExtTextOutW(hdc, 5, posY + offsetY, ETO_CLIPPED, &rect, buf.data, buf.count, NULL);
+							context.ExtTextOutW(5, posY + offsetY, ETO_CLIPPED, &rect, buf);
 						}
 
 						int lastDrawnRight = 0;
@@ -1830,73 +1878,20 @@ namespace uba
 							auto clampRect = [&](RECT& r) { r.left = Min(Max(r.left, progressRect.left), progressRect.right); r.right = Max(Min(r.right, progressRect.right), progressRect.left); };
 
 							clampRect(rect);
-							fillRect(rect, brush);
+							context.FillRect(rect, brush);
 							++rect.top;
 
 							int processWidth = rect.right - rect.left;
 							if (shouldDrawText && m_config.ShowProcessText && processWidth > 3)
 							{
-								if (!work.bitmap || work.bitmapDirty)
-								{
-									if (!work.bitmap)
-									{
-										if (m_lastBitmapOffset == BitmapCacheHeight)
-										{
-											if (m_lastBitmap)
-												m_textBitmaps.push_back(m_lastBitmap);
-
-											m_lastBitmapOffset = 0;
-											m_lastBitmap = CreateCompatibleBitmap(hdc, 256, BitmapCacheHeight);
-										}
-										work.bitmap = m_lastBitmap;
-										work.bitmapOffset = m_lastBitmapOffset;
-										m_lastBitmapOffset += m_processFont.height;
-									}
-									if (lastSelectedBitmap != work.bitmap)
+								context.DrawCached(work, processWidth, offsetY, textHeight, left, rect, progressRect, [&](PaintContext& textCacheContext)
 									{
-										SelectObject(textDC, work.bitmap);
-										lastSelectedBitmap = work.bitmap;
-									}
-
-									RECT rect2{ 0,int(work.bitmapOffset), 256, int(work.bitmapOffset) + m_processFont.height };
-
-									SelectObject(textDC, brush);
-									PatBlt(textDC, rect2.left, rect2.top, rect2.right - rect2.left, rect2.bottom - rect2.top, PATCOPY);
-									ExtTextOutW(textDC, rect2.left, rect2.top, ETO_CLIPPED, &rect2, work.description, int(wcslen(work.description)), NULL);
-									if (!selected)
-										work.bitmapDirty = false;
-								}
-								else if (lastSelectedBitmap != work.bitmap)
-								{
-									SelectObject(textDC, work.bitmap);
-									lastSelectedBitmap = work.bitmap;
-								}
-
-								int width = Min(processWidth, 256);
-								int bitmapOffsetY = work.bitmapOffset;
-								int bltOffsetY = offsetY;
-								if (bltOffsetY < 0)
-								{
-									bitmapOffsetY -= bltOffsetY;
-									bltOffsetY = 0;
-								}
-								int height = Min(textHeight, m_processFont.height);
-								if (bltOffsetY + height > textHeight)
-									height = textHeight - bltOffsetY;
-
-								if (left > -256 && height >= 0)
-								{
-									int bitmapOffsetX = rect.left - left;
-
-									if (left < progressRect.left)
-									{
-										int diff = progressRect.left - left;
-										rect.left = progressRect.left;
-										width -= diff;
-										bitmapOffsetX += diff;
-									}
-									BitBlt(hdc, rect.left, rect.top + bltOffsetY, width, height, textDC, bitmapOffsetX, bitmapOffsetY, SRCCOPY);
-								}
+										RECT rect2{ 0,int(work.bitmapOffset), 256, int(work.bitmapOffset) + m_processFont.height };
+										textCacheContext.FillRect(rect2, brush);
+										textCacheContext.ExtTextOutW(rect2.left, rect2.top, ETO_CLIPPED, &rect2, ToView(work.description));
+										if (!selected)
+											work.bitmapDirty = false;
+									});
 							}
 						}
 					}
@@ -1905,7 +1900,7 @@ namespace uba
 				}
 			}
 
-			SetActiveFont(m_defaultFont);
+			context.SetActiveFont(m_defaultFont);
 		}
 
 		if (true)
@@ -1928,7 +1923,7 @@ namespace uba
 					if (int(x) >= progressRect.right)
 						break;
 
-					ExtTextOutW(hdc, x, y, ETO_CLIPPED, nullptr, item, TStrlen(item), NULL);
+					context.ExtTextOutW(x, y, ETO_CLIPPED, nullptr, ToView(item));
 
 					++rowIndex;
 					if (rowIndex > table.rowMaxCount || int(y) >= progressRect.bottom)
@@ -1940,10 +1935,6 @@ namespace uba
 			}
 		}
 
-		SelectObject(textDC, oldBmp);
-		DeleteObject(nullBmp);
-		DeleteDC(textDC);
-
 		m_contentWidth = m_progressRectLeft + Max(0, int(TimeToS((lastStop != 0 && lastStop != ~u64(0)) ? lastStop : playTime) * scaleX));
 
 
@@ -1952,7 +1943,7 @@ namespace uba
 		float timelineSelected = m_timelineSelected;
 
 		if (m_config.showTimeline && !m_traceView.sessions.empty())
-			PaintTimeline(hdc, clientRect);
+			PaintTimeline(context, clientRect);
 
 		if (!m_filterString.empty())
 		{
@@ -1960,8 +1951,8 @@ namespace uba
 			str[0].Append(TCV("Box Filter: ")).Append(m_filterString);
 			str[1].Appendf(TC("Box Count: %u"), visibleBoxes);
 
-			SetActiveFont(m_popupFont);
-			DrawCenteredText(str, 2, true);
+			context.SetActiveFont(m_popupFont);
+			context.DrawCenteredText(str, 2, true);
 		}
 
 		if (m_config.showCursorLine && m_mouseOverWindow)
@@ -1981,8 +1972,7 @@ namespace uba
 			int timelineTop = GetTimelineTop(clientRect);
 
 			// TODO: Draw line up
-			MoveToEx(hdc, left, 2, NULL);
-			LineTo(hdc, left, timelineTop);
+			context.Line(left, 2, left, timelineTop);
 
 			if (timelineSelected >= 0)
 			{
@@ -2015,8 +2005,8 @@ namespace uba
 
 				StringBuffer<128> wt = GetWorldTime(timelineSelected);
 
-				SetActiveFont(m_popupFont);
-				DrawTextLogger logger(m_hwnd, hdc, m_popupFont.height, m_tooltipBackgroundBrush);
+				context.SetActiveFont(m_popupFont);
+				DrawTextLogger logger(context, m_popupFont.height, m_tooltipBackgroundBrush);
 
 				logger.Info(L"%s (%s)", b.data, wt.data);
 				logger.DrawAtPos(left + 4, timelineTop - 20);
@@ -2033,16 +2023,13 @@ namespace uba
 			bool* values = &m_config.showProgress;
 			for (int i = VisualizerFlag_Count - 1; i >= 0; --i)
 			{
-				SelectObject(hdc, m_buttonSelected == u32(i) ? m_textPen : m_checkboxPen);
-				SelectObject(hdc, GetStockObject(NULL_BRUSH));
-				Rectangle(hdc, left, top, right, bottom);
+				context.SetPen(m_buttonSelected == u32(i) ? m_textPen : m_checkboxPen);
+				context.Rectangle(left, top, right, bottom);
 
 				if (values[i])
 				{
-					MoveToEx(hdc, left + 2, top + 2, NULL);
-					LineTo(hdc, right - 2, bottom - 2);
-					MoveToEx(hdc, right - 3, top + 2, NULL);
-					LineTo(hdc, left + 1, bottom - 2);
+					context.Line(left + 2, top + 2, right - 2, bottom - 2);
+					context.Line(right - 3, top + 2, left + 1, bottom - 2);
 				}
 
 				left -= boxStride;
@@ -2051,38 +2038,34 @@ namespace uba
 
 			top -= 2;
 
-			auto drawText = [&](const tchar* text, COLORREF color)
+			auto drawText = [&](StringView text, COLORREF color)
 				{
-					SetTextColor(hdc, color);
-					RECT r{left, top, left+200, top+200};
-					auto strLen = TStrlen(text);
-					DrawTextW(hdc, text, strLen, &r, DT_SINGLELINE|DT_NOCLIP|DT_NOPREFIX|DT_CALCRECT);
-					left -= (r.right - r.left) + 5;
-					r.left = left;
-					DrawTextW(hdc, text, strLen, &r, DT_SINGLELINE|DT_NOCLIP|DT_NOPREFIX);
+					context.SetTextColor(color);
+					left -= (context.GetTextSize(text).cx) + 5;
+					context.ExtTextOutW(left, top, 0, NULL, text);
 				};
 
 			if (m_config.showDriveStats)
 			{
-				SetActiveFont(m_defaultFont);
-				drawText(L"DRV", m_driveColor);
+				context.SetActiveFont(m_defaultFont);
+				drawText(TCV("DRV"), m_driveColor);
 			}
 
 			if (m_config.showNetworkStats)
 			{
-				SetActiveFont(m_defaultFont);
-				drawText(L"SND", m_sendColor);
-				drawText(L"RCV", m_recvColor);
+				context.SetActiveFont(m_defaultFont);
+				drawText(TCV("SND"), m_sendColor);
+				drawText(TCV("RCV"), m_recvColor);
 			}
 
 			if (m_config.showCpuMemStats)
 			{
-				SetActiveFont(m_defaultFont);
-				drawText(L"CPU", m_cpuColor);
-				drawText(L"MEM", m_memColor);
+				context.SetActiveFont(m_defaultFont);
+				drawText(TCV("CPU"), m_cpuColor);
+				drawText(TCV("MEM"), m_memColor);
 			}
 
-			SetTextColor(hdc, m_textColor);
+			context.SetTextColor(m_textColor);
 		}
 
 		if (m_processSelected)
@@ -2121,8 +2104,8 @@ namespace uba
 				duration = playTime - process.start;
 			}
 
-			SetActiveFont(m_popupFont);
-			DrawTextLogger logger(m_hwnd, hdc, m_popupFont.height, m_tooltipBackgroundBrush);
+			context.SetActiveFont(m_popupFont);
+			DrawTextLogger logger(context, m_popupFont.height, m_tooltipBackgroundBrush);
 
 			logger.AddTextOffset(-10); // Remove spaces in the front
 			logger.AddWidth(3);
@@ -2277,8 +2260,8 @@ namespace uba
 			else
 				duration = playTime - selectedWork.start;
 
-			SetActiveFont(m_popupFont);
-			DrawTextLogger logger(m_hwnd, hdc, m_popupFont.height, m_tooltipBackgroundBrush);
+			context.SetActiveFont(m_popupFont);
+			DrawTextLogger logger(context, m_popupFont.height, m_tooltipBackgroundBrush);
 
 			logger.AddSpace();
 			logger.Info(L"  %ls", selectedWork.description);
@@ -2322,8 +2305,8 @@ namespace uba
 		}
 		else if (m_sessionSelectedIndex != ~0u)
 		{
-			SetActiveFont(m_popupFont);
-			DrawTextLogger logger(m_hwnd, hdc, m_popupFont.height, m_tooltipBackgroundBrush);
+			context.SetActiveFont(m_popupFont);
+			DrawTextLogger logger(context, m_popupFont.height, m_tooltipBackgroundBrush);
 			logger.AddWidth(3);
 			logger.AddSpace(2);
 
@@ -2360,8 +2343,8 @@ namespace uba
 		}
 		else if (m_statsSelected)
 		{
-			SetActiveFont(m_popupFont);
-			DrawTextLogger logger(m_hwnd, hdc, m_popupFont.height, m_tooltipBackgroundBrush);
+			context.SetActiveFont(m_popupFont);
+			DrawTextLogger logger(context, m_popupFont.height, m_tooltipBackgroundBrush);
 			logger.AddSpace(3);
 			logger.SetColor(m_cpuColor).Info(L"  Cpu: %.1f%%", m_stats.cpuLoad * 100.0f);
 			logger.SetColor(m_memColor).Info(L"  Mem: %ls/%ls", BytesToText(m_stats.memTotal - m_stats.memAvail).str, BytesToText(m_stats.memTotal).str);
@@ -2385,8 +2368,8 @@ namespace uba
 		}
 		else if (m_activeProcessGraphSelected)
 		{
-			SetActiveFont(m_popupFont);
-			DrawTextLogger logger(m_hwnd, hdc, m_popupFont.height, m_tooltipBackgroundBrush);
+			context.SetActiveFont(m_popupFont);
+			DrawTextLogger logger(context, m_popupFont.height, m_tooltipBackgroundBrush);
 			logger.AddSpace(3);
 			logger.AddSpace(3);
 			logger.SetColor(m_activeProcColor).Info(L"  Active Processes: %u", m_activeProcessCount);
@@ -2401,8 +2384,8 @@ namespace uba
 #undef UBA_VISUALIZER_FLAG
 			};
 
-			SetActiveFont(m_popupFont);
-			DrawTextLogger logger(m_hwnd, hdc, m_popupFont.height, m_tooltipBackgroundBrush);
+			context.SetActiveFont(m_popupFont);
+			DrawTextLogger logger(context, m_popupFont.height, m_tooltipBackgroundBrush);
 			logger.Info(L"%ls %ls", L"Show", tooltip[m_buttonSelected]);
 			logger.DrawAtCursor();
 		}
@@ -2417,8 +2400,8 @@ namespace uba
 				int width = colWidth * 2;
 				int height = Min(int(clientRect.bottom), int(fetchedFiles.size() * m_popupFont.height));
 
-				SetActiveFont(m_defaultFont);
-				DrawTextLogger logger(m_hwnd, hdc, r, m_font.height, m_tooltipBackgroundBrush);
+				context.SetActiveFont(m_defaultFont);
+				DrawTextLogger logger(context, r, m_font.height, m_tooltipBackgroundBrush);
 				for (auto& f : fetchedFiles)
 				{
 				if (f.hint == TC("KnownInput"))
@@ -2443,9 +2426,9 @@ namespace uba
 
 	u64 ConvertTime(const TraceView& view, u64 time);
 
-	void VisualizerWin::PaintActiveProcesses(int& posY, const RECT& clientRect, const Function<void(TraceView::ProcessLocation&, u32, bool)>& drawProcess)
+	void VisualizerWin::PaintActiveProcesses(PaintContext& context, int& posY, const RECT& clientRect, const Function<void(TraceView::ProcessLocation&, u32, bool)>& drawProcess)
 	{
-		SetActiveFont(m_processFont);
+		context.SetActiveFont(m_processFont);
 		int startPosY = posY;
 
 		Map<u64, TraceView::ProcessLocation*> activeProcesses;
@@ -2540,10 +2523,10 @@ namespace uba
 		else
 			posY += 3;
 
-		SetActiveFont(m_defaultFont);
+		context.SetActiveFont(m_defaultFont);
 	}
 
-	void VisualizerWin::PaintProcessRect(TraceView::Process& process, HDC hdc, RECT rect, const RECT& progressRect, bool selected, bool writingBitmap, HBRUSH& lastSelectedBrush)
+	void VisualizerWin::PaintProcessRect(PaintContext& context, TraceView::Process& process, RECT rect, const RECT& progressRect, bool selected, bool writingBitmap)
 	{
 		auto clampRect = [&](RECT& r) { r.left = Min(Max(r.left, progressRect.left), progressRect.right); r.right = Max(Min(r.right, progressRect.right), progressRect.left); };
 		bool done = process.stop != ~u64(0);
@@ -2564,22 +2547,12 @@ namespace uba
 
 		u64 writeFilesTime = process.writeFilesTime;
 
-		auto fillRect = [&](const RECT& r, HBRUSH b)
-			{
-				if (lastSelectedBrush != b)
-				{
-					SelectObject(hdc, b);
-					lastSelectedBrush = b;
-				}
-				PatBlt(hdc, r.left, r.top, r.right - r.left, r.bottom - r.top, PATCOPY);
-			};
-
 		if (!done || process.exitCode != 0 || !m_config.ShowReadWriteColors)
 		{
 			if (writingBitmap)
 				rect.right = 256;
 			clampRect(rect);
-			fillRect(rect, brush);
+			context.FillRect(rect, brush);
 			return;
 		}
 
@@ -2596,7 +2569,7 @@ namespace uba
 			r2.right = r2.left + headSize;
 			clampRect(r2);
 			if (r2.left != r2.right)
-				fillRect(r2, m_processBrushes[selected].recv);
+				context.FillRect(r2, m_processBrushes[selected].recv);
 		}
 
 		double sendPart = (double(ConvertTime(m_traceView, writeFilesTime)) / duration);
@@ -2607,38 +2580,17 @@ namespace uba
 			r2.left = r2.right - tailSize;
 			clampRect(r2);
 			if (r2.left != r2.right)
-				fillRect(r2, m_processBrushes[selected].send);
+				context.FillRect(r2, m_processBrushes[selected].send);
 		}
 
 		clampRect(main);
 		if (main.left != main.right)
-			fillRect(main, brush);
-
-		//clampRect(rect);
-		/*
-		if (!process.updates.empty())
-		{
-		shouldDrawText = false;
-		int prevUpdateX = rect.left;
-		for (auto& update : process.updates)
-		{
-		int updateX = int(posX + TimeToS(update.time) * scaleX) - 1;
-
-		RECT textRect{ prevUpdateX + 5, rect.top, updateX, rect.bottom };
-		DrawTextW(hdc, update.reason.c_str(), u32(update.reason.size()), &textRect, DT_SINGLELINE);
-
-		SelectObject(hdc, m_processUpdatePen);
-		MoveToEx(hdc, updateX, rect.top, NULL);
-		LineTo(hdc, updateX, rect.bottom - 1);
-		prevUpdateX = updateX;
-		}
-		}
-		*/
+			context.FillRect(main, brush);
 	}
 
-	void VisualizerWin::PaintTimeline(HDC hdc, const RECT& clientRect)
+	void VisualizerWin::PaintTimeline(PaintContext& context, const RECT& clientRect)
 	{
-		SetActiveFont(m_timelineFont);
+		context.SetActiveFont(m_timelineFont);
 		int top = GetTimelineTop(clientRect);
 		float timeScale = (m_config.horizontalScaleValue*m_zoomValue)*50.0f;
 		float startOffset = ((m_scrollPosX/timeScale) - float(int(m_scrollPosX/timeScale))) * timeScale;
@@ -2665,7 +2617,7 @@ namespace uba
 		RECT progressRect = clientRect;
 		progressRect.left += m_progressRectLeft;
 
-		SelectObject(hdc, m_textPen);
+		context.SetPen(m_textPen);
 		Vector<DWORD> lines;
 		Vector<POINT> points;
 
@@ -2696,14 +2648,13 @@ namespace uba
 
 				if (shouldDraw)
 				{
-					SIZE s;
-					GetTextExtentPoint32W(hdc, buffer.data, buffer.count, &s);
+					SIZE s = context.GetTextSize(buffer);
 					RECT textRect;
 					textRect.top = top + 8;
 					textRect.bottom = textRect.top + m_activeFont.height;
 					textRect.right = pos + s.cx/2;
 					textRect.left = pos - s.cx/2;
-					ExtTextOutW(hdc, textRect.left, textRect.top, 0, NULL, buffer.data, buffer.count, NULL);
+					context.ExtTextOutW(textRect.left, textRect.top, 0, NULL, buffer);
 				}
 			}
 			if (!(number % lineStepSize))
@@ -2721,19 +2672,7 @@ namespace uba
 		points.push_back({m_contentWidth, top});
 		lines.push_back(2);
 
-		PolyPolyline(hdc, points.data(), lines.data(), int(lines.size()));
-
-		/*
-		POINT cursorPos;
-		GetCursorPos(&cursorPos);
-		ScreenToClient(m_hwnd, &cursorPos);
-		RECT lineRect;
-		lineRect.left = cursorPos.x;
-		lineRect.right= cursorPos.x + 1;
-		lineRect.top = top;
-		lineRect.bottom = top + 15;
-		FillRect(hdc, &lineRect, m_lineBrush);
-		*/
+		context.PolyPolyline(points.data(), lines.data(), int(lines.size()));
 	}
 
 	void VisualizerWin::PaintDetailedStats(int& posY, const RECT& progressRect, TraceView::Session& session, bool isRemote, u64 playTime, const DrawTextFunc& drawTextFunc)
@@ -2911,17 +2850,6 @@ namespace uba
 		posY = maxY;
 	}
 
-	u64 VisualizerWin::GetPlayTime()
-	{
-		u64 currentTime = m_paused ? m_pauseStart : GetTime();
-		u64 playTime = 0;
-		if (m_traceView.startTime && currentTime > (m_traceView.startTime + m_pauseTime))
-			playTime = currentTime - m_traceView.startTime - m_pauseTime;
-		if (m_replay)
-			playTime *= m_replay;
-		return playTime;
-	}
-
 	int VisualizerWin::GetTimelineHeight()
 	{
 		return m_timelineFont.height + 8;
@@ -2937,13 +2865,14 @@ namespace uba
 
 	void VisualizerWin::HitTest(HitTestResult& outResult, const POINT& pos)
 	{
-		SetActiveFont(m_defaultFont);
-
 		u64 playTime = GetPlayTime();
 
 		RECT clientRect;
 		GetClientRect(m_hwnd, &clientRect);
 
+		PaintContext context { *this, NULL, NULL, clientRect };
+		context.SetActiveFont(m_defaultFont);
+
 		int posY = int(m_scrollPosY);
 		int boxHeight = m_config.boxHeight;
 		int processStepY = boxHeight + 1;
@@ -3055,7 +2984,7 @@ namespace uba
 
 		if (m_config.showActiveProcesses && !m_trace.m_activeProcesses.empty())
 		{
-			PaintActiveProcesses(posY, clientRect, [&](TraceView::ProcessLocation& processLocation, u32 boxHeight, bool firstWithHeight)
+			PaintActiveProcesses(context, posY, clientRect, [&](TraceView::ProcessLocation& processLocation, u32 boxHeight, bool firstWithHeight)
 				{
 					if (pos.y < posY || pos.y > posY + int(boxHeight))
 						return;
@@ -3579,13 +3508,6 @@ namespace uba
 			SetWindowPos(m_hwnd, NULL, 0, 0, 0, 0, SWP_NOSIZE | SWP_NOMOVE | SWP_NOZORDER | SWP_FRAMECHANGED);
 	}
 
-	void VisualizerWin::SetActiveFont(const Font& font)
-	{
-		m_activeFont = font;
-		if (m_activeHdc)
-			SelectObject(m_activeHdc, font.handle);
-	}
-
 	void VisualizerWin::UpdateTheme()
 	{
 		SetWindowTheme(m_hwnd, m_config.DarkMode ? L"DarkMode_Explorer" : L"Explorer", NULL);
@@ -3709,9 +3631,24 @@ namespace uba
 			PaintClient([&](HDC hdc, HDC memDC, RECT& rect)
 				{
 					FillRect(memDC, &rect, m_backgroundBrush);
-					m_activeHdc = memDC;
-					PaintAll(memDC, rect);
-					m_activeHdc = 0;
+					SetBkColor(memDC, m_config.DarkMode ? RGB(70, 70, 70) : RGB(180, 180, 180));
+
+					HDC textDC = CreateCompatibleDC(memDC);
+					SetTextColor(textDC, m_textColor);
+					SelectObject(textDC, m_processFont.handle);
+					SelectObject(textDC, GetStockObject(NULL_BRUSH));
+					SetBkMode(textDC, TRANSPARENT);
+
+					HBITMAP nullBmp = CreateCompatibleBitmap(memDC, 1, 1);
+					HBITMAP oldBmp = (HBITMAP)SelectObject(textDC, nullBmp);
+
+					PaintContext context { *this, m_hwnd, memDC, rect, textDC };
+					PaintAll(context, rect);
+
+					SelectObject(textDC, oldBmp);
+					DeleteObject(nullBmp);
+					DeleteDC(textDC);
+
 					BitBlt(hdc, 0, 0, rect.right - rect.left, rect.bottom - rect.top, memDC, 0, 0, SRCCOPY);
 				});
 			m_lastPaintTimeMs = TimeToMs(GetTime() - start);
@@ -3921,23 +3858,6 @@ namespace uba
 			{
 				if (UpdateSelection() || m_config.showCursorLine)
 					Redraw(false);
-				/*
-				else
-				{
-				PaintClient([&](HDC hdc, HDC memDC, RECT& rect)
-				{
-				RECT timelineRect = rect;
-				rect.top = Min(rect.bottom, m_contentHeight) - 28;
-				rect.bottom = Min(rect.bottom, rect.top + 28);
-				FillRect(memDC, &rect, m_backgroundBrush);
-				SetBkMode(memDC, TRANSPARENT);
-				SetTextColor(memDC, m_textColor);
-				SelectObject(memDC, m_font);
-				PaintTimeline(memDC, timelineRect);
-				BitBlt(hdc, 0, rect.top, rect.right - rect.left, rect.bottom - rect.top, memDC, 0, rect.top, SRCCOPY);
-				});
-				}
-				*/
 			}
 
 			TRACKMOUSEEVENT tme;
diff --git a/Engine/Source/Programs/UnrealBuildAccelerator/Visualizer/Private/Windows/UbaVisualizerWin.h b/Engine/Source/Programs/UnrealBuildAccelerator/Visualizer/Private/Windows/UbaVisualizerWin.h
--- a/Engine/Source/Programs/UnrealBuildAccelerator/Visualizer/Private/Windows/UbaVisualizerWin.h
+++ b/Engine/Source/Programs/UnrealBuildAccelerator/Visualizer/Private/Windows/UbaVisualizerWin.h
@@ -28,52 +28,17 @@ namespace uba
 		bool Unselect();
 		void Reset();
 		void PaintClient(const Function<void(HDC hdc, HDC memDC, RECT& clientRect)>& paintFunc);
-		void PaintAll(HDC hdc, const RECT& clientRect);
-		void PaintActiveProcesses(int& posY, const RECT& clientRect, const Function<void(TraceView::ProcessLocation&, u32, bool)>& drawProcess);
-		void PaintProcessRect(TraceView::Process& process, HDC hdc, RECT rect, const RECT& progressRect, bool selected, bool writingBitmap, HBRUSH& lastSelectedBrush);
-		void PaintTimeline(HDC hdc, const RECT& clientRect);
+
+		struct PaintContext;
+		struct DrawTextLogger;
+
+		void PaintAll(PaintContext& context, const RECT& clientRect);
+		void PaintActiveProcesses(PaintContext& context, int& posY, const RECT& clientRect, const Function<void(TraceView::ProcessLocation&, u32, bool)>& drawProcess);
+		void PaintProcessRect(PaintContext& context, TraceView::Process& process, RECT rect, const RECT& progressRect, bool selected, bool writingBitmap);
+		void PaintTimeline(PaintContext& context, const RECT& clientRect);
 		using DrawTextFunc = Function<void(const StringBufferBase& text, RECT& rect, u32* outWidth)>;
 		void PaintDetailedStats(int& posY, const RECT& progressRect, TraceView::Session& session, bool isRemote, u64 playTime, const DrawTextFunc& drawTextFunc);
 
-		struct Stats
-		{
-			u64 recvBytesPerSecond = 0;
-			u64 sendBytesPerSecond = 0;
-			u64 ping = 0;
-			u64 memAvail = 0;
-			u64 memTotal = 0;
-			u64 recvBytes = 0;
-			u64 sendBytes = 0;
-			u64 procActive = 0;
-			u16 connectionCount = 0;
-			float cpuLoad = 0;
-			struct Drive
-			{
-				u8 busyPercent;
-				u64 readPerSecond = 0;
-				u64 writePerSecond = 0;
-			};
-			Map<char, Drive> drives;
-		};
-		struct HitTestResult
-		{
-			u32 section = ~0u;
-			TraceView::ProcessLocation processLocation;
-			bool processSelected = false;
-			u32 sessionSelectedIndex = ~0u;
-			bool statsSelected = false;
-			Stats stats;
-			u32 buttonSelected = ~0u;
-			float timelineSelected = 0;
-			u32 fetchedFilesSelected = ~0u;
-			bool workSelected = false;
-			u32 workTrack = ~0u;
-			u32 workIndex = ~0u;
-			bool activeProcessGraphSelected = false;
-			u16 activeProcessCount = 0u;
-			TString hyperLink;
-		};
-		u64 GetPlayTime();
 		int GetTimelineHeight();
 		int GetTimelineTop(const RECT& clientRect);
 		void HitTest(HitTestResult& outResult, const POINT& pos);
@@ -108,7 +73,6 @@ namespace uba
 		void UpdateProcessFont();
 		void ChangeFontSize(int offset);
 		void Redraw(bool now);
-		void SetActiveFont(const Font& font);
 		void UpdateTheme();
 		bool ActiveProcessesShouldFillHeight();
 		StringBuffer<128> GetWorldTime(u64 time);
@@ -116,9 +80,6 @@ namespace uba
 
 		StringBuffer<256> m_namedTrace;
 		StringBuffer<256> m_fileName;
-		u32 m_replay = 0;
-		u64 m_startTime = 0;
-		u64 m_pauseTime = 0;
 
 		u64 m_lastPaintTimeMs;
 
@@ -175,7 +136,6 @@ namespace uba
 		u32 m_activeProcessCountHistory[5];
 		u32 m_activeProcessCountHistoryIterator = 0;
 
-		HDC m_activeHdc = 0;
 		Font m_activeFont;
 
 		int m_progressRectLeft = 30;
@@ -197,8 +157,6 @@ namespace uba
 		float m_zoomValue = 0.5f;
 		float m_horizontalScaleValue = 0.5f;
 		bool m_autoScroll = true;
-		bool m_paused = false;
-		u64 m_pauseStart = 0;
 
 		bool m_isInPaint = false;
 
