[UBA]
* Moved Paint and HitTest code from VisualizerWin to Visualizer

diff --git a/Engine/Source/Programs/UnrealBuildAccelerator/Visualizer/Private/UbaVisualizer.cpp b/Engine/Source/Programs/UnrealBuildAccelerator/Visualizer/Private/UbaVisualizer.cpp
--- a/Engine/Source/Programs/UnrealBuildAccelerator/Visualizer/Private/UbaVisualizer.cpp
+++ b/Engine/Source/Programs/UnrealBuildAccelerator/Visualizer/Private/UbaVisualizer.cpp
@@ -1,9 +1,112 @@
 // Copyright Epic Games, Inc. All Rights Reserved.
 
 #include "UbaVisualizer.h"
+#include <algorithm>
 
 namespace uba
 {
+	////////////////////////////////////////////////////////////////////////////////////////////////////
+
+	u64 ConvertTime(const TraceView& view, u64 time);
+
+	////////////////////////////////////////////////////////////////////////////////////////////////////
+
+	struct Visualizer::DrawTextLogger : public Logger
+	{
+		DrawTextLogger(PaintContext& c, int fh, HBRUSH bb)
+			:	context(c)
+			,	fontHeight(fh)
+			,	backgroundBrush(bb)
+		{
+			textColor = context.GetTextColor();
+		}
+
+		virtual void BeginScope() override {}
+		virtual void EndScope() override {}
+		virtual void Log(LogEntryType type, const wchar_t* str, u32 strLen) override
+		{
+			u32 textWidth = context.GetTextSize(StringView(str, strLen)).cx;
+
+			lines.emplace_back(TString(str, strLen), textOffset, height, textColor);
+			width = Max(width, int(textWidth + textOffset));
+			height += fontHeight;
+		}
+
+		void AddSpace(int space = 5)
+		{
+			height += space;
+		}
+
+		void AddTextOffset(int offset)
+		{
+			textOffset += offset;
+		}
+
+		void AddWidth(int extra)
+		{
+			extraWidth += extra;
+		}
+
+		void DrawAtPos(int x, int y)
+		{
+			RECT r;
+			r.left = x;
+			r.top = y;
+			r.right = r.left + width;
+			r.bottom = r.top + height;
+
+			RECT clientRect = context.GetClientRect();
+
+			if (r.right > clientRect.right)
+				OffsetRect(&r, -width - 15, 0);
+			if (r.bottom > clientRect.bottom)
+			{
+				OffsetRect(&r, 0, clientRect.bottom - r.bottom);
+				if (r.top < 0)
+					OffsetRect(&r, 0, -r.top);
+			}
+
+			RECT fillRect = r;
+			fillRect.right += 2 + extraWidth;
+			context.FillRect(fillRect, backgroundBrush);
+
+			for (auto& line : lines)
+			{
+				RECT tr = r;
+				tr.left += line.left;
+				tr.top += line.top;
+				context.SetTextColor(line.color);
+				context.ExtTextOutW(tr.left, tr.top, 0, NULL, line.str);
+			}
+		}
+
+		void DrawAtCursor()
+		{
+			POINT p = context.GetCursorClientPos();
+			p.x += 3;
+			p.y += 3;
+			DrawAtPos(p.x, p.y);
+		}
+
+		DrawTextLogger& SetColor(COLORREF c) { textColor = c; return *this; }
+
+		int width = 0;
+		int height = 0;
+		int textOffset = 2;
+		int extraWidth = 0;
+		struct Line { TString str; int left; int top; COLORREF color; };
+		Vector<Line> lines;
+
+		PaintContext& context;
+		int fontHeight;
+		HBRUSH backgroundBrush;
+		COLORREF textColor;
+		bool isFirst = true;
+		using Logger::Log;
+	};
+
+	////////////////////////////////////////////////////////////////////////////////////////////////////
+
 	Visualizer::Visualizer(VisualizerConfig& config, Logger& logger)
 	:	m_logger(logger)
 	,	m_config(config)
@@ -21,4 +124,2257 @@ namespace uba
 			playTime *= m_replay;
 		return playTime;
 	}
-}
\ No newline at end of file
+
+	struct SessionRec
+	{
+		TraceView::Session* session;
+		u32 index;
+	};
+	void Populate(SessionRec* recs, TraceView& traceView, bool sort)
+	{
+		u32 count = u32(traceView.sessions.size());
+		for (u32 i = 0, e = count; i != e; ++i)
+			recs[i] = { &traceView.sessions[i], i };
+		if (count <= 1 || !sort)
+			return;
+		std::sort(recs + 1, recs + traceView.sessions.size(), [](SessionRec& a, SessionRec& b)
+			{
+				auto& as = *a.session;
+				auto& bs = *b.session;
+				if ((as.processActiveCount != 0) != (bs.processActiveCount != 0))
+					return as.processActiveCount > bs.processActiveCount;
+				if (as.processActiveCount && as.proxyCreated != bs.proxyCreated)
+					return int(as.proxyCreated) > int(bs.proxyCreated);
+				return a.index < b.index;
+			});
+	}
+
+	int Visualizer::GetTimelineHeight()
+	{
+		return m_timelineFont.height + 8;
+	}
+
+	int Visualizer::GetTimelineTop(const RECT& clientRect)
+	{
+		int timelineHeight = GetTimelineHeight();
+		int posY = m_contentHeight - timelineHeight;
+		int maxY = int(clientRect.bottom - timelineHeight);
+		return m_config.LockTimelineToBottom ? maxY : Min(posY, maxY);
+	}
+
+	void Visualizer::HitTest(HitTestResult& outResult, PaintContext& context, const POINT& pos)
+	{
+		u64 playTime = GetPlayTime();
+
+		RECT clientRect = context.GetClientRect();
+		context.SetActiveFont(m_defaultFont);
+
+		int posY = int(m_scrollPosY);
+		int boxHeight = m_config.boxHeight;
+		int processStepY = boxHeight + 1;
+		float scaleX = 50.0f*m_zoomValue*m_config.horizontalScaleValue;
+
+		RECT progressRect = clientRect;
+		progressRect.left += m_progressRectLeft;
+		progressRect.bottom -= 30;
+
+		{
+			int boxSide = 8;
+			int boxStride = boxSide + 2;
+			int top = 5;
+			int bottom = top + boxSide;
+			int left = progressRect.right - 7 - boxSide;
+			int right = progressRect.right - 7;
+			for (int i = VisualizerFlag_Count - 1; i >= 0; --i)
+			{
+				if (pos.x >= left && pos.x <= right && pos.y >= top && pos.y <= bottom)
+				{
+					outResult.buttonSelected = i;
+					return;
+				}
+				left -= boxStride;
+				right -= boxStride;
+			}
+		}
+
+		outResult.section = 0;
+
+		if (m_config.showTimeline && !m_traceView.sessions.empty())
+		{
+			int timelineTop = GetTimelineTop(clientRect);
+			if (pos.y >= timelineTop)// && pos.y < timelineTop + 40)
+			{
+				outResult.section = 3;
+				float timeScale = (m_config.horizontalScaleValue * m_zoomValue)*50.0f;
+				float startOffset = -(m_scrollPosX / timeScale);
+				outResult.timelineSelected = startOffset + float(pos.x - m_progressRectLeft) / timeScale;
+				return;
+			}
+		}
+
+		u64 lastStop = 0;
+
+		Font& activeFont = context.GetActiveFont();
+
+		if (m_config.showProgress && m_traceView.progressProcessesTotal)
+			posY += activeFont.height + 2;
+
+		if (m_config.showStatus && !m_traceView.statusMap.empty())
+		{
+			u32 lastRow = ~0u;
+			u32 row = ~0u;
+			for (auto& kv : m_traceView.statusMap)
+			{
+				if (kv.second.text.empty())
+					continue;
+
+				row = u32(kv.first >> 32);
+				if (lastRow != ~0u && lastRow != row)
+					posY += activeFont.height + 2;
+				lastRow = row;
+
+				if (!kv.second.link.empty())
+				{
+					if (pos.y >= posY && pos.y < posY + activeFont.height && pos.x > 20 && pos.x < 80) // TODO: x is just hard coded to fit horde for now
+					{
+						outResult.hyperLink = kv.second.link;
+						return;
+					}
+				}
+			}
+			if (row != ~0u)
+				posY += activeFont.height + 2;
+			posY += 3;
+		}
+
+		if (pos.y < posY)
+			return;
+
+		outResult.section = 1;
+
+		if (m_config.showActiveProcessGraph)
+		{
+			if (pos.y > posY && pos.y < posY + GraphHeight)
+			{
+				float timeScale = (m_config.horizontalScaleValue * m_zoomValue)*50.0f;
+				float startOffset = -(m_scrollPosX / timeScale);
+				u64 selectedTimeMs = u64(1000.0f*(startOffset + float(pos.x - m_progressRectLeft) / timeScale));
+
+				u64 lastTime = m_traceView.activeProcessCounts.empty() ? 0 : m_traceView.activeProcessCounts[m_traceView.activeProcessCounts.size()-1].time;
+				if (selectedTimeMs < TimeToMs(lastTime))
+				{
+					u16 count = 0;
+					for (auto& entry : m_traceView.activeProcessCounts)
+					{
+						count = entry.count;
+						if (TimeToMs(entry.time) > selectedTimeMs)
+							break;
+					}
+
+					outResult.activeProcessGraphSelected = true;
+					outResult.activeProcessCount = count;
+				}
+			}
+			posY += GraphHeight;
+		}
+
+		TraceView::ProcessLocation& outLocation = outResult.processLocation;
+
+		if (m_config.showActiveProcesses && !m_trace.m_activeProcesses.empty())
+		{
+			PaintActiveProcesses(context, posY, clientRect, [&](TraceView::ProcessLocation& processLocation, u32 boxHeight, bool firstWithHeight)
+				{
+					if (pos.y < posY || pos.y > posY + int(boxHeight))
+						return;
+					auto& session = m_trace.GetSession(m_traceView, processLocation.sessionIndex);
+					TraceView::Process& process = session.processors[processLocation.processorIndex].processes[processLocation.processIndex];
+
+					int posX = int(m_scrollPosX) + progressRect.left;
+					u64 stop = process.stop;
+					bool done = stop != ~u64(0);
+					if (!done)
+						stop = playTime;
+					int left = posX + int(float(TimeToS(process.start)) * scaleX);
+					int right = posX + int(float(TimeToS(stop)) * scaleX) - 1;
+
+					if (pos.x >= left && pos.x <= right)
+					{
+						outLocation.sessionIndex = processLocation.sessionIndex;
+						outLocation.processorIndex = processLocation.processorIndex;
+						outLocation.processIndex = processLocation.processIndex;
+						outResult.processSelected = true;
+						return;
+					}
+				});
+			if (outResult.processSelected)
+				return;
+		}
+
+		if (pos.y < posY)
+			return;
+
+		outResult.section = 2;
+
+		SessionRec sortedSessions[1024];
+		Populate(sortedSessions, m_traceView, m_config.SortActiveRemoteSessions);
+
+		for (u64 sessionIt = 0, sessionEnd = m_traceView.sessions.size(); sessionIt != sessionEnd; ++sessionIt)
+		{
+			bool isFirst = sessionIt == 0;
+			auto& session = *sortedSessions[sessionIt].session;
+			bool hasUpdates = !session.updates.empty();
+			if (!isFirst)
+			{
+				if (!hasUpdates && session.processors.empty())
+					continue;
+				if (!m_config.showFinishedProcesses && session.disconnectTime != ~u64(0))
+					continue;
+			}
+
+			u32 sessionIndex = sortedSessions[sessionIt].index;
+			if (!isFirst)
+				posY += 3;
+
+			if (m_config.showTitleBars)
+			{
+				if (pos.y >= posY && pos.y < posY + m_sessionStepY)
+				{
+					if (pos.x < int(session.fullNameWidth) + 5)
+					{
+						outResult.sessionSelectedIndex = sessionIndex;
+						return;
+					}
+				}
+
+				posY += m_sessionStepY;
+			}
+
+			bool showGraph = m_config.showNetworkStats || m_config.showCpuMemStats || m_config.showDriveStats;
+			if (showGraph && !session.updates.empty())
+			{
+				int posX = int(m_scrollPosX) + progressRect.left;
+
+				if (pos.y >= posY && pos.y < posY + GraphHeight && pos.x >= posX)
+				{
+					bool loop = true;
+					u32 reconnectIndex = 0;
+					while (loop)
+					{
+						u64 i = 0;
+						u64 e = 0;
+
+						if (reconnectIndex)
+							i = session.reconnectIndices[reconnectIndex - 1];
+
+						if (reconnectIndex < session.reconnectIndices.size())
+							e = session.reconnectIndices[reconnectIndex];
+						else
+						{
+							e = session.updates.size();
+							loop = false;
+						}
+
+						int prevX = -1;
+						int prevCenterX = -1;
+						for (; i!=e; ++i)
+						{
+							u64 updateTime = session.updates[i];
+							auto updateSend = session.networkSend[i];
+							auto updateRecv = session.networkRecv[i];
+
+							int x = posX + int(float(TimeToS(updateTime)) * scaleX);
+							int centerX = prevX != -1 ? (prevX + (x - prevX)/2) : x;
+
+							if (prevCenterX != -1 && pos.x >= prevCenterX && pos.x < centerX)
+							{
+								u64 prevTime = 0;
+								u64 prevSend = 0;
+								u64 prevRecv = 0;
+								if (i > 0)
+								{
+									prevTime = session.updates[i - 1];
+									prevSend = Min(session.networkSend[i - 1], updateSend);
+									prevRecv = Min(session.networkRecv[i - 1], updateRecv);
+								}
+
+								double duration = TimeToS(updateTime - prevTime);
+								outResult.stats.recvBytes = updateRecv;
+								outResult.stats.sendBytes = updateSend;
+								outResult.stats.recvBytesPerSecond = u64(float(updateRecv - prevRecv) / duration);
+								outResult.stats.sendBytesPerSecond = u64(float(updateSend - prevSend) / duration);
+								outResult.stats.ping = session.ping[i];
+								outResult.stats.memAvail = session.memAvail[i];
+								outResult.stats.cpuLoad = session.cpuLoad[i];
+								outResult.stats.memTotal = session.memTotal;
+								outResult.stats.connectionCount = session.connectionCount[i];
+								outResult.statsSelected = true;
+
+								for (auto& pair : session.drives)
+								{
+									auto& updateBusy = pair.second.busyPercent[i];
+									u64 updateRead = pair.second.readBytes[i];
+									u64 updateWrite = pair.second.writeBytes[i];
+									auto& statsDrive = outResult.stats.drives[pair.first];
+									statsDrive.busyPercent = updateBusy;
+									statsDrive.readPerSecond = u64(float(updateRead) / duration);
+									statsDrive.writePerSecond = u64(float(updateWrite) / duration);
+								}
+								return;
+							}
+
+							prevCenterX = centerX;
+							prevX = x;
+						}
+						++reconnectIndex;
+					}
+					posY += GraphHeight;
+				}
+
+				posY += GraphHeight;
+			}
+
+			if (m_config.showDetailedData)
+			{
+				auto drawText = [&](const StringBufferBase& text, RECT& rect, u32* outWidth)
+					{
+						if (pos.x >= rect.left && pos.x < rect.right && pos.y >= rect.top && pos.y < rect.bottom && text.StartsWith(TC("Fetched Files")))
+							outResult.fetchedFilesSelected = sessionIndex;
+					};
+				PaintDetailedStats(context, posY, progressRect, session, sessionIt != 0, playTime, drawText);
+			}
+
+			if (m_config.showProcessBars)
+			{
+				u32 processorIndex = 0;
+				for (auto& processor : session.processors)
+				{
+					bool drawProcessorIndex = m_config.showFinishedProcesses;
+
+					if (pos.y < progressRect.bottom && posY + processStepY >= progressRect.top && posY <= progressRect.bottom && pos.y >= posY-1 && pos.y < posY-1 + processStepY)
+					{
+						u32 processIndex = 0;
+						int posX = int(m_scrollPosX) + progressRect.left;
+						for (auto& process : processor.processes)
+						{
+							int left = posX + int(float(TimeToS(process.start)) * scaleX);
+
+							auto pig = MakeGuard([&]() { ++processIndex; });
+
+							if (left >= progressRect.right)
+								continue;
+
+							if (left < progressRect.left)
+								left = progressRect.left;
+
+							u64 stopTime = process.stop;
+							bool done = stopTime != ~u64(0);
+							if (!done)
+								stopTime = playTime;
+							else if (!m_config.showFinishedProcesses)
+								continue;
+
+							drawProcessorIndex = true;
+
+							RECT rect;
+							rect.left = left;
+							rect.right = posX + int(float(TimeToS(stopTime)) * scaleX);
+
+							if (rect.right <= progressRect.left)
+								continue;
+
+							if (!m_filterString.empty() && !Contains(process.description.c_str(), m_filterString.c_str()) && !Contains(process.breadcrumbs.c_str(), m_filterString.c_str()))
+								continue;
+
+							rect.right = Max(int(rect.right), left + 1);
+							rect.top = posY;
+							rect.bottom = posY + int(float(18) * m_zoomValue);
+
+							if (pos.x >= rect.left && pos.x <= rect.right)
+							{
+								outLocation.sessionIndex = sessionIndex;
+								outLocation.processorIndex = processorIndex;
+								outLocation.processIndex = processIndex;
+								outResult.processSelected = true;
+								return;
+							}
+						}
+					}
+
+					if (!processor.processes.empty())
+						lastStop = Max(lastStop, processor.processes.rbegin()->stop);
+
+					if (drawProcessorIndex)
+						posY += processStepY;
+
+					++processorIndex;
+				}
+			}
+			else
+			{
+				for (auto& processor : session.processors)
+					if (!processor.processes.empty())
+						lastStop = Max(lastStop, processor.processes.rbegin()->stop);
+			}
+
+			if (m_config.showWorkers && isFirst)
+			{
+				int trackIndex = 0;
+				for (auto& workTrack : m_traceView.workTracks)
+				{
+					if (pos.y < progressRect.bottom && posY + processStepY >= progressRect.top && posY <= progressRect.bottom && pos.y >= posY-1 && pos.y < posY-1 + processStepY)
+					{
+						u32 workIndex = 0;
+						int posX = int(m_scrollPosX) + progressRect.left;
+						for (auto& work : workTrack.records)
+						{
+							auto inc = MakeGuard([&](){ ++workIndex; });
+
+							int left = posX + int(float(TimeToS(work.start)) * scaleX);
+
+							if (left >= progressRect.right)
+								continue;
+
+							if (!m_filterString.empty())
+							{
+								bool keep = Contains(work.description, m_filterString.c_str());
+								if (!keep)
+									for (auto& en : work.entries)
+										keep |= Contains(en.text, m_filterString.c_str());
+								if (!keep)
+									continue;
+							}
+
+							if (left < progressRect.left)
+								left = progressRect.left;
+
+							u64 stopTime = work.stop;
+							bool done = stopTime != ~u64(0);
+							if (!done)
+								stopTime = playTime;
+
+							RECT rect;
+							rect.left = left;
+							rect.right = posX + int(float(TimeToS(stopTime)) * scaleX);
+
+							if (rect.right <= progressRect.left)
+								continue;
+
+							rect.right = Max(int(rect.right), left + 1);
+							rect.top = posY;
+							rect.bottom = posY + int(float(18) * m_zoomValue);
+
+							if (pos.x >= rect.left && pos.x <= rect.right)
+							{
+								outResult.workTrack = trackIndex;
+								outResult.workIndex = workIndex;
+								outResult.workSelected = true;
+								return;
+							}
+						}
+					}
+					++trackIndex;
+					posY += processStepY;
+				}
+			}
+		}
+
+		m_contentWidth = m_progressRectLeft + Max(0, int(TimeToS((lastStop != 0 && lastStop != ~u64(0)) ? lastStop : playTime) * scaleX));
+		m_contentHeight = posY - int(m_scrollPosY) + processStepY + 14;
+	}
+
+	void Visualizer::PaintAll(PaintContext& context, const RECT& clientRect)
+	{
+		u64 playTime = GetPlayTime();
+
+		context.SetBkMode(TRANSPARENT);
+		context.SetTextColor(m_textColor);
+
+
+		if (m_traceView.sessions.empty())
+		{
+			PaintEmptyScreen(context);
+			return;
+		}
+
+		if (m_traceView.version && (m_traceView.version < TraceReadCompatibilityVersion || m_traceView.version > TraceVersion))
+		{
+			if (!m_traceView.finished)
+			{
+				StringBuffer<> str;
+				str.Appendf(TC("Unsupported trace version %u (Versions supported are %u to %u)"), m_traceView.version, TraceReadCompatibilityVersion, TraceVersion);
+				context.SetTextColor(m_textWarningColor);
+				context.DrawCenteredText(&str, 1, false);
+			}
+			else
+			{
+				m_traceView.Clear();
+				FinishIncompatible();
+			}
+			return;
+		}
+
+		int posY = int(m_scrollPosY);
+		float scaleX = 50.0f*m_zoomValue*m_config.horizontalScaleValue;
+
+		RECT progressRect = clientRect;
+		progressRect.left += m_progressRectLeft;
+
+		if (m_config.showTimeline)
+		{
+			progressRect.bottom -= m_defaultFont.height + 10;
+		}
+
+		u64 lastStop = 0;
+
+		context.SetActiveFont(m_defaultFont);
+
+		auto drawStatusText = [&](StringView text, LogEntryType type, int posX, int endX, bool moveY, bool underlined = false, bool centered = false)
+			{
+				Font& activeFont = context.GetActiveFont();
+				RECT rect;
+				rect.left = posX;
+				rect.right = endX;
+				rect.top = posY + activeFont.offset;
+				rect.bottom = posY + activeFont.height + 2;
+				context.SetTextColor(type == LogEntryType_Info ? m_textColor : (type == LogEntryType_Error ? m_textErrorColor : m_textWarningColor));
+
+				u32 align = 0;
+				if (centered)
+				{
+					posX += (endX - posX)/2;
+					align = context.SetTextAlign(TA_CENTER);
+				}
+				if (underlined)
+					context.SetFont(activeFont.handleUnderlined);
+				context.ExtTextOutW(posX, posY, ETO_CLIPPED, &rect, text);
+				if (underlined)
+					context.SetFont(activeFont.handle);
+				if (centered)
+					context.SetTextAlign(align);
+				if (moveY)
+					posY = rect.bottom;
+			};
+
+		auto drawIndentedText = [&](const StringView& text, LogEntryType type, int indent, bool moveY, bool underlined = false)
+			{
+				int posX = 5 + indent*m_defaultFont.height;
+				drawStatusText(text, type, posX, clientRect.right, moveY, underlined);
+			};
+
+		if (m_config.showProgress && m_traceView.progressProcessesTotal)
+		{
+			drawIndentedText(AsView(L"Progress"), LogEntryType_Info, 1, false);
+
+			Font& activeFont = context.GetActiveFont();
+
+			float progress = float(m_traceView.progressProcessesDone) / float(m_traceView.progressProcessesTotal);
+			u32 start = 3 + 6*activeFont.height;
+			u32 width = activeFont.height * 18;
+			RECT rect;
+			rect.left = start;
+			rect.right = start + width;
+			rect.top = posY;
+			rect.bottom = posY + activeFont.height;
+			context.FillRect(rect, m_processBrushes[0].inProgress);
+
+			rect.right = rect.left + int(progress*float(width));//durationMs/100;
+			context.FillRect(rect, m_traceView.progressErrorCount ? m_processBrushes[0].error : m_processBrushes[0].success);
+
+			StringBuffer<> str;
+			str.Appendf(TC("%u%%    %u / %u"), u32(progress*100.0f), m_traceView.progressProcessesDone, m_traceView.progressProcessesTotal);
+
+			const tchar* remoteDisabled = TC(""); // TODO: Only show this in "advanced mode" m_traceView.remoteExecutionDisabled ? TC("   (Remote spawn disabled)") : TC("");
+			str.Append(remoteDisabled);
+
+			if (u32 active = m_traceView.totalProcessActiveCount)
+				str.Appendf(TC("   (%u active)"), active);
+
+			if (m_traceView.lastKillProcessTime && playTime >= m_traceView.lastKillProcessTime && playTime < m_traceView.lastKillProcessTime + MsToTime(800))
+			{
+				rect.left = start + width + 20;
+				rect.right = rect.left + activeFont.height * 7;
+				context.FillRect(rect, m_processBrushes[0].error);
+				drawStatusText(TCV("Killing (mem)"), LogEntryType_Info, rect.left, rect.right, 0, false, true);
+			}
+			else if (m_traceView.lastSpawningDelayStartTime && playTime >= m_traceView.lastSpawningDelayStartTime && playTime < m_traceView.lastSpawningDelayEndTime + MsToTime(500))
+			{
+				rect.left = start + width + 20;
+				rect.right = rect.left + activeFont.height * 7;
+				context.FillRect(rect, m_processBrushes[0].warning);
+				drawStatusText(TCV("Throttling (mem)"), LogEntryType_Info, rect.left, rect.right, 0, false, true);
+			}
+
+
+			drawIndentedText(str, LogEntryType_Info, 6, true);
+		}
+
+		if (m_config.showStatus && !m_traceView.statusMap.empty())
+		{
+			u32 lastRow = ~0u;
+			u32 row = ~0u;
+			Font& activeFont = context.GetActiveFont();
+			for (auto& kv : m_traceView.statusMap)
+			{
+				auto& status = kv.second;
+				if (status.text.empty())
+					continue;
+				row = u32(kv.first >> 32);
+				if (lastRow != ~0u && lastRow != row)
+					posY += activeFont.height + 2;
+				lastRow = row;
+				u32 column = u32(kv.first & ~0u);
+				drawIndentedText(status.text, status.type, column, false, !status.link.empty());
+			}
+			if (row != ~0u)
+				posY += activeFont.height + 2;
+
+			context.SetTextColor(m_textColor);
+			posY += 3;
+		}
+
+		if (m_config.showActiveProcessGraph)
+		{
+			if (posY + GraphHeight >= progressRect.top && posY + GraphHeight - 5 < progressRect.bottom)
+			{
+				int graphBaseY = posY + GraphHeight - 4;
+				double graphHeight = double(GraphHeight - 2);
+
+				Vector<POINT> line;
+
+				context.SetPen(m_activeProcPen);
+				bool isFirstUpdate = true;
+				bool isFirstDraw = true;
+				u64 prevValue = 0;
+				int prevX = 0;
+				int prevY = 0;
+
+				float timeScale = (m_config.horizontalScaleValue*m_zoomValue)*50.0f;
+				float startOffset = ((m_scrollPosX/timeScale) - float(int(m_scrollPosX/timeScale))) * timeScale;
+				float time = -startOffset/timeScale;
+				float playTimeS = TimeToS(playTime);
+				u32 activeProcessCountIndex = 0;
+				u16 count = 0;
+				while (time < playTimeS && activeProcessCountIndex < m_traceView.activeProcessCounts.size())
+				{
+					int posX = progressRect.left + int(startOffset + time*timeScale);
+					if (posX >= clientRect.right)
+						break;
+
+					while (activeProcessCountIndex < m_traceView.activeProcessCounts.size())
+					{
+						auto& entry = m_traceView.activeProcessCounts[activeProcessCountIndex];
+						count = entry.count;
+						if (float(TimeToMs(entry.time))/1000.0f > time)
+							break;
+						++activeProcessCountIndex;
+					}
+
+					int x = posX;
+					int y = graphBaseY;
+
+					double div = double(m_traceView.maxActiveProcessCount);
+					y -= int(double(count) * graphHeight / div);
+
+					if (prevValue <= 0)
+						prevY = y;
+
+					if (x > clientRect.left)
+					{
+						if (isFirstUpdate)
+						{
+							line.push_back({ x, y });
+							isFirstDraw = false;
+						}
+						else
+						{
+							if (isFirstDraw)
+								line.push_back({ prevX, prevY });
+							line.push_back({ x, y });
+							isFirstDraw = false;
+						}
+					}
+					if (x > clientRect.right)
+						break;
+					isFirstUpdate = false;
+					prevX = x;
+					prevY = y;
+					prevValue = count;
+
+					time += 0.25f;
+				}
+				if (line.size() > 1)
+					context.Polyline(line);
+			}
+			posY += GraphHeight;
+		}
+
+		if (m_config.showActiveProcesses && !m_trace.m_activeProcesses.empty())
+		{
+			auto drawBox = [&](u64 start, u64 stop, int height, bool selected, bool inProgress)
+				{
+					//int left = 3 + 6*m_activeFont.height;
+					//int right = rect.left + durationMs/100;
+					int posX = int(m_scrollPosX) + progressRect.left;
+					bool done = stop != ~u64(0);
+					if (!done)
+						stop = playTime;
+					int left = posX + int(float(TimeToS(start)) * scaleX);
+					int right = posX + int(float(TimeToS(stop)) * scaleX) - 1;
+
+					RECT rect;
+					rect.left = left;
+					rect.right = right;
+					rect.top = posY;
+					rect.bottom = posY + height;
+					context.FillRect(rect, inProgress ? m_processBrushes[selected].inProgress : m_processBrushes[selected].success);
+					return rect;
+				};
+
+			m_activeProcessCountHistory[m_activeProcessCountHistoryIterator++ % sizeof_array(m_activeProcessCountHistory)] = u32(m_trace.m_activeProcesses.size());
+
+			PaintActiveProcesses(context, posY, clientRect, [&](TraceView::ProcessLocation& processLocation, u32 boxHeight, bool firstWithHeight)
+				{
+					auto& session = m_trace.GetSession(m_traceView, processLocation.sessionIndex);
+					TraceView::Process& process = session.processors[processLocation.processorIndex].processes[processLocation.processIndex];
+
+					bool selected = m_processSelected && m_processSelectedLocation == processLocation;
+
+					if (m_config.showFinishedProcesses)
+					{
+						u32 index = processLocation.processIndex;
+						while (index > 0)
+						{
+							--index;
+							TraceView::Process& process2 = session.processors[processLocation.processorIndex].processes[index];
+							drawBox(process2.start, process2.stop, boxHeight, false, false);
+						}
+					}
+
+					RECT boxRect = drawBox(process.start, process.stop, boxHeight, selected, true);
+
+					u32 v = boxHeight - 1;
+
+					if (v > 4)
+					{
+						u32 fontIndex = Min(v, u32(sizeof_array(m_activeProcessFont) - 1));
+						if (!m_activeProcessFont[fontIndex].handle)
+						{
+							UpdateFont(m_activeProcessFont[fontIndex], fontIndex - 1, false);
+							m_activeProcessFont[fontIndex].offset += 1;
+						}
+						if (firstWithHeight)
+							context.SetActiveFont(m_activeProcessFont[fontIndex]);
+
+						StringBuffer<> str;
+						auto firstPar = -1;//process.description.find_first_of('(');
+						if (firstPar != -1 && *process.description.rbegin() == ')')
+						{
+							str.Append(process.description.c_str() + firstPar + 1).Resize(str.count-1);
+							str.Append(L"  ").Append(process.description.c_str(), firstPar);
+						}
+						else
+							str.Append(process.description);
+
+						if (process.isRemote)
+							str.Append(L" [").Append(session.name).Append(']');
+						else if (process.type == TraceView::ProcessType_CacheFetchTest || process.type == TraceView::ProcessType_CacheFetchDownload)
+							str.Append(L" [cache]");
+						if (boxRect.left < 0)
+							str.Appendf(L"   %s", TimeToText(playTime - process.start, true).str);
+						//drawIndentedText(str, LogEntryType_Info, 1, true);
+						drawStatusText(str, LogEntryType_Info, Max(int(boxRect.left+1), 1), boxRect.right, false);
+					}
+				});
+		}
+
+		int boxHeight = m_config.boxHeight;
+		int stepY = int(boxHeight) + 2;
+		int processStepY = boxHeight + 1;
+
+		TraceView::WorkRecord selectedWork;
+
+		SessionRec sortedSessions[1024];
+		Populate(sortedSessions, m_traceView, m_config.SortActiveRemoteSessions);
+
+		u32 visibleBoxes = 0;
+
+		TraceView::ProcessLocation processLocation { 0, 0, 0 };
+		for (u64 sessionIt = 0, sessionEnd = m_traceView.sessions.size(); sessionIt != sessionEnd; ++sessionIt)
+		{
+			bool isFirst = sessionIt == 0;
+			auto& session = *sortedSessions[sessionIt].session;
+			bool hasUpdates = !session.updates.empty();
+			if (!isFirst)
+			{
+				if (!hasUpdates && session.processors.empty())
+					continue;
+
+				if (!m_config.showFinishedProcesses && session.disconnectTime != ~u64(0))
+					continue;
+			}
+
+			processLocation.sessionIndex = sortedSessions[sessionIt].index;
+			if (!isFirst)
+				posY += 3;
+
+			if (m_config.showTitleBars)
+			{
+				if (posY + stepY >= progressRect.top && posY <= progressRect.bottom)
+				{
+					context.SetPen(m_separatorPen);
+					context.Line(0, posY, clientRect.right, posY);
+
+					StringBuffer<> text;
+					session.GetFullName(text);
+
+					if (hasUpdates && session.disconnectTime == ~u64(0))
+					{
+						u64 ping = session.ping.back();
+						//u64 memAvail = session.updates.back().memAvail;
+						//float cpuLoad = session.updates.back().cpuLoad;
+
+						//text.Appendf(L" - Cpu: %.1f%%", cpuLoad * 100.0f);
+						//if (memAvail)
+						//	text.Appendf(L" Mem: %ls/%ls", BytesToText(session.memTotal - memAvail).str, BytesToText(session.memTotal).str);
+						if (ping)
+							text.Appendf(L" Ping: %ls", TimeToText(ping, false, m_traceView.frequency).str);
+						if (!session.notification.empty())
+							text.Append(L" - ").Append(session.notification);
+					}
+					else if (!isFirst)
+					{
+						text.Append(L" - Disconnected");
+						if (!session.notification.empty())
+							text.Append(L" (").Append(session.notification).Append(')');
+					}
+
+					bool selected = m_sessionSelectedIndex == processLocation.sessionIndex;
+
+					int textBottom = Min(posY + m_sessionStepY, int(progressRect.bottom));
+
+					RECT rect;
+					rect.left = 5;
+					rect.right = clientRect.right;
+					rect.top = posY;
+					rect.bottom = textBottom;
+
+					if (selected)
+						context.SetBkMode(OPAQUE);
+
+					SIZE textSize = context.GetTextSize(text);
+					if (textSize.cx)
+						session.fullNameWidth = u32(textSize.cx);
+					context.ExtTextOutW(5, posY+2, ETO_CLIPPED, &rect, text);
+					if (selected)
+						context.SetBkMode(TRANSPARENT);
+				}
+				posY += m_sessionStepY;
+			}
+
+			bool showGraph = m_config.showNetworkStats || m_config.showCpuMemStats || m_config.showDriveStats;
+			if (showGraph && hasUpdates)
+			{
+				if (posY + GraphHeight >= progressRect.top && posY + GraphHeight - 5 < progressRect.bottom)
+				{
+					int posX = int(m_scrollPosX) + progressRect.left;
+					int graphBaseY = posY + GraphHeight - 4;
+					double graphHeight = double(GraphHeight - 2);
+
+					Vector<POINT> line;
+
+					auto drawGraph = [&](auto& values, auto maxValue, double scale, HPEN pen, bool accumulatingValue, int offsetY = 0)
+						{
+							bool loop = true;
+							u32 reconnectIndex = 0;
+							while (loop)
+							{
+								context.SetPen(pen);
+								bool isFirstUpdate = true;
+								bool isFirstDraw = true;
+								decltype(maxValue) prevValue = 0;
+								u64 prevTime = 0;
+								int prevX = 0;
+								int prevPrevX = 0;
+								int prevY = 0;
+
+								u64 i = 0;
+								u64 e = 0;
+
+								if (reconnectIndex)
+									i = session.reconnectIndices[reconnectIndex - 1];
+
+								if (reconnectIndex < session.reconnectIndices.size())
+									e = session.reconnectIndices[reconnectIndex];
+								else
+								{
+									e = Min(session.updates.size(), values.size());
+									loop = false;
+								}
+
+								line.clear();
+								for (; i!=e; ++i)
+								{
+									u64 updateTime = session.updates[i];
+									auto value = values[i];
+									int x = posX + int(float(TimeToS(updateTime)) * scaleX);
+									int y = graphBaseY;
+
+									double duration = TimeToS(updateTime - prevTime);
+									if (updateTime == 0)
+										isFirstUpdate = true;
+
+									if (accumulatingValue)
+									{
+										double invScaleY = duration * scale;
+										if (invScaleY != 0 && prevValue != 0)
+											y -= int(double(value - prevValue) / invScaleY) + offsetY;
+									}
+									else
+									{
+										double div = 1.0f;
+										if (maxValue != 0)
+											div = double(maxValue);
+										y -= int(double((double)maxValue - (double)value*scale)*graphHeight/div) + offsetY;
+									}
+
+									if (prevValue <= 0)
+										prevY = y;
+
+									if (x > clientRect.left)
+									{
+										if (isFirstUpdate)
+										{
+											isFirstDraw = false;
+										}
+										else
+										{
+											if (isFirstDraw)
+												line.push_back({prevPrevX, prevY});
+											line.push_back({prevX, y});
+											isFirstDraw = false;
+										}
+									}
+									if (prevX > clientRect.right)
+										break;
+									isFirstUpdate = false;
+									prevPrevX = prevX;
+									prevX = x;
+									prevY = y;
+									prevValue = value;
+									prevTime = updateTime;
+								}
+								if (line.size() > 1)
+									context.Polyline(line);
+								++reconnectIndex;
+							}
+						};
+
+					if (m_config.showNetworkStats && session.highestSendPerS != 0 && session.highestRecvPerS != 0)
+					{
+						drawGraph(session.networkSend, 100000000000000ull, double(session.highestSendPerS) / graphHeight, m_sendPen, true);
+						drawGraph(session.networkRecv, 100000000000000ull, double(session.highestRecvPerS) / graphHeight, m_recvPen, true, 1);
+					}
+
+					if (m_config.showCpuMemStats)
+					{
+						drawGraph(session.cpuLoad, 0.0f, -1.0, m_cpuPen, false);
+						drawGraph(session.memAvail, session.memTotal, 1.0, m_memPen, false);
+					}
+
+					if (m_config.showDriveStats)
+						for (auto& pair : session.drives)
+							if (pair.second.busyHighest)
+								drawGraph(pair.second.busyPercent, 0, -0.01, m_drivePen, false);
+
+				}
+				posY += GraphHeight;
+			}
+
+			if (m_config.showDetailedData)
+			{
+				auto drawText = [&](const StringBufferBase& text, RECT& rect, u32* outWidth)
+					{
+						if (rect.top > progressRect.bottom)
+							return;
+						bool selected = m_fetchedFilesSelected == processLocation.sessionIndex && text.StartsWith(TC("Fetched Files"));
+						if (selected)
+							context.SetBkMode(OPAQUE);
+						if (rect.bottom > progressRect.bottom)
+						{
+							rect.bottom = progressRect.bottom;
+							context.ExtTextOutW(rect.left, rect.top, ETO_CLIPPED, &rect, text);
+						}
+						else
+							context.ExtTextOutW(rect.left, rect.top, 0, NULL, text);
+						if (selected)
+							context.SetBkMode(TRANSPARENT);
+						if (outWidth)
+							*outWidth = u32(context.GetTextSize(text).cx);
+					};
+				bool isRemote = processLocation.sessionIndex != 0;
+				PaintDetailedStats(context, posY, progressRect, session, isRemote, playTime, drawText);
+			}
+
+			context.SetActiveFont(m_processFont);
+
+			bool shouldDrawText = m_processFont.height > 4;
+
+			if (m_config.showProcessBars)
+			{
+				processLocation.processorIndex = 0;
+				for (auto& processor : session.processors)
+				{
+					bool drawProcessorIndex = m_config.showFinishedProcesses;
+
+					if (posY + m_sessionStepY >= progressRect.top && posY < progressRect.bottom)
+					{
+						int barHeight = boxHeight;
+						int textOffsetY = 0;
+						if (posY + boxHeight > progressRect.bottom)
+						{
+							int newBarHeight = Min(barHeight, int(progressRect.bottom - posY));
+							textOffsetY = int(barHeight - newBarHeight);
+							barHeight = newBarHeight;
+						}
+
+						const int textHeight = barHeight;
+						const int rectBottom = posY + textHeight;
+						const int offsetY = (textHeight - m_processFont.height + textOffsetY) / 2;
+
+						processLocation.processIndex = 0;
+						int posX = int(m_scrollPosX) + progressRect.left;
+						for (auto& process : processor.processes)
+						{
+							int left = posX + int(float(TimeToS(process.start)) * scaleX);
+
+							auto pig = MakeGuard([&]() { ++processLocation.processIndex; });
+
+							if (left >= progressRect.right)
+								continue;
+
+							u64 stop = process.stop;
+							bool done = stop != ~u64(0);
+							if (!done)
+								stop = playTime;
+							else if (!m_config.showFinishedProcesses)
+								continue;
+
+							drawProcessorIndex = true;
+
+							RECT rect;
+							rect.left = left;
+							rect.right = Max(posX + int(float(TimeToS(stop)) * scaleX) - 1, left + 1);
+							rect.top = posY;
+							rect.bottom = rectBottom - 1;
+
+							if (rect.right <= progressRect.left)
+								continue;
+
+							if (!m_filterString.empty() && !Contains(process.description.c_str(), m_filterString.c_str()) && !Contains(process.breadcrumbs.c_str(), m_filterString.c_str()))
+								continue;
+							++visibleBoxes;
+
+							bool selected = m_processSelected && m_processSelectedLocation == processLocation;
+							if (selected)
+								process.bitmapDirty = true;
+
+							--rect.top;
+							PaintProcessRect(context, process, rect, progressRect, selected, false);
+							++rect.top;
+
+							int processWidth = rect.right - rect.left;
+							if (shouldDrawText && m_config.ShowProcessText && processWidth > 3)
+							{
+								context.DrawCached(process, rect, offsetY, textHeight, progressRect, [&](PaintContext& textCacheContext)
+									{
+										RECT rect2{ 0, int(process.bitmapOffset), 256, int(process.bitmapOffset) + m_processFont.height };
+										RECT rect3{ 0, int(process.bitmapOffset), processWidth, int(process.bitmapOffset) + m_processFont.height };
+										if (!done)
+											rect3.right = 256;
+
+										PaintProcessRect(textCacheContext, process, rect3, rect2, selected, true);
+
+										rect2.left += 3; // Move in text a bit
+
+										int textY = rect2.top + m_processFont.offset;
+
+										bool dropShadow = m_config.DarkMode;
+										if (dropShadow)
+										{
+											textCacheContext.SetTextColor(RGB(5, 60, 5));
+											++rect2.left;
+											++textY;
+											textCacheContext.ExtTextOutW(rect2.left, textY, ETO_CLIPPED, &rect2, process.description);
+											--rect2.left;
+											--textY;
+										}
+										textCacheContext.SetTextColor(m_textColor);
+										textCacheContext.ExtTextOutW(rect2.left, textY, ETO_CLIPPED, &rect2, process.description);
+
+										if (!selected)
+											process.bitmapDirty = false;
+									});
+							}
+						}
+
+						if (drawProcessorIndex)
+						{
+							RECT rect;
+							rect.left = 5;
+							rect.right = progressRect.left - 2;
+							rect.top = posY;
+							rect.bottom = rectBottom;
+
+							StringBuffer<> buf;
+							buf.AppendValue(u64(processLocation.processorIndex) + 1);
+							context.ExtTextOutW(5, posY + offsetY, ETO_CLIPPED, &rect, buf);
+						}
+					}
+
+					lastStop = Max(lastStop, processor.processes.rbegin()->stop);
+
+					++processLocation.processorIndex;
+
+					if (drawProcessorIndex)
+						posY += processStepY;
+				}
+			}
+			else
+			{
+				for (auto& processor : session.processors)
+					if (!processor.processes.empty())
+						lastStop = Max(lastStop, processor.processes.rbegin()->stop);
+			}
+
+			if (m_config.showWorkers && isFirst)
+			{
+				u32 trackIndex = 0;
+				for (auto& workTrack : m_traceView.workTracks)
+				{
+					if (posY + m_sessionStepY >= progressRect.top && posY <= progressRect.bottom)
+					{
+						int textOffsetY = 0;
+						int barHeight = boxHeight;
+						if (posY + int(boxHeight) > progressRect.bottom)
+						{
+							int newBarHeight = Min(barHeight, int(progressRect.bottom - posY));
+							textOffsetY = barHeight - newBarHeight;
+							barHeight = newBarHeight;
+						}
+
+						const int textHeight = barHeight;
+						const int rectBottom = posY + textHeight;
+						const int offsetY = (textHeight - m_processFont.height + textOffsetY) / 2;
+
+						if (shouldDrawText)
+						{
+							RECT rect;
+							rect.left = 5;
+							rect.right = progressRect.left - 5;
+							rect.top = posY;
+							rect.bottom = rectBottom;
+
+							StringBuffer<> buf;
+							buf.AppendValue(u64(trackIndex) + 1);
+							context.ExtTextOutW(5, posY + offsetY, ETO_CLIPPED, &rect, buf);
+						}
+
+						int lastDrawnRight = 0;
+
+						u32 workIndex = 0;
+						int posX = int(m_scrollPosX) + progressRect.left;
+						for (auto& work : workTrack.records)
+						{
+							auto inc = MakeGuard([&](){ ++workIndex; });
+
+							if (work.start == work.stop)
+								continue;
+
+							float startTime = TimeToS(work.start);
+
+							int left = posX + int(startTime * scaleX);
+
+							if (left >= progressRect.right)
+								continue;
+
+							if (!m_filterString.empty())
+							{
+								bool keep = Contains(work.description, m_filterString.c_str());
+								if (!keep)
+									for (auto& en : work.entries)
+										keep |= Contains(en.text, m_filterString.c_str());
+								if (!keep)
+									continue;
+							}
+
+							HBRUSH brush;
+
+							u64 stop = work.stop;
+
+							float stopTime = TimeToS(stop);
+
+							RECT rect;
+							rect.left = left;
+							rect.right = posX + int(stopTime * scaleX) - 1;
+							rect.top = posY;
+							rect.bottom = rectBottom - 1;
+
+							if (rect.right <= progressRect.left)
+								continue;
+
+							++visibleBoxes;
+
+							rect.right = Max(int(rect.right), left + 1);
+
+							Color color = work.color;
+							bool selected = m_workSelected && m_workTrack == trackIndex && m_workIndex == workIndex;
+							if (selected)
+							{
+								selectedWork = work;
+								work.bitmapDirty = true;
+								auto c = (u8*)&color;
+								for (u32 j=0; j!=3; ++j)
+									c[j] = u8(Min(int(c[j]) + 40, 255));
+							}
+							else
+							{
+								if (rect.left + 1 == rect.right)
+									if (lastDrawnRight == rect.right)
+										continue;
+							}
+
+							lastDrawnRight = rect.right;
+
+							bool done = stop != ~u64(0);
+							if (done)
+							{
+								auto insres = m_coloredBrushes.try_emplace(color);
+								if (insres.second)
+								{
+									auto c = (u8*)&color;
+									insres.first->second = CreateSolidBrush(RGB(c[2], c[1], c[0]));
+								}
+								brush = insres.first->second;
+							}
+							else
+							{
+								//stop = playTime;
+								brush = m_processBrushes[0].inProgress;
+							}
+
+							--rect.top;
+
+							auto clampRect = [&](RECT& r) { r.left = Min(Max(r.left, progressRect.left), progressRect.right); r.right = Max(Min(r.right, progressRect.right), progressRect.left); };
+
+							clampRect(rect);
+							context.FillRect(rect, brush);
+							++rect.top;
+
+							int processWidth = rect.right - rect.left;
+							if (shouldDrawText && m_config.ShowProcessText && processWidth > 3)
+							{
+								context.DrawCached(work, rect, offsetY, textHeight, progressRect, [&](PaintContext& textCacheContext)
+									{
+										RECT rect2{ 0,int(work.bitmapOffset), 256, int(work.bitmapOffset) + m_processFont.height };
+										textCacheContext.FillRect(rect2, brush);
+										textCacheContext.ExtTextOutW(rect2.left, rect2.top, ETO_CLIPPED, &rect2, ToView(work.description));
+										if (!selected)
+											work.bitmapDirty = false;
+									});
+							}
+						}
+					}
+					++trackIndex;
+					posY += processStepY;
+				}
+			}
+
+			context.SetActiveFont(m_defaultFont);
+		}
+
+		if (true)
+		{
+			Font& activeFont = context.GetActiveFont();
+			for (auto& kv : m_traceView.tables)
+			{
+				auto& table = kv.second;
+				u32 tableOffsetX = activeFont.height*40;
+				u32 tableOffsetY = 30;
+
+				u32 rowOffset = activeFont.height + 2;
+				u32 columnOffset = table.rowCharWidth * activeFont.height;
+
+				u32 columnIndex = 0;
+				u32 rowIndex = 0;
+				for (const tchar* item : table.items)
+				{
+					u32 y = tableOffsetY + rowIndex*rowOffset;
+					u32 x = tableOffsetX + columnIndex*columnOffset;
+					if (int(x) >= progressRect.right)
+						break;
+
+					context.ExtTextOutW(x, y, ETO_CLIPPED, nullptr, ToView(item));
+
+					++rowIndex;
+					if (rowIndex > table.rowMaxCount || int(y) >= progressRect.bottom)
+					{
+						rowIndex = 0;
+						++columnIndex;
+					}
+				}
+			}
+		}
+
+		m_contentWidth = m_progressRectLeft + Max(0, int(TimeToS((lastStop != 0 && lastStop != ~u64(0)) ? lastStop : playTime) * scaleX));
+
+
+		m_contentHeight = posY - int(m_scrollPosY) + stepY + 14;
+
+		float timelineSelected = m_timelineSelected;
+
+		if (m_config.showTimeline && !m_traceView.sessions.empty())
+			PaintTimeline(context, clientRect);
+
+		if (!m_filterString.empty())
+		{
+			StringBuffer<> str[2];
+			str[0].Append(TCV("Box Filter: ")).Append(m_filterString);
+			str[1].Appendf(TC("Box Count: %u"), visibleBoxes);
+
+			context.SetActiveFont(m_popupFont);
+			context.DrawCenteredText(str, 2, true);
+		}
+
+		if (m_config.showCursorLine && m_mouseOverWindow)
+		{
+			float timeScale = (m_config.horizontalScaleValue * m_zoomValue)*50.0f;
+			float startOffset = -(m_scrollPosX / timeScale);
+			POINT pos = context.GetCursorClientPos();;
+			timelineSelected = startOffset + float(pos.x - m_progressRectLeft) / timeScale;
+		}
+
+		if (timelineSelected != 0)
+		{
+			int posX = int(m_scrollPosX) + progressRect.left;
+			int left = posX + int(timelineSelected * scaleX);
+			int timelineTop = GetTimelineTop(clientRect);
+
+			// TODO: Draw line up
+			context.Line(left, 2, left, timelineTop);
+
+			if (timelineSelected >= 0)
+			{
+				StringBuffer<> b;
+				u32 milliseconds = u32(timelineSelected * 1000.0f);
+				u32 seconds = milliseconds / 1000;
+				milliseconds -= seconds * 1000;
+				u32 minutes = seconds / 60;
+				seconds -= minutes * 60;
+				u32 hours = minutes / 60;
+				minutes -= hours * 60;
+				if (hours)
+				{
+					b.AppendValue(hours).Append('h');
+					if (minutes < 10)
+						b.Append('0');
+				}
+				if (minutes || hours)
+				{
+					b.AppendValue(minutes).Append('m');
+					if (seconds < 10)
+						b.Append('0');
+				}
+				b.AppendValue(seconds).Append('.');
+				if (milliseconds < 100)
+					b.Append('0');
+				if (milliseconds < 10)
+					b.Append('0');
+				b.AppendValue(milliseconds);
+
+				StringBuffer<128> wt = GetWorldTime(timelineSelected);
+
+				context.SetActiveFont(m_popupFont);
+				DrawTextLogger logger(context, m_popupFont.height, m_tooltipBackgroundBrush);
+
+				logger.Info(L"%s (%s)", b.data, wt.data);
+				logger.DrawAtPos(left + 4, timelineTop - 20);
+			}
+		}
+
+		{
+			int boxSide = 8;
+			int boxStride = boxSide + 2;
+			int top = 5;
+			int bottom = top + boxSide;
+			int left = progressRect.right - 7 - boxSide;
+			int right = progressRect.right - 7;
+			bool* values = &m_config.showProgress;
+			for (int i = VisualizerFlag_Count - 1; i >= 0; --i)
+			{
+				context.SetPen(m_buttonSelected == u32(i) ? m_textPen : m_checkboxPen);
+				context.Rectangle(left, top, right, bottom);
+
+				if (values[i])
+				{
+					context.Line(left + 2, top + 2, right - 2, bottom - 2);
+					context.Line(right - 3, top + 2, left + 1, bottom - 2);
+				}
+
+				left -= boxStride;
+				right -= boxStride;
+			}
+
+			top -= 2;
+
+			auto drawText = [&](StringView text, COLORREF color)
+				{
+					context.SetTextColor(color);
+					left -= (context.GetTextSize(text).cx) + 5;
+					context.ExtTextOutW(left, top, 0, NULL, text);
+				};
+
+			if (m_config.showDriveStats)
+			{
+				context.SetActiveFont(m_defaultFont);
+				drawText(TCV("DRV"), m_driveColor);
+			}
+
+			if (m_config.showNetworkStats)
+			{
+				context.SetActiveFont(m_defaultFont);
+				drawText(TCV("SND"), m_sendColor);
+				drawText(TCV("RCV"), m_recvColor);
+			}
+
+			if (m_config.showCpuMemStats)
+			{
+				context.SetActiveFont(m_defaultFont);
+				drawText(TCV("CPU"), m_cpuColor);
+				drawText(TCV("MEM"), m_memColor);
+			}
+
+			context.SetTextColor(m_textColor);
+		}
+
+		if (m_processSelected)
+		{
+			const TraceView::Process& process = m_traceView.GetProcess(m_processSelectedLocation);
+			u64 duration = 0;
+
+			Vector<TString> logLines;
+			u32 maxCharCount = 50u;
+
+			bool hasExited = process.stop != ~u64(0);
+			if (hasExited)
+			{
+				duration = process.stop - process.start;
+
+				if (!process.logLines.empty())
+				{
+					u32 lineMaxCount = 0;
+					for (auto& line : process.logLines)
+					{
+						u32 offset = 0;
+						u32 left = u32(line.text.size());
+						while (left)
+						{
+							u32 toCopy = Min(left, maxCharCount);
+							lineMaxCount = Max(lineMaxCount, toCopy);
+							logLines.push_back(line.text.substr(offset, toCopy));
+							offset += toCopy;
+							left -= toCopy;
+						}
+					}
+				}
+			}
+			else
+			{
+				duration = playTime - process.start;
+			}
+
+			context.SetActiveFont(m_popupFont);
+			DrawTextLogger logger(context, m_popupFont.height, m_tooltipBackgroundBrush);
+
+			logger.AddTextOffset(-10); // Remove spaces in the front
+			logger.AddWidth(3);
+
+			logger.AddSpace(2);
+			logger.Info(L"  %ls", process.description.c_str());
+			if (process.type != TraceView::ProcessType_Task)
+			{
+				logger.Info(L"  Host:        %ls", m_processSelectedLocation.sessionIndex == 0 ? TC("local") : m_traceView.GetSession(m_processSelectedLocation).name.c_str());
+				logger.Info(L"  ProcessId:  %6u", process.id);
+			}
+			logger.Info(L"  Start:     %7ls (%ls)", TimeToText(process.start, true).str, GetWorldTime(process.start).data);
+			logger.Info(L"  Duration:  %7ls", TimeToText(duration, true).str);
+			if (!process.returnedReason.empty())
+				logger.Info(L"  Returned: %7s", process.returnedReason.data());
+			if (hasExited && process.exitCode != 0)
+			{
+				if (process.exitCode == ProcessCancelExitCode)
+					logger.Info(L"  ExitCode: Cancelled");
+				else if (process.exitCode >= 0xC0000000u)
+					logger.Info(L"  ExitCode: 0x%x", process.exitCode);
+				else
+					logger.Info(L"  ExitCode: %7u", process.exitCode);
+			}
+
+			const auto& breadcrumbs = process.breadcrumbs;
+			if (!breadcrumbs.empty())
+			{
+				constexpr TString::size_type maxLineLen = 50;
+				logger.Info(L"");
+				logger.Info(L"  ------------ Breadcrumbs ------------");
+				for (TString::size_type lineStart = 0, lineEnd = 0; lineEnd < breadcrumbs.size(); lineStart = lineEnd + 1)
+				{
+					// Log each individual line
+					lineEnd = breadcrumbs.find(L'\n', lineStart);
+					TString line = (lineEnd == TString::npos ? breadcrumbs.substr(lineStart) : breadcrumbs.substr(lineStart, lineEnd - lineStart));
+
+					// Break each line down into smaller section if they are longer than the maximum allowed length
+					if (line.size() > maxLineLen)
+					{
+						for (TString::size_type sectionStart = 0, sectionEnd = 0; sectionStart < line.size(); sectionStart = sectionEnd)
+						{
+							const TString::size_type maxSectionLen = sectionStart == 0 ? maxLineLen : maxLineLen - 2;
+							sectionEnd = std::min<TString::size_type>(sectionEnd + maxSectionLen, line.size());
+							TString section = (sectionStart == 0 ? L"  " : L"    ") + line.substr(sectionStart, sectionEnd - sectionStart);
+							logger.Info(section.c_str());
+						}
+					}
+					else
+					{
+						line = L"  " + line;
+						logger.Info(line.c_str());
+					}
+				}
+			}
+
+			if (process.stop != ~u64(0) && !process.stats.empty())
+			{
+				BinaryReader reader(process.stats.data(), 0, process.stats.size());
+				ProcessStats processStats;
+				SessionStats sessionStats;
+				StorageStats storageStats;
+				KernelStats kernelStats;
+				CacheFetchStats cacheStats;
+
+				if (process.type == TraceView::ProcessType_CacheFetchTest || process.type == TraceView::ProcessType_CacheFetchDownload)
+				{
+					if (!process.returnedReason.empty())
+						logger.Info(L"  Cache:       Miss");
+					else
+						logger.Info(L"  Cache:        Hit");
+					cacheStats.Read(reader, m_traceView.version);
+					if (reader.GetLeft())
+					{
+						storageStats.Read(reader, m_traceView.version);
+						kernelStats.Read(reader, m_traceView.version);
+					}
+				}
+				else
+				{
+					processStats.Read(reader, m_traceView.version);
+
+					if (reader.GetLeft())
+					{
+						if (process.isRemote || (m_traceView.version >= 36 && !process.isReuse))
+							sessionStats.Read(reader, m_traceView.version);
+						storageStats.Read(reader, m_traceView.version);
+						kernelStats.Read(reader, m_traceView.version);
+					}
+				}
+
+				if (processStats.hostTotalTime)
+				{
+					logger.Info(L"");
+					logger.Info(L"  ----------- Detours stats -----------");
+					processStats.Print(logger, m_traceView.frequency);
+				}
+				else if (processStats.peakMemory)
+				{
+					logger.Info(L"");
+					logger.Info(L"  ----------- Process stats -----------");
+					processStats.Print(logger, m_traceView.frequency);
+				}
+
+				if (!sessionStats.IsEmpty())
+				{
+					logger.Info(L"");
+					logger.Info(L"  ----------- Session stats -----------");
+					sessionStats.Print(logger, m_traceView.frequency);
+				}
+
+				if (!cacheStats.IsEmpty())
+				{
+					logger.Info(L"");
+					logger.Info(L"  ------------ Cache stats ------------");
+					cacheStats.Print(logger, m_traceView.frequency);
+				}
+
+				if (!storageStats.IsEmpty())
+				{
+					logger.Info(L"");
+					logger.Info(L"  ----------- Storage stats -----------");
+					storageStats.Print(logger, m_traceView.frequency);
+				}
+
+				if (!kernelStats.IsEmpty())
+				{
+					logger.Info(L"");
+					logger.Info(L"  ----------- Kernel stats ------------");
+					kernelStats.Print(logger, false, m_traceView.frequency);
+				}
+
+				PrintCacheWriteStats(logger, process.id);
+
+				if (!logLines.empty())
+				{
+					logger.Info(L"");
+					logger.Info(L"  ---------------- Log ----------------");
+					logger.AddTextOffset(14);
+					for (auto& line : logLines)
+						logger.Log(LogEntryType_Info, line);
+				}
+			}
+			logger.AddSpace(3);
+			logger.DrawAtCursor();
+		}
+		else if (m_workSelected && selectedWork.description)
+		{
+			u64 duration;
+			if (selectedWork.stop != ~u64(0))
+				duration = selectedWork.stop - selectedWork.start;
+			else
+				duration = playTime - selectedWork.start;
+
+			context.SetActiveFont(m_popupFont);
+			DrawTextLogger logger(context, m_popupFont.height, m_tooltipBackgroundBrush);
+
+			logger.AddSpace();
+			logger.Info(L"  %ls", selectedWork.description);
+			logger.Info(L"  Start:     %ls (%ls)", TimeToText(selectedWork.start, true).str, GetWorldTime(selectedWork.start).data);
+			logger.Info(L"  Duration:  %ls", TimeToText(duration, true).str);
+			if (!selectedWork.entries.empty())
+				logger.AddSpace();
+			for (u32 i=0,e=u32(selectedWork.entries.size()); i!=e; ++i)
+			{
+				auto& entry = selectedWork.entries[i];
+				u64 time = 0;
+				if (i == 0)
+				{
+					u64 start = entry.startTime ? entry.startTime : entry.time;
+					if (TimeToMs(start - selectedWork.start) > 1)
+						logger.Info(L"   Start (%ls)", TimeToText(start - selectedWork.start).str);
+				}
+
+				if (entry.startTime)
+					time = entry.time - entry.startTime;
+
+				if (!time)
+					for (u32 j=i+1;j<e;++j)
+					{
+						auto& next = selectedWork.entries[j];
+						if (next.startTime)
+							continue;
+						time = next.time - entry.time;
+						break;
+					}
+
+				if (!time && selectedWork.stop != ~0u)
+					time = selectedWork.stop - entry.time;
+				if (entry.count == 1)
+					logger.Info(L"%s  %ls (%ls)", (entry.startTime ? L" " : L""), entry.text, TimeToText(time).str);
+				else
+					logger.Info(L"%s  %ls (%ls %u)", (entry.startTime ? L" " : L""), entry.text, TimeToText(time).str, entry.count);
+			}
+			logger.AddSpace();
+			logger.DrawAtCursor();
+		}
+		else if (m_sessionSelectedIndex != ~0u)
+		{
+			context.SetActiveFont(m_popupFont);
+			DrawTextLogger logger(context, m_popupFont.height, m_tooltipBackgroundBrush);
+			logger.AddWidth(3);
+			logger.AddSpace(2);
+
+			auto& session = m_traceView.sessions[m_sessionSelectedIndex];
+			Vector<TString>& summary = session.summary;
+			if (summary.empty())
+			{
+				if (m_traceView.finished)
+					logger.Info(L" Session summary was never received");
+				else
+					logger.Info(L" Session summary not available until session is done");
+			}
+
+			logger.AddTextOffset(-10); // Remove spaces in the front
+			for (auto& line : summary)
+				logger.Log(LogEntryType_Info, line);
+			logger.AddTextOffset(0);
+
+			if (m_sessionSelectedIndex == 0)
+			{
+				for (auto& cacheSummary : m_traceView.cacheSummaries)
+					for (auto& line : cacheSummary.lines)
+						logger.Log(LogEntryType_Info, line);
+			}
+
+			for (auto& pair : session.drives)
+			{
+				auto& d = pair.second;
+				logger.Info(L"  %c: Read: %s (%u) Write: %s (%u)", pair.first, BytesToText(d.totalReadBytes).str, d.totalReadCount, BytesToText(d.totalWriteBytes).str, d.totalWriteCount);
+			}
+
+			logger.AddSpace(3);
+			logger.DrawAtCursor();
+		}
+		else if (m_statsSelected)
+		{
+			context.SetActiveFont(m_popupFont);
+			DrawTextLogger logger(context, m_popupFont.height, m_tooltipBackgroundBrush);
+			logger.AddSpace(3);
+			logger.SetColor(m_cpuColor).Info(L"  Cpu: %.1f%%", m_stats.cpuLoad * 100.0f);
+			logger.SetColor(m_memColor).Info(L"  Mem: %ls/%ls", BytesToText(m_stats.memTotal - m_stats.memAvail).str, BytesToText(m_stats.memTotal).str);
+			if (m_stats.recvBytes || m_stats.sendBytes)
+			{
+				logger.SetColor(m_recvColor).Info(L"  Recv: %lsit/s", BytesToText(m_stats.recvBytesPerSecond*8).str);
+				logger.SetColor(m_sendColor).Info(L"  Send: %lsit/s", BytesToText(m_stats.sendBytesPerSecond*8).str);
+			}
+			logger.SetColor(m_textColor);
+
+			if (m_stats.ping)
+				logger.Info(L"  Ping: %ls", TimeToText(m_stats.ping, false, m_traceView.frequency).str);
+			if (m_stats.connectionCount)
+				logger.Info(L"  Conn: %u", u32(m_stats.connectionCount));
+
+			for (auto& pair : m_stats.drives)
+				logger.SetColor(m_driveColor).Info(L"  %c: %u%% R:%s/s W:%s/s", pair.first, pair.second.busyPercent, BytesToText(pair.second.readPerSecond).str, BytesToText(pair.second.writePerSecond).str);
+
+			logger.AddSpace(3);
+			logger.DrawAtCursor();
+		}
+		else if (m_activeProcessGraphSelected)
+		{
+			context.SetActiveFont(m_popupFont);
+			DrawTextLogger logger(context, m_popupFont.height, m_tooltipBackgroundBrush);
+			logger.AddSpace(3);
+			logger.AddSpace(3);
+			logger.SetColor(m_activeProcColor).Info(L"  Active Processes: %u", m_activeProcessCount);
+			logger.DrawAtCursor();
+		}
+		else if (m_buttonSelected != ~0u)
+		{
+			const wchar_t* tooltip[] =
+			{
+#define UBA_VISUALIZER_FLAG(name, defaultValue, desc) desc,
+				UBA_VISUALIZER_FLAGS1
+#undef UBA_VISUALIZER_FLAG
+			};
+
+			context.SetActiveFont(m_popupFont);
+			DrawTextLogger logger(context, m_popupFont.height, m_tooltipBackgroundBrush);
+			logger.Info(L"%ls %ls", L"Show", tooltip[m_buttonSelected]);
+			logger.DrawAtCursor();
+		}
+		else if (m_fetchedFilesSelected != ~0u)
+		{
+			auto& session = m_traceView.sessions[m_fetchedFilesSelected];
+			auto& fetchedFiles = session.fetchedFiles;
+			if (!fetchedFiles.empty() && !fetchedFiles[0].hint.empty())
+			{
+				/*
+				int colWidth = 500;
+				int width = colWidth * 2;
+				int height = Min(int(clientRect.bottom), int(fetchedFiles.size() * m_popupFont.height));
+
+				context.SetActiveFont(m_defaultFont);
+				DrawTextLogger logger(context, r, m_font.height, m_tooltipBackgroundBrush);
+				for (auto& f : fetchedFiles)
+				{
+				if (f.hint == TC("KnownInput"))
+				continue;
+				if (logger.rect.top >= r.bottom - m_font.height)
+				{
+				if (logger.rect.left + colWidth >= r.right)
+				{
+				logger.Info(L"...");
+				break;
+				}
+				logger.rect.top = r.top;
+				logger.rect.left += colWidth;
+				}
+				logger.Info(L"%s", f.hint.c_str());
+				}
+				logger.DrawAtCursor();
+				*/
+			}
+		}
+	}
+
+	void Visualizer::PaintProcessRect(PaintContext& context, TraceView::Process& process, RECT rect, const RECT& progressRect, bool selected, bool writingBitmap)
+	{
+		auto clampRect = [&](RECT& r) { r.left = Min(Max(r.left, progressRect.left), progressRect.right); r.right = Max(Min(r.right, progressRect.right), progressRect.left); };
+		bool done = process.stop != ~u64(0);
+
+		HBRUSH brush = m_processBrushes[selected].success;
+		if (!process.returnedReason.empty())
+			brush = m_processBrushes[selected].returned;
+		else if (!done)
+			brush = process.isRacing ? m_processBrushes[selected].inProgressRace : m_processBrushes[selected].inProgress;
+		else if (process.type == TraceView::ProcessType_CacheFetchTest)
+			brush = m_processBrushes[selected].cacheFetchTest;
+		else if (process.type == TraceView::ProcessType_CacheFetchDownload)
+			brush = m_processBrushes[selected].cacheFetchDownload;
+		else if (process.exitCode == ProcessCancelExitCode)
+			brush = m_processBrushes[selected].returned;
+		else if (process.exitCode != 0)
+			brush = m_processBrushes[selected].error;
+
+		u64 writeFilesTime = process.writeFilesTime;
+
+		if (!done || process.exitCode != 0 || !m_config.ShowReadWriteColors)
+		{
+			if (writingBitmap)
+				rect.right = 256;
+			clampRect(rect);
+			context.FillRect(rect, brush);
+			return;
+		}
+
+		double duration = double(process.stop - process.start);
+
+		RECT main = rect;
+		int width = rect.right - rect.left;
+
+		double recvPart = (double(ConvertTime(m_traceView, process.createFilesTime)) / duration);
+		if (int headSize = int(recvPart * width))
+		{
+			main.left += headSize;
+			RECT r2 = rect;
+			r2.right = r2.left + headSize;
+			clampRect(r2);
+			if (r2.left != r2.right)
+				context.FillRect(r2, m_processBrushes[selected].recv);
+		}
+
+		double sendPart = (double(ConvertTime(m_traceView, writeFilesTime)) / duration);
+		if (int tailSize = int(sendPart * width))
+		{
+			main.right -= tailSize;
+			RECT r2 = rect;
+			r2.left = r2.right - tailSize;
+			clampRect(r2);
+			if (r2.left != r2.right)
+				context.FillRect(r2, m_processBrushes[selected].send);
+		}
+
+		clampRect(main);
+		if (main.left != main.right)
+			context.FillRect(main, brush);
+	}
+
+	void Visualizer::PaintActiveProcesses(PaintContext& context, int& posY, const RECT& clientRect, const Function<void(TraceView::ProcessLocation&, u32, bool)>& drawProcess)
+	{
+		context.SetActiveFont(m_processFont);
+		int startPosY = posY;
+
+		Map<u64, TraceView::ProcessLocation*> activeProcesses;
+		u32 remoteCount = 0;
+		for (auto& kv : m_trace.m_activeProcesses)
+		{
+			auto& active = kv.second;
+			TraceView::Process& process = m_trace.GetSession(m_traceView, active.sessionIndex).processors[active.processorIndex].processes[active.processIndex];
+			u64 start = process.start; //~0llu - process.start;
+			activeProcesses.try_emplace(start, &active);
+			if (process.isRemote)
+				++remoteCount;
+		}
+
+		Font& activeFont = context.GetActiveFont();
+
+		u32 maxHeight = clientRect.bottom;
+
+		bool fillHeight = ActiveProcessesShouldFillHeight();
+		if (fillHeight)
+		{
+			maxHeight = clientRect.bottom - posY;
+			if (m_config.showTimeline)
+				maxHeight -= GetTimelineHeight();
+		}
+		else
+		{
+			u32 maxHeight2 = m_config.maxActiveVisible * (activeFont.height + 2);
+			maxHeight = Min(maxHeight, maxHeight2);
+		}
+
+		u32 maxSize = Min(m_config.maxActiveProcessHeight, 32u);
+		u32 maxSizeMinusOne = maxSize - 1;
+
+		u32 counts[128] = { 0 };
+
+		u32 highestHistoryCount = 0;
+		for (u32 i=0; i!= sizeof_array(m_activeProcessCountHistory) - 1; ++i)
+			highestHistoryCount = Max(highestHistoryCount, m_activeProcessCountHistory[i]);
+
+		u32 activeProcessCount = highestHistoryCount;
+		counts[0] = activeProcessCount;
+		u32 totalHeight = counts[0]*2;
+		while (totalHeight < maxHeight && counts[maxSizeMinusOne] != activeProcessCount)
+		{
+			bool changed = false;
+			for (u32 i=0; i!=maxSizeMinusOne; ++i)
+			{
+				if (counts[i] && counts[i] > counts[i+1]*2 + 1)
+				{
+					--counts[i];
+					++counts[i+1];
+					++totalHeight;
+					changed = true;
+				}
+			}
+			if (!changed)
+			{
+				for (u32 j=0; j!=maxSizeMinusOne; ++j)
+				{
+					if (!counts[j])
+						continue;
+					++counts[j + 1];
+					--counts[j];
+					++totalHeight;
+					break;
+				}
+			}
+		}
+
+		auto it = activeProcesses.begin();
+		auto itEnd = activeProcesses.end();
+		int endY = int(startPosY + maxHeight);
+		for (u32 i=0;i!=maxSize;++i)
+		{
+			u32 v = maxSizeMinusOne - i;
+			u32 boxHeight = v + 1;
+
+			for (u32 j=0;j!=counts[v];++j)
+			{
+				if (it == itEnd || posY >= endY)
+					break;
+				auto& active = *it->second;
+				++it;
+				drawProcess(active, boxHeight, j == 0);
+				posY += boxHeight + 1;
+			}
+		}
+
+		//if (fillHeight)
+		//	posY = clientRect.bottom-1;//startPosY; // To prevent vertical scrollbar
+		if (fillHeight || counts[maxSizeMinusOne] != activeProcessCount)
+			posY = startPosY + maxHeight;
+		else
+			posY += 3;
+
+		context.SetActiveFont(m_defaultFont);
+	}
+
+	void Visualizer::PaintTimeline(PaintContext& context, const RECT& clientRect)
+	{
+		context.SetActiveFont(m_timelineFont);
+		int top = GetTimelineTop(clientRect);
+		float timeScale = (m_config.horizontalScaleValue*m_zoomValue)*50.0f;
+		float startOffset = ((m_scrollPosX/timeScale) - float(int(m_scrollPosX/timeScale))) * timeScale;
+		int index = -int(startOffset/timeScale);
+
+		int number = -int(float(m_scrollPosX)/timeScale);
+
+		int textStepSize = int((5.0f / timeScale) + 1) * 5;
+		if (textStepSize > 150)
+			textStepSize = 600;
+		else if (textStepSize > 120)
+			textStepSize = 300;
+		else if (textStepSize > 90)
+			textStepSize = 240;
+		else if (textStepSize > 45)
+			textStepSize = 120;
+		else if (textStepSize > 30)
+			textStepSize = 60;
+		else if (textStepSize > 10)
+			textStepSize = 30;
+
+		int lineStepSize = textStepSize / 5;
+
+		RECT progressRect = clientRect;
+		progressRect.left += m_progressRectLeft;
+
+		context.SetPen(m_textPen);
+		Vector<DWORD> lines;
+		Vector<POINT> points;
+
+		while (true)
+		{
+			int pos = progressRect.left + int(startOffset + float(index)*timeScale);
+			if (pos >= clientRect.right)
+				break;
+
+			int lineBottom = top + 5;
+			if (!(number % textStepSize))
+			{
+				bool shouldDraw = true;
+				int seconds = number;
+				StringBuffer<> buffer;
+				if (seconds >= 60)
+				{
+					int min = seconds / 60;
+					seconds -= min * 60;
+					if (!seconds)
+					{
+						buffer.Appendf(L"%um", min);
+						lineBottom += 4;
+					}
+				}
+				if (!number || seconds)
+					buffer.Appendf(L"%u", seconds);
+
+				if (shouldDraw)
+				{
+					SIZE s = context.GetTextSize(buffer);
+					RECT textRect;
+					textRect.top = top + 8;
+					textRect.bottom = textRect.top + m_timelineFont.height;
+					textRect.right = pos + s.cx/2;
+					textRect.left = pos - s.cx/2;
+					context.ExtTextOutW(textRect.left, textRect.top, 0, NULL, buffer);
+				}
+			}
+			if (!(number % lineStepSize))
+			{
+				points.push_back({pos, top});
+				points.push_back({pos, lineBottom});
+				lines.push_back(2);
+			}
+
+			++number;
+			++index;
+		}
+
+		points.push_back({m_contentWidth, top - 25});
+		points.push_back({m_contentWidth, top});
+		lines.push_back(2);
+
+		context.PolyPolyline(points.data(), lines.data(), int(lines.size()));
+	}
+
+	void Visualizer::PaintDetailedStats(PaintContext& context, int& posY, const RECT& progressRect, TraceView::Session& session, bool isRemote, u64 playTime, const DrawTextFunc& drawTextFunc)
+	{
+		int stepY = context.GetActiveFont().height;
+		int startPosY = posY;
+		int posX = progressRect.left + 5;
+		int maxY = posY + stepY;
+		u32 maxTextWidth = 0;
+
+		RECT textRect;
+		auto drawText = [&](const wchar_t* format, ...)
+			{
+				textRect.top = posY;
+				textRect.bottom = posY + 20;
+				textRect.left = posX;
+				textRect.right = posX + 1000;
+				posY += stepY;
+				StringBuffer<> str;
+				va_list arg;
+				va_start(arg, format);
+				str.Append(format, arg);
+				va_end(arg);
+				u32 lastTextWidth = 0;
+				drawTextFunc(str, textRect, &lastTextWidth);
+				maxTextWidth = Max(maxTextWidth, lastTextWidth);
+				maxY = Max(maxY, posY);
+			};
+
+		if (isRemote)
+		{
+			drawText(L"Finished Processes: %u", session.processExitedCount);
+			drawText(L"Active Processes: %u", session.processActiveCount);
+			drawText(L"ClientId: %u  TcpCount: %u", session.clientUid.data1, session.connectionCount.empty() ? 1 : session.connectionCount.back());
+
+			if (session.disconnectTime == ~u64(0))
+			{
+				if (session.proxyCreated)
+					drawText(L"Proxy(HOSTED): %ls", session.proxyName.c_str());
+				else if (!session.proxyName.empty())
+					drawText(L"Proxy: %ls", session.proxyName.c_str());
+				else
+					drawText(L"Proxy: None");
+			}
+
+			bool hasFileDetails = !session.fetchedFiles.empty() || !session.storedFiles.empty();
+
+			if (!hasFileDetails)
+			{
+				posY = startPosY;
+				posX += maxTextWidth + 15;
+			}
+
+			if (!session.updates.empty())
+			{
+				auto updateTime = session.updates.back();
+				auto updateSend = session.networkSend.back();
+				auto updateRecv = session.networkRecv.back();
+				u64 sendPerS = 0;
+				u64 recvPerS = 0;
+				float duration = TimeToS(updateTime - session.prevUpdateTime);
+				if (duration != 0)
+				{
+					sendPerS = u64(float(updateSend - session.prevSend) / duration);
+					recvPerS = u64(float(updateRecv - session.prevRecv) / duration);
+				}
+				drawText(L"Recv: %ls (%sit/s)", BytesToText(updateRecv), BytesToText(recvPerS*8));
+				drawText(L"Send: %ls (%sit/s)", BytesToText(updateSend), BytesToText(sendPerS*8));
+			}
+
+			int fileWidth = 700;
+
+			auto drawFiles = [&](const wchar_t* fileType, Vector<TraceView::FileTransfer>& files, u64 count, u64 bytes, u64 activeCount, u32& maxVisibleFiles)
+				{
+					textRect.left = posX;
+					textRect.right = posX + fileWidth;
+					drawText(L"%ls Files: %u (%u) %s", fileType, u32(count), u32(activeCount), BytesToText(bytes));
+					u32 fileCount = 0;
+					for (auto rit = files.rbegin(), rend = files.rend(); rit != rend; ++rit)
+					{
+						TraceView::FileTransfer& file = *rit;
+						if (file.stop != ~u64(0))
+							continue;
+						u64 time = 0;
+						if (file.start < playTime)
+							time = playTime - file.start;
+
+						if ((((u8*)&file.key)[19] & IsProxyMask) != 0)
+							drawText(L"%s (proxy) %s", file.hint.c_str(), TimeToText(time, true).str);
+						else if (file.size == 0)
+							drawText(L"%s (calc) %s", file.hint.c_str(), TimeToText(time, true).str);
+						else
+							drawText(L"%s (%s) %s", file.hint.c_str(), BytesToText(file.size).str, TimeToText(time, true).str);
+
+						if (fileCount++ > 5)
+							break;
+					}
+					posY += stepY * (maxVisibleFiles - fileCount);
+					maxVisibleFiles = Max(maxVisibleFiles, fileCount);
+				};
+
+			Vector<TraceView::FileTransfer> fetchedFiles;
+			for (auto& kv : session.fetchedFilesActive)
+				fetchedFiles.push_back(session.fetchedFiles[kv.second]);
+			std::sort(fetchedFiles.begin(), fetchedFiles.end(), [](const TraceView::FileTransfer& a, const TraceView::FileTransfer& b) { return a.start > b.start; });
+
+			if (hasFileDetails)
+			{
+				posY = startPosY;
+				posX += maxTextWidth + 15;
+			}
+			drawFiles(L"Fetched", fetchedFiles, session.fetchedFilesCount, session.fetchedFilesBytes, session.fetchedFilesActive.size(), session.maxVisibleFiles);
+			if (hasFileDetails)
+			{
+				posY = startPosY;
+				posX += fileWidth;
+			}
+			drawFiles(L"Stored", session.storedFiles, session.storedFilesCount, session.storedFilesBytes, session.storedFilesActive.size(), session.maxVisibleFiles);
+		}
+		else
+		{
+			if (m_traceView.totalProcessActiveCount || m_traceView.totalProcessExitedCount)
+			{
+				drawText(L"Finished Processes: %u (local: %u)", m_traceView.totalProcessExitedCount, session.processExitedCount);
+				drawText(L"Active Processes: %u (local: %u)", m_traceView.totalProcessActiveCount, session.processActiveCount);
+				drawText(L"Active Helpers: %u", Max(1u, m_traceView.activeSessionCount) - 1);
+			}
+
+			if (session.highestSendPerS > 0.0f || session.highestRecvPerS > 0.0f)
+			{
+				auto updateTime = session.updates.back();
+				auto updateSend = session.networkSend.back();
+				auto updateRecv = session.networkRecv.back();
+				if (updateSend || updateRecv)
+				{
+					u64 sendPerS = 0;
+					u64 recvPerS = 0;
+					float duration = TimeToS(updateTime - session.prevUpdateTime);
+					if (duration != 0)
+					{
+						sendPerS = u64(float(updateSend - session.prevSend) / duration);
+						recvPerS = u64(float(updateRecv - session.prevRecv) / duration);
+					}
+
+					drawText(L"Recv: %ls (%sit/s)", BytesToText(updateRecv), BytesToText(recvPerS*8));
+					drawText(L"Send: %ls (%sit/s)", BytesToText(updateSend), BytesToText(sendPerS*8));
+				}
+			}
+
+			if (!session.updates.empty())
+			{
+				posY = startPosY;
+				posX += maxTextWidth + 10;
+
+				for (auto& pair : session.drives)
+				{
+					auto& drive = pair.second;
+
+					u64 readPerS = 0;
+					u64 writePerS = 0;
+
+					u64 updateCount = session.updates.size();
+					float duration = TimeToS(session.updates.back() - session.prevUpdateTime);
+					if (duration != 0 && updateCount > 1)
+					{
+						readPerS = u64(float(drive.readBytes.back()) / duration);
+						writePerS = u64(float(drive.writeBytes.back()) / duration);
+					}
+
+					drawText(L"%c: Rd %ls (%s/s) Wr %ls (%s/s) ", pair.first, BytesToText(drive.totalReadBytes).str, BytesToText(readPerS).str, BytesToText(drive.totalWriteBytes).str, BytesToText(writePerS).str);
+				}
+			}
+		}
+
+		posY = maxY;
+	}
+
+	void Visualizer::PrintCacheWriteStats(Logger& logger, u32 processId)
+	{
+		auto findIt = m_traceView.cacheWrites.find(processId);
+		if (findIt == m_traceView.cacheWrites.end())
+			return;
+		TraceView::CacheWrite& write = findIt->second;
+		logger.Info(L"");
+		logger.Info(L"  -------- Cache write stats ----------");
+		if (write.end)
+		{
+			logger.Info(L"  Duration                    %9s", TimeToText(write.end - write.start).str);
+			logger.Info(L"  Success                     %9s", write.success ? L"true" : L"false");
+
+			BinaryReader reader(write.stats.data(), 0, write.stats.size());
+			CacheSendStats stats;
+			if (m_traceView.version < 45)
+			{
+				stats.sendFile.bytes = reader.Read7BitEncoded();
+				stats.sendFile.count = 1;
+			}
+			else
+				stats.Read(reader, m_traceView.version);
+			stats.Print(logger, m_traceView.frequency);
+		}
+		else
+		{
+			logger.Info(L"  In progress...");
+		}
+	}
+
+	void Visualizer::PaintEmptyScreen(PaintContext& context)
+	{
+	}
+
+	void Visualizer::FinishIncompatible()
+	{
+	}
+
+	////////////////////////////////////////////////////////////////////////////////////////////////////
+}
diff --git a/Engine/Source/Programs/UnrealBuildAccelerator/Visualizer/Private/UbaVisualizer.h b/Engine/Source/Programs/UnrealBuildAccelerator/Visualizer/Private/UbaVisualizer.h
--- a/Engine/Source/Programs/UnrealBuildAccelerator/Visualizer/Private/UbaVisualizer.h
+++ b/Engine/Source/Programs/UnrealBuildAccelerator/Visualizer/Private/UbaVisualizer.h
@@ -7,14 +7,76 @@
 
 namespace uba
 {
+	////////////////////////////////////////////////////////////////////////////////////////////////////
+
 	class Visualizer
 	{
 	public:
 		Visualizer(VisualizerConfig& config, Logger& logger);
 
-
 	protected:
 		u64 GetPlayTime();
+		int GetTimelineHeight();
+		int GetTimelineTop(const RECT& clientRect);
+		StringBuffer<128> GetWorldTime(u64 time);
+		StringBuffer<128> GetWorldTime(float seconds);
+		bool ActiveProcessesShouldFillHeight();
+
+		struct Font
+		{
+			HFONT handle = 0;
+			HFONT handleUnderlined = 0;
+			int height = 0;
+			int offset = 0;
+		};
+
+		struct PaintContext
+		{
+			virtual RECT GetClientRect() = 0;
+			virtual void SetBkMode(int mode) = 0;
+			virtual void SetTextColor(COLORREF color) = 0;
+			virtual COLORREF GetTextColor() = 0;
+			virtual void SetPen(HPEN pen) = 0;
+			virtual void SetFont(HFONT font) = 0;
+			virtual u32 SetTextAlign(u32 align) = 0;
+			virtual void ExtTextOutW(int x, int y, UINT options, CONST RECT * lprect, StringView str) = 0;
+			virtual void SetActiveFont(const Font& font) = 0;
+			virtual Font& GetActiveFont() = 0;
+			virtual void SetBrush(HBRUSH b) = 0;
+			virtual void FillRect(const RECT& r, HBRUSH b) = 0;
+			virtual void Rectangle(int left, int top, int right, int bottom) = 0;
+			virtual void Line(int fromX, int fromY, int toX, int toY) = 0;
+			virtual void Polyline(const Vector<POINT>& polyLine) = 0;
+			virtual void PolyPolyline(const POINT* points, const DWORD* lines, u32 lineCount) = 0;
+			virtual void DrawCenteredText(const StringBuffer<>* rows, u32 rowCount, bool drawBackground) = 0;
+			virtual SIZE GetTextSize(StringView str) = 0;
+			virtual void DrawCached(TraceView::BitmapCache& process, RECT rect, int offsetY, int textHeight, RECT clippingRect, const Function<void(PaintContext& textCacheContext)>& func) = 0;
+
+			virtual POINT GetCursorClientPos() = 0;
+		};
+
+		struct HitTestResult;
+
+		void HitTest(HitTestResult& outResult, PaintContext& context, const POINT& pos);
+		void PaintAll(PaintContext& context, const RECT& clientRect);
+		void PaintProcessRect(PaintContext& context, TraceView::Process& process, RECT rect, const RECT& progressRect, bool selected, bool writingBitmap);
+		void PaintActiveProcesses(PaintContext& context, int& posY, const RECT& clientRect, const Function<void(TraceView::ProcessLocation&, u32, bool)>& drawProcess);
+		void PaintTimeline(PaintContext& context, const RECT& clientRect);
+		using DrawTextFunc = Function<void(const StringBufferBase& text, RECT& rect, u32* outWidth)>;
+		void PaintDetailedStats(PaintContext& context, int& posY, const RECT& progressRect, TraceView::Session& session, bool isRemote, u64 playTime, const DrawTextFunc& drawTextFunc);
+		void PrintCacheWriteStats(Logger& logger, u32 processId);
+		virtual void PaintEmptyScreen(PaintContext& context);
+		virtual void FinishIncompatible();
+
+		void UpdateFont(Font& font, int height, bool createUnderline);
+
+		enum VisualizerFlag
+		{
+			#define UBA_VISUALIZER_FLAG(name, defaultValue, desc) VisualizerFlag_##name,
+			UBA_VISUALIZER_FLAGS1
+			#undef UBA_VISUALIZER_FLAG
+			VisualizerFlag_Count
+		};
 
 
 		struct Stats
@@ -37,6 +99,7 @@ namespace uba
 			};
 			Map<char, Drive> drives;
 		};
+
 		struct HitTestResult
 		{
 			u32 section = ~0u;
@@ -56,16 +119,103 @@ namespace uba
 			TString hyperLink;
 		};
 
+		struct ProcessBrushes
+		{
+			HBRUSH inProgress = 0;
+			HBRUSH inProgressRace = 0;
+			HBRUSH success = 0;
+			HBRUSH error = 0;
+			HBRUSH returned = 0;
+			HBRUSH recv = 0;
+			HBRUSH send = 0;
+			HBRUSH cacheFetchTest = 0;
+			HBRUSH cacheFetchDownload = 0;
+			HBRUSH warning = 0;
+		};
 
 		Logger& m_logger;
 		VisualizerConfig m_config;
 		TraceReader m_trace;
 		TraceView m_traceView;
 
+		ProcessBrushes m_processBrushes[2]; // Non-selected and selected
+
+		static constexpr int GraphHeight = 30;
+
+		COLORREF m_textColor = {};
+		COLORREF m_textWarningColor = {};
+		COLORREF m_textErrorColor = {};
+		COLORREF m_sendColor = {};
+		COLORREF m_recvColor = {};
+		COLORREF m_cpuColor = {};
+		COLORREF m_memColor = {};
+		COLORREF m_driveColor = {};
+		COLORREF m_activeProcColor = {};
+		HBRUSH m_backgroundBrush = 0;
+		HBRUSH m_tooltipBackgroundBrush = 0;
+		HPEN m_textPen = 0;
+		HPEN m_separatorPen = 0;
+		HPEN m_sendPen = 0;
+		HPEN m_recvPen = 0;
+		HPEN m_cpuPen = 0;
+		HPEN m_memPen = 0;
+		HPEN m_drivePen = 0;
+		HPEN m_activeProcPen = 0;
+		HPEN m_processUpdatePen = 0;
+		HPEN m_checkboxPen = 0;
+		int m_sessionStepY = 0;
+
+		int m_processFontOffsetY = 0;
+
+		Font m_defaultFont;
+		Font m_processFont;
+		Font m_timelineFont;
+		Font m_popupFont;
+
+		Font m_activeProcessFont[32];
+		u32 m_activeProcessCountHistory[5];
+		u32 m_activeProcessCountHistoryIterator = 0;
+
 		bool m_paused = false;
 		u32 m_replay = 0;
 		u64 m_startTime = 0;
 		u64 m_pauseTime = 0;
 		u64 m_pauseStart = 0;
+
+		bool m_mouseOverWindow = false;
+
+		int m_progressRectLeft = 30;
+		int m_contentWidth = 0;
+		int m_contentHeight = 0;
+
+		float m_scrollPosX = 0;
+		float m_scrollPosY = 0;
+		float m_zoomValue = 0.5f;
+		float m_horizontalScaleValue = 0.5f;
+
+		u32 m_activeSection = ~0u;
+		TraceView::ProcessLocation m_processSelectedLocation;
+		bool m_processSelected = false;
+		u32 m_sessionSelectedIndex = ~0u;
+		bool m_statsSelected = false;
+		bool m_activeProcessGraphSelected = false;
+		u64 m_activeProcessCount = 0u;
+		Stats m_stats;
+		u32 m_buttonSelected = ~0u;
+		float m_timelineSelected = 0;
+		u32 m_fetchedFilesSelected = ~0u;
+		TString m_hyperLinkSelected;
+
+		TString m_filterString;
+
+		UnorderedMap<Color, HBRUSH> m_coloredBrushes;
+
+		bool m_workSelected = false;
+		u32 m_workTrack = ~0u;
+		u32 m_workIndex = ~0u;
+
+		struct DrawTextLogger;
 	};
-}
+
+	////////////////////////////////////////////////////////////////////////////////////////////////////
+}
\ No newline at end of file
diff --git a/Engine/Source/Programs/UnrealBuildAccelerator/Visualizer/Private/Windows/UbaVisualizerWin.cpp b/Engine/Source/Programs/UnrealBuildAccelerator/Visualizer/Private/Windows/UbaVisualizerWin.cpp
--- a/Engine/Source/Programs/UnrealBuildAccelerator/Visualizer/Private/Windows/UbaVisualizerWin.cpp
+++ b/Engine/Source/Programs/UnrealBuildAccelerator/Visualizer/Private/Windows/UbaVisualizerWin.cpp
@@ -45,32 +45,29 @@ namespace uba
 		Popup_Quit,
 	};
 
-	enum VisualizerFlag
+	struct VisualizerWin::PaintContextWin : public Visualizer::PaintContext 
 	{
-#define UBA_VISUALIZER_FLAG(name, defaultValue, desc) VisualizerFlag_##name,
-		UBA_VISUALIZER_FLAGS1
-#undef UBA_VISUALIZER_FLAG
-		VisualizerFlag_Count
-	};
-
-	struct VisualizerWin::PaintContext
-	{
-		void SetBkMode(int mode) { ::SetBkMode(hdc, mode); }
-		void SetTextColor(COLORREF color) { ::SetTextColor(hdc, color); }
-		void SetPen(HPEN pen) { ::SelectObject(hdc, pen); }
-		void SetFont(HFONT font) { ::SelectObject(hdc, font); }
-		u32 SetTextAlign(u32 align) { return ::SetTextAlign(hdc, align); }
-		void ExtTextOutW(int x, int y, UINT options, CONST RECT * lprect, StringView str) { ::ExtTextOutW(hdc, x, y, options, lprect, str.data, str.count, NULL); }
+		virtual RECT GetClientRect() override { return clientRect; }
+		virtual void SetBkMode(int mode) override { ::SetBkMode(hdc, mode); }
+		virtual void SetTextColor(COLORREF color) override { ::SetTextColor(hdc, color); }
+		virtual COLORREF GetTextColor() override { return ::GetTextColor(hdc); }
+		virtual void SetPen(HPEN pen) override { ::SelectObject(hdc, pen); }
+		virtual void SetFont(HFONT font) override { ::SelectObject(hdc, font); }
+		virtual u32 SetTextAlign(u32 align) override { return ::SetTextAlign(hdc, align); }
+		virtual void ExtTextOutW(int x, int y, UINT options, CONST RECT * lprect, StringView str) override { ::ExtTextOutW(hdc, x, y, options, lprect, str.data, str.count, NULL); }
 		
-		void SetActiveFont(const Font& font)
+		virtual void SetActiveFont(const Font& font) override
 		{
-			owner.m_activeFont = font;
+			activeFont = font;
 			::SelectObject(hdc, font.handle);
 		}
 
-		Font& GetActiveFont() { return owner.m_activeFont; }
+		virtual Font& GetActiveFont() override
+		{
+			return activeFont;
+		}
 
-		void SetBrush(HBRUSH b)
+		virtual void SetBrush(HBRUSH b) override
 		{
 			if (lastSelectedBrush != b)
 			{
@@ -79,35 +76,35 @@ namespace uba
 			}
 		}
 
-		void FillRect(const RECT& r, HBRUSH b)
+		virtual void FillRect(const RECT& r, HBRUSH b) override
 		{
 			SetBrush(b);
 			PatBlt(hdc, r.left, r.top, r.right - r.left, r.bottom - r.top, PATCOPY);
 		}
 
-		void Rectangle(int left, int top, int right, int bottom)
+		virtual void Rectangle(int left, int top, int right, int bottom) override
 		{
 			SetBrush((HBRUSH)GetStockObject(NULL_BRUSH));
 			::Rectangle(hdc, left, top, right, bottom);
 		}
 
-		void Line(int fromX, int fromY, int toX, int toY)
+		virtual void Line(int fromX, int fromY, int toX, int toY) override
 		{
 			::MoveToEx(hdc, fromX, fromY, NULL);
 			::LineTo(hdc, toX, toY);
 		}
 
-		void Polyline(const Vector<POINT>& polyLine)
+		virtual void Polyline(const Vector<POINT>& polyLine) override
 		{
 			::Polyline(hdc, polyLine.data(), int(polyLine.size()));
 		}
 
-		void PolyPolyline(const POINT* points, const DWORD* lines, u32 lineCount)
+		virtual void PolyPolyline(const POINT* points, const DWORD* lines, u32 lineCount) override
 		{
 			::PolyPolyline(hdc, points, lines, int(lineCount));
 		}
 
-		void DrawCenteredText(const StringBuffer<>* rows, u32 rowCount, bool drawBackground)
+		virtual void DrawCenteredText(const StringBuffer<>* rows, u32 rowCount, bool drawBackground) override
 		{
 			for (u32 i=0; i!=rowCount; ++i)
 			{ 
@@ -130,14 +127,13 @@ namespace uba
 			}
 		}
 
-		SIZE GetTextSize(StringView str)
+		virtual SIZE GetTextSize(StringView str) override
 		{
 			SIZE s;
 			return GetTextExtentPoint32W(hdc, str.data, str.count, &s) ? s : SIZE();
 		}
 
-
-		void DrawCached(TraceView::BitmapCache& process, int processWidth, int offsetY, int textHeight, int left, RECT rect, RECT progressRect, const Function<void(PaintContext& textCacheContext)>& func)
+		virtual void DrawCached(TraceView::BitmapCache& process, RECT rect, int offsetY, int textHeight, RECT clippingRect, const Function<void(PaintContext& textCacheContext)>& func) override
 		{
 			if (!process.bitmap || process.bitmapDirty)
 			{
@@ -161,7 +157,7 @@ namespace uba
 					lastSelectedBitmap = process.bitmap;
 				}
 
-				PaintContext textCacheContext{owner, NULL, textDC, clientRect };
+				PaintContextWin textCacheContext{owner, NULL, textDC, clientRect };
 				func(textCacheContext);
 			}
 
@@ -171,6 +167,7 @@ namespace uba
 				lastSelectedBitmap = process.bitmap;
 			}
 
+			int processWidth = rect.right - rect.left;
 			int width = Min(processWidth, 256);
 			int bitmapOffsetY = process.bitmapOffset;
 			int bltOffsetY = offsetY;
@@ -183,22 +180,31 @@ namespace uba
 			if (bltOffsetY + height > textHeight)
 				height = textHeight - bltOffsetY;
 
-			if (left > -256 && height >= 0)
-			{
-				int bitmapOffsetX = rect.left - left;
+			int left = rect.left;
+			if (left <= -256 || height <= 0)
+				return;
 
-				if (left < progressRect.left)
-				{
-					int diff = progressRect.left - left;
-					rect.left = progressRect.left;
-					width -= diff;
-					bitmapOffsetX += diff;
-				}
-				BitBlt(hdc, rect.left, rect.top + bltOffsetY, width, height, textDC, bitmapOffsetX, bitmapOffsetY, SRCCOPY);
+			int bitmapOffsetX = rect.left - left;
+
+			if (left < clippingRect.left)
+			{
+				int diff = clippingRect.left - left;
+				rect.left = clippingRect.left;
+				width -= diff;
+				bitmapOffsetX += diff;
 			}
+			BitBlt(hdc, rect.left, rect.top + bltOffsetY, width, height, textDC, bitmapOffsetX, bitmapOffsetY, SRCCOPY);
 		}
 
+		virtual POINT GetCursorClientPos() override
+		{
+			POINT p;
+			GetCursorPos(&p);
+			ScreenToClient(hwnd, &p);
+			return p;
+		}
 
+		PaintContextWin(VisualizerWin& o, HWND hw, HDC dc, const RECT& cr, HDC tdc = 0) : owner(o), hwnd(hw), hdc(dc), clientRect(cr), textDC(tdc) {}
 
 
 		VisualizerWin& owner;
@@ -206,108 +212,13 @@ namespace uba
 		HDC hdc;
 		const RECT& clientRect;
 
+		Font activeFont;
+
 		HDC textDC = NULL;
 		HBRUSH lastSelectedBrush = 0;
 		HBITMAP lastSelectedBitmap = 0;
 	};
 
-	struct VisualizerWin::DrawTextLogger : public Logger
-	{
-		DrawTextLogger(PaintContext& c, int fh, HBRUSH bb)
-			:	context(c)
-			,	fontHeight(fh)
-			,	backgroundBrush(bb)
-		{
-			textColor = GetTextColor(c.hdc);
-		}
-
-		virtual void BeginScope() override {}
-		virtual void EndScope() override {}
-		virtual void Log(LogEntryType type, const wchar_t* str, u32 strLen) override
-		{
-			u32 textWidth = context.GetTextSize(StringView(str, strLen)).cx;
-
-			lines.emplace_back(TString(str, strLen), textOffset, height, textColor);
-			width = Max(width, int(textWidth + textOffset));
-			height += fontHeight;
-		}
-
-		void AddSpace(int space = 5)
-		{
-			height += space;
-		}
-
-		void AddTextOffset(int offset)
-		{
-			textOffset += offset;
-		}
-
-		void AddWidth(int extra)
-		{
-			extraWidth += extra;
-		}
-
-		void DrawAtPos(int x, int y)
-		{
-			RECT r;
-			r.left = x;
-			r.top = y;
-			r.right = r.left + width;
-			r.bottom = r.top + height;
-
-			RECT clientRect;
-			GetClientRect(context.hwnd, &clientRect);
-
-			if (r.right > clientRect.right)
-				OffsetRect(&r, -width - 15, 0);
-			if (r.bottom > clientRect.bottom)
-			{
-				OffsetRect(&r, 0, clientRect.bottom - r.bottom);
-				if (r.top < 0)
-					OffsetRect(&r, 0, -r.top);
-			}
-
-			RECT fillRect = r;
-			fillRect.right += 2 + extraWidth;
-			context.FillRect(fillRect, backgroundBrush);
-
-			for (auto& line : lines)
-			{
-				RECT tr = r;
-				tr.left += line.left;
-				tr.top += line.top;
-				context.SetTextColor(line.color);
-				context.ExtTextOutW(tr.left, tr.top, 0, NULL, line.str);
-			}
-		}
-
-		void DrawAtCursor()
-		{
-			POINT p;
-			GetCursorPos(&p);
-			ScreenToClient(context.hwnd, &p);
-			p.x += 3;
-			p.y += 3;
-			DrawAtPos(p.x, p.y);
-		}
-
-		DrawTextLogger& SetColor(COLORREF c) { textColor = c; return *this; }
-
-		int width = 0;
-		int height = 0;
-		int textOffset = 2;
-		int extraWidth = 0;
-		struct Line { TString str; int left; int top; COLORREF color; };
-		Vector<Line> lines;
-
-		PaintContext& context;
-		int fontHeight;
-		HBRUSH backgroundBrush;
-		COLORREF textColor;
-		bool isFirst = true;
-		using Logger::Log;
-	};
-
 	class WriteTextLogger : public Logger
 	{
 	public:
@@ -758,8 +669,13 @@ namespace uba
 		if (m_config.DarkMode)
 			UpdateTheme();
 
-		HitTestResult res;
-		HitTest(res, { -1, -1 });
+		{
+			RECT clientRect;
+			GetClientRect(m_hwnd, &clientRect);
+			PaintContextWin context(*this, NULL, NULL, clientRect);
+			HitTestResult res;
+			HitTest(res, context, { -1, -1 });
+		}
 
 		if (m_config.parent)
 		{
@@ -914,37 +830,7 @@ namespace uba
 		m_lastBitmap = 0;
 	}
 
-	void VisualizerWin::PrintCacheWriteStats(Logger& logger, u32 processId)
-	{
-		auto findIt = m_traceView.cacheWrites.find(processId);
-		if (findIt == m_traceView.cacheWrites.end())
-			return;
-		TraceView::CacheWrite& write = findIt->second;
-		logger.Info(L"");
-		logger.Info(L"  -------- Cache write stats ----------");
-		if (write.end)
-		{
-			logger.Info(L"  Duration                    %9s", TimeToText(write.end - write.start).str);
-			logger.Info(L"  Success                     %9s", write.success ? L"true" : L"false");
-
-			BinaryReader reader(write.stats.data(), 0, write.stats.size());
-			CacheSendStats stats;
-			if (m_traceView.version < 45)
-			{
-				stats.sendFile.bytes = reader.Read7BitEncoded();
-				stats.sendFile.count = 1;
-			}
-			else
-				stats.Read(reader, m_traceView.version);
-			stats.Print(logger, m_traceView.frequency);
-		}
-		else
-		{
-			logger.Info(L"  In progress...");
-		}
-	}
-
-	void VisualizerWin::UpdateFont(Font& font, int height, bool createUnderline)
+	void Visualizer::UpdateFont(Font& font, int height, bool createUnderline)
 	{
 		font.height = height;
 
@@ -1049,2373 +935,173 @@ namespace uba
 		ReleaseDC(m_hwnd, hdc);
 	}
 
-	constexpr int GraphHeight = 30;
-
-	struct SessionRec
-	{
-		TraceView::Session* session;
-		u32 index;
-	};
-	void Populate(SessionRec* recs, TraceView& traceView, bool sort)
+	void VisualizerWin::PaintEmptyScreen(PaintContext& context)
 	{
-		u32 count = u32(traceView.sessions.size());
-		for (u32 i = 0, e = count; i != e; ++i)
-			recs[i] = { &traceView.sessions[i], i };
-		if (count <= 1 || !sort)
+		if (m_parentHwnd)
 			return;
-		std::sort(recs + 1, recs + traceView.sessions.size(), [](SessionRec& a, SessionRec& b)
-			{
-				auto& as = *a.session;
-				auto& bs = *b.session;
-				if ((as.processActiveCount != 0) != (bs.processActiveCount != 0))
-					return as.processActiveCount > bs.processActiveCount;
-				if (as.processActiveCount && as.proxyCreated != bs.proxyCreated)
-					return int(as.proxyCreated) > int(bs.proxyCreated);
-				return a.index < b.index;
-			});
-	}
-
-	void VisualizerWin::PaintAll(PaintContext& context, const RECT& clientRect)
-	{
-		u64 playTime = GetPlayTime();
 
-		context.SetBkMode(TRANSPARENT);
-		context.SetTextColor(m_textColor);
-
-		if (m_traceView.sessions.empty() && !m_parentHwnd)
+		StringBuffer<> str[2];
+		if (!m_fileName.IsEmpty())
+		{
+			str[0].Append(TCV("Loading trace file"));
+		}
+		else
 		{
-			StringBuffer<> str[2];
-			if (!m_fileName.IsEmpty())
+			str[1].Append(TCV("Click here to open trace file"));
+			if (m_listenChannel.count)
 			{
-				str[0].Append(TCV("Loading trace file"));
+				str[0].Appendf(TC("Listening for new sessions on channel '%s'"), m_listenChannel.data);
+				str[1].Append(TCV(" instead"));
 			}
 			else
-			{
-				str[1].Append(TCV("Click here to open trace file"));
-				if (m_listenChannel.count)
-				{
-					str[0].Appendf(TC("Listening for new sessions on channel '%s'"), m_listenChannel.data);
-					str[1].Append(TCV(" instead"));
-				}
-				else
-					str[0].Append(TCV("No trace active"));
-			}
-			context.SetActiveFont(m_popupFont);
-			context.DrawCenteredText(str, 2, false);
-			return;
-		}
-
-
-		int posY = int(m_scrollPosY);
-		float scaleX = 50.0f*m_zoomValue*m_config.horizontalScaleValue;
-
-		RECT progressRect = clientRect;
-		progressRect.left += m_progressRectLeft;
-
-		if (m_config.showTimeline)
-		{
-			progressRect.bottom -= m_defaultFont.height + 10;
+				str[0].Append(TCV("No trace active"));
 		}
+		context.SetActiveFont(m_popupFont);
+		context.DrawCenteredText(str, 2, false);
+	}
 
-		u64 lastStop = 0;
-
-		context.SetActiveFont(m_defaultFont);
-
-		auto drawStatusText = [&](StringView text, LogEntryType type, int posX, int endX, bool moveY, bool underlined = false, bool centered = false)
-			{
-				Font& activeFont = context.GetActiveFont();
-				RECT rect;
-				rect.left = posX;
-				rect.right = endX;
-				rect.top = posY + activeFont.offset;
-				rect.bottom = posY + activeFont.height + 2;
-				context.SetTextColor(type == LogEntryType_Info ? m_textColor : (type == LogEntryType_Error ? m_textErrorColor : m_textWarningColor));
-
-				u32 align = 0;
-				if (centered)
-				{
-					posX += (endX - posX)/2;
-					align = context.SetTextAlign(TA_CENTER);
-				}
-				if (underlined)
-					context.SetFont(activeFont.handleUnderlined);
-				context.ExtTextOutW(posX, posY, ETO_CLIPPED, &rect, text);
-				if (underlined)
-					context.SetFont(activeFont.handle);
-				if (centered)
-					context.SetTextAlign(align);
-				if (moveY)
-					posY = rect.bottom;
-			};
+	void VisualizerWin::FinishIncompatible()
+	{
+		StringBuffer<> title;
+		SetWindowTextW(m_hwnd, GetTitlePrefix(title).Appendf(L" (Listening for new sessions on channel '%s')", m_listenChannel.data).data);
+	}
 
-		auto drawIndentedText = [&](const StringView& text, LogEntryType type, int indent, bool moveY, bool underlined = false)
-			{
-				int posX = 5 + indent*m_defaultFont.height;
-				drawStatusText(text, type, posX, clientRect.right, moveY, underlined);
-			};
+	void VisualizerWin::WriteProcessStats(Logger& out, const TraceView::Process& process)
+	{
+		bool hasExited = process.stop != ~u64(0);
+		out.Info(L"  %ls", process.description.c_str());
+		out.Info(L"  ProcessId: %u", process.id);
+		out.Info(L"  Start:     %ls", TimeToText(process.start, true).str);
+		if (hasExited)
+			out.Info(L"  Duration:  %ls", TimeToText(process.stop - process.start, true).str);
+		if (hasExited && process.exitCode != 0)
+			out.Info(L"  ExitCode:  %u", process.exitCode);
 
-		if (m_config.showProgress && m_traceView.progressProcessesTotal)
+		if (process.stop != ~u64(0))
 		{
-			drawIndentedText(AsView(L"Progress"), LogEntryType_Info, 1, false);
-
-			Font& activeFont = context.GetActiveFont();
-
-			float progress = float(m_traceView.progressProcessesDone) / float(m_traceView.progressProcessesTotal);
-			u32 start = 3 + 6*activeFont.height;
-			u32 width = activeFont.height * 18;
-			RECT rect;
-			rect.left = start;
-			rect.right = start + width;
-			rect.top = posY;
-			rect.bottom = posY + activeFont.height;
-			context.FillRect(rect, m_processBrushes[0].inProgress);
-
-			rect.right = rect.left + int(progress*float(width));//durationMs/100;
-			context.FillRect(rect, m_traceView.progressErrorCount ? m_processBrushes[0].error : m_processBrushes[0].success);
-
-			StringBuffer<> str;
-			str.Appendf(TC("%u%%    %u / %u"), u32(progress*100.0f), m_traceView.progressProcessesDone, m_traceView.progressProcessesTotal);
-
-			const tchar* remoteDisabled = TC(""); // TODO: Only show this in "advanced mode" m_traceView.remoteExecutionDisabled ? TC("   (Remote spawn disabled)") : TC("");
-			str.Append(remoteDisabled);
+			out.Info(L"");
 
-			if (u32 active = m_traceView.totalProcessActiveCount)
-				str.Appendf(TC("   (%u active)"), active);
+			BinaryReader reader(process.stats.data(), 0, process.stats.size());
+			ProcessStats processStats;
+			SessionStats sessionStats;
+			StorageStats storageStats;
+			KernelStats kernelStats;
 
-			if (m_traceView.lastKillProcessTime && playTime >= m_traceView.lastKillProcessTime && playTime < m_traceView.lastKillProcessTime + MsToTime(800))
-			{
-				rect.left = start + width + 20;
-				rect.right = rect.left + activeFont.height * 7;
-				context.FillRect(rect, m_processBrushes[0].error);
-				drawStatusText(TCV("Killing (mem)"), LogEntryType_Info, rect.left, rect.right, 0, false, true);
-			}
-			else if (m_traceView.lastSpawningDelayStartTime && playTime >= m_traceView.lastSpawningDelayStartTime && playTime < m_traceView.lastSpawningDelayEndTime + MsToTime(500))
+			processStats.Read(reader, m_traceView.version);
+			if (reader.GetLeft())
 			{
-				rect.left = start + width + 20;
-				rect.right = rect.left + activeFont.height * 7;
-				context.FillRect(rect, m_processBrushes[0].warning);
-				drawStatusText(TCV("Throttling (mem)"), LogEntryType_Info, rect.left, rect.right, 0, false, true);
+				if (process.isRemote || (m_traceView.version >= 36 && !process.isReuse))
+					sessionStats.Read(reader, m_traceView.version);
+				storageStats.Read(reader, m_traceView.version);
+				kernelStats.Read(reader, m_traceView.version);
 			}
 
+			out.Info(L"  ----------- Detours stats -----------");
+			processStats.Print(out, m_traceView.frequency);
 
-			drawIndentedText(str, LogEntryType_Info, 6, true);
-		}
-
-		if (m_traceView.version && (m_traceView.version < TraceReadCompatibilityVersion || m_traceView.version > TraceVersion))
-		{
-			if (!m_traceView.finished)
+			if (!sessionStats.IsEmpty())
 			{
-				StringBuffer<> str;
-				str.Appendf(TC("Unsupported trace version %u (Versions supported are %u to %u)"), m_traceView.version, TraceReadCompatibilityVersion, TraceVersion);
-				context.SetTextColor(m_textWarningColor);
-				context.DrawCenteredText(&str, 1, false);
+				out.Info(L"");
+				out.Info(L"  ----------- Session stats -----------");
+				sessionStats.Print(out, m_traceView.frequency);
 			}
-			else
+
+			if (!storageStats.IsEmpty())
 			{
-				m_traceView.Clear();
-				StringBuffer<> title;
-				SetWindowTextW(m_hwnd, GetTitlePrefix(title).Appendf(L" (Listening for new sessions on channel '%s')", m_listenChannel.data).data);
+				out.Info(L"");
+				out.Info(L"  ----------- Storage stats -----------");
+				storageStats.Print(out, m_traceView.frequency);
 			}
-			return;
-		}
 
-		if (m_config.showStatus && !m_traceView.statusMap.empty())
-		{
-			u32 lastRow = ~0u;
-			u32 row = ~0u;
-			Font& activeFont = context.GetActiveFont();
-			for (auto& kv : m_traceView.statusMap)
+			if (!kernelStats.IsEmpty())
 			{
-				auto& status = kv.second;
-				if (status.text.empty())
-					continue;
-				row = u32(kv.first >> 32);
-				if (lastRow != ~0u && lastRow != row)
-					posY += activeFont.height + 2;
-				lastRow = row;
-				u32 column = u32(kv.first & ~0u);
-				drawIndentedText(status.text, status.type, column, false, !status.link.empty());
+				out.Info(L"");
+				out.Info(L"  ----------- Kernel stats ------------");
+				kernelStats.Print(out, false, m_traceView.frequency);
 			}
-			if (row != ~0u)
-				posY += activeFont.height + 2;
 
-			context.SetTextColor(m_textColor);
-			posY += 3;
+			PrintCacheWriteStats(out, process.id);
 		}
+	}
+
+	void VisualizerWin::WriteWorkStats(Logger& out, const TraceView::WorkRecord& record)
+	{
+		out.Info(L"  %ls", record.description);
+		out.Info(L"  Start:     %ls", TimeToText(record.start, true).str);
+		if (record.stop != ~u64(0))
+			out.Info(L"  Duration:  %ls", TimeToText(record.stop - record.start, true).str);
+		for (auto& e : record.entries)
+			out.Info(L"   %ls (%ls)", e.text, TimeToText(e.time - record.start).str);
+	}
 
-		if (m_config.showActiveProcessGraph)
+	void VisualizerWin::CopyTextToClipboard(const TString& str)
+	{
+		if (!OpenClipboard(m_hwnd))
+			return;
+		if (auto hglbCopy = GlobalAlloc(GMEM_MOVEABLE, (str.size() + 1) * sizeof(TCHAR)))
 		{
-			if (posY + GraphHeight >= progressRect.top && posY + GraphHeight - 5 < progressRect.bottom)
+			if (auto lptstrCopy = GlobalLock(hglbCopy))
 			{
-				int graphBaseY = posY + GraphHeight - 4;
-				double graphHeight = double(GraphHeight - 2);
-
-				Vector<POINT> line;
-
-				context.SetPen(m_activeProcPen);
-				bool isFirstUpdate = true;
-				bool isFirstDraw = true;
-				u64 prevValue = 0;
-				int prevX = 0;
-				int prevY = 0;
-
-				float timeScale = (m_config.horizontalScaleValue*m_zoomValue)*50.0f;
-				float startOffset = ((m_scrollPosX/timeScale) - float(int(m_scrollPosX/timeScale))) * timeScale;
-				float time = -startOffset/timeScale;
-				float playTimeS = TimeToS(playTime);
-				u32 activeProcessCountIndex = 0;
-				u16 count = 0;
-				while (time < playTimeS && activeProcessCountIndex < m_traceView.activeProcessCounts.size())
-				{
-					int posX = progressRect.left + int(startOffset + time*timeScale);
-					if (posX >= clientRect.right)
-						break;
-
-					while (activeProcessCountIndex < m_traceView.activeProcessCounts.size())
-					{
-						auto& entry = m_traceView.activeProcessCounts[activeProcessCountIndex];
-						count = entry.count;
-						if (float(TimeToMs(entry.time))/1000.0f > time)
-							break;
-						++activeProcessCountIndex;
-					}
-
-					int x = posX;
-					int y = graphBaseY;
-
-					double div = double(m_traceView.maxActiveProcessCount);
-					y -= int(double(count) * graphHeight / div);
-
-					if (prevValue <= 0)
-						prevY = y;
-
-					if (x > clientRect.left)
-					{
-						if (isFirstUpdate)
-						{
-							line.push_back({ x, y });
-							isFirstDraw = false;
-						}
-						else
-						{
-							if (isFirstDraw)
-								line.push_back({ prevX, prevY });
-							line.push_back({ x, y });
-							isFirstDraw = false;
-						}
-					}
-					if (x > clientRect.right)
-						break;
-					isFirstUpdate = false;
-					prevX = x;
-					prevY = y;
-					prevValue = count;
-
-					time += 0.25f;
-				}
-				if (line.size() > 1)
-					context.Polyline(line);
+				memcpy(lptstrCopy, str.data(), (str.size() + 1) * sizeof(TCHAR));
+				GlobalUnlock(hglbCopy);
+				EmptyClipboard();
+				SetClipboardData(CF_UNICODETEXT, hglbCopy);
 			}
-			posY += GraphHeight;
 		}
+		CloseClipboard();
+	}
 
-		if (m_config.showActiveProcesses && !m_trace.m_activeProcesses.empty())
-		{
-			auto drawBox = [&](u64 start, u64 stop, int height, bool selected, bool inProgress)
-				{
-					//int left = 3 + 6*m_activeFont.height;
-					//int right = rect.left + durationMs/100;
-					int posX = int(m_scrollPosX) + progressRect.left;
-					bool done = stop != ~u64(0);
-					if (!done)
-						stop = playTime;
-					int left = posX + int(float(TimeToS(start)) * scaleX);
-					int right = posX + int(float(TimeToS(stop)) * scaleX) - 1;
-
-					RECT rect;
-					rect.left = left;
-					rect.right = right;
-					rect.top = posY;
-					rect.bottom = posY + height;
-					context.FillRect(rect, inProgress ? m_processBrushes[selected].inProgress : m_processBrushes[selected].success);
-					return rect;
-				};
-
-			m_activeProcessCountHistory[m_activeProcessCountHistoryIterator++ % sizeof_array(m_activeProcessCountHistory)] = u32(m_trace.m_activeProcesses.size());
-
-			PaintActiveProcesses(context, posY, clientRect, [&](TraceView::ProcessLocation& processLocation, u32 boxHeight, bool firstWithHeight)
-				{
-					auto& session = m_trace.GetSession(m_traceView, processLocation.sessionIndex);
-					TraceView::Process& process = session.processors[processLocation.processorIndex].processes[processLocation.processIndex];
+	void VisualizerWin::UnselectAndRedraw()
+	{
+		if (Unselect() || m_config.showCursorLine)
+			Redraw(false);
+	}
 
-					bool selected = m_processSelected && m_processSelectedLocation == processLocation;
+	bool VisualizerWin::UpdateAutoscroll()
+	{
+		if (!m_autoScroll)
+			return false;
 
-					if (m_config.showFinishedProcesses)
-					{
-						u32 index = processLocation.processIndex;
-						while (index > 0)
-						{
-							--index;
-							TraceView::Process& process2 = session.processors[processLocation.processorIndex].processes[index];
-							drawBox(process2.start, process2.stop, boxHeight, false, false);
-						}
-					}
+		u64 playTime = GetPlayTime();
 
-					RECT boxRect = drawBox(process.start, process.stop, boxHeight, selected, true);
 
-					u32 v = boxHeight - 1;
+		RECT rect;
+		GetClientRect(m_hwnd, &rect);
+		if (rect.right == 0)
+			return false;
 
-					if (v > 4)
-					{
-						u32 fontIndex = Min(v, u32(sizeof_array(m_activeProcessFont) - 1));
-						if (!m_activeProcessFont[fontIndex].handle)
-						{
-							UpdateFont(m_activeProcessFont[fontIndex], fontIndex - 1, false);
-							m_activeProcessFont[fontIndex].offset += 1;
-						}
-						if (firstWithHeight)
-							context.SetActiveFont(m_activeProcessFont[fontIndex]);
+		float timeS = TimeToS(playTime);
 
-						StringBuffer<> str;
-						auto firstPar = -1;//process.description.find_first_of('(');
-						if (firstPar != -1 && *process.description.rbegin() == ')')
-						{
-							str.Append(process.description.c_str() + firstPar + 1).Resize(str.count-1);
-							str.Append(L"  ").Append(process.description.c_str(), firstPar);
-						}
-						else
-							str.Append(process.description);
-
-						if (process.isRemote)
-							str.Append(L" [").Append(session.name).Append(']');
-						else if (process.type == TraceView::ProcessType_CacheFetchTest || process.type == TraceView::ProcessType_CacheFetchDownload)
-							str.Append(L" [cache]");
-						if (boxRect.left < 0)
-							str.Appendf(L"   %s", TimeToText(playTime - process.start, true).str);
-						//drawIndentedText(str, LogEntryType_Info, 1, true);
-						drawStatusText(str, LogEntryType_Info, Max(int(boxRect.left+1), 1), boxRect.right, false);
-					}
-				});
+		if (m_config.AutoScaleHorizontal)
+		{
+			m_scrollPosX = 0;
+			timeS = Max(timeS, 20.0f/m_zoomValue);
+			m_config.horizontalScaleValue = Max(float(rect.right - m_progressRectLeft - 2)/(m_zoomValue*timeS*50.0f), 0.001f);
+			return true;
 		}
+		else
+		{
+			float oldScrollPosX = m_scrollPosX;
+			m_scrollPosX = Min(0.0f, (float)rect.right - timeS*50.0f*m_config.horizontalScaleValue*m_zoomValue - float(m_progressRectLeft));
+			return oldScrollPosX != m_scrollPosX;
+		}
+	}
 
-		int boxHeight = m_config.boxHeight;
-		int stepY = int(boxHeight) + 2;
-		int processStepY = boxHeight + 1;
-
-		TraceView::WorkRecord selectedWork;
-
-		SessionRec sortedSessions[1024];
-		Populate(sortedSessions, m_traceView, m_config.SortActiveRemoteSessions);
-
-		u32 visibleBoxes = 0;
+	bool VisualizerWin::UpdateSelection()
+	{
+		if (!m_mouseOverWindow || m_dragToScrollCounter > 0)
+			return false;
+		POINT pos;
+		GetCursorPos(&pos);
+		ScreenToClient(m_hwnd, &pos);
 
-		TraceView::ProcessLocation processLocation { 0, 0, 0 };
-		for (u64 sessionIt = 0, sessionEnd = m_traceView.sessions.size(); sessionIt != sessionEnd; ++sessionIt)
+		HitTestResult res;
 		{
-			bool isFirst = sessionIt == 0;
-			auto& session = *sortedSessions[sessionIt].session;
-			bool hasUpdates = !session.updates.empty();
-			if (!isFirst)
-			{
-				if (!hasUpdates && session.processors.empty())
-					continue;
-
-				if (!m_config.showFinishedProcesses && session.disconnectTime != ~u64(0))
-					continue;
-			}
-
-			processLocation.sessionIndex = sortedSessions[sessionIt].index;
-			if (!isFirst)
-				posY += 3;
-
-			if (m_config.showTitleBars)
-			{
-				if (posY + stepY >= progressRect.top && posY <= progressRect.bottom)
-				{
-					context.SetPen(m_separatorPen);
-					context.Line(0, posY, clientRect.right, posY);
-
-					StringBuffer<> text;
-					session.GetFullName(text);
-
-					if (hasUpdates && session.disconnectTime == ~u64(0))
-					{
-						u64 ping = session.ping.back();
-						//u64 memAvail = session.updates.back().memAvail;
-						//float cpuLoad = session.updates.back().cpuLoad;
-
-						//text.Appendf(L" - Cpu: %.1f%%", cpuLoad * 100.0f);
-						//if (memAvail)
-						//	text.Appendf(L" Mem: %ls/%ls", BytesToText(session.memTotal - memAvail).str, BytesToText(session.memTotal).str);
-						if (ping)
-							text.Appendf(L" Ping: %ls", TimeToText(ping, false, m_traceView.frequency).str);
-						if (!session.notification.empty())
-							text.Append(L" - ").Append(session.notification);
-					}
-					else if (!isFirst)
-					{
-						text.Append(L" - Disconnected");
-						if (!session.notification.empty())
-							text.Append(L" (").Append(session.notification).Append(')');
-					}
-
-					bool selected = m_sessionSelectedIndex == processLocation.sessionIndex;
-
-					int textBottom = Min(posY + m_sessionStepY, int(progressRect.bottom));
-
-					RECT rect;
-					rect.left = 5;
-					rect.right = clientRect.right;
-					rect.top = posY;
-					rect.bottom = textBottom;
-
-					if (selected)
-						context.SetBkMode(OPAQUE);
-
-					SIZE textSize = context.GetTextSize(text);
-					if (textSize.cx)
-						session.fullNameWidth = u32(textSize.cx);
-					context.ExtTextOutW(5, posY+2, ETO_CLIPPED, &rect, text);
-					if (selected)
-						context.SetBkMode(TRANSPARENT);
-				}
-				posY += m_sessionStepY;
-			}
-
-			bool showGraph = m_config.showNetworkStats || m_config.showCpuMemStats || m_config.showDriveStats;
-			if (showGraph && hasUpdates)
-			{
-				if (posY + GraphHeight >= progressRect.top && posY + GraphHeight - 5 < progressRect.bottom)
-				{
-					int posX = int(m_scrollPosX) + progressRect.left;
-					int graphBaseY = posY + GraphHeight - 4;
-					double graphHeight = double(GraphHeight - 2);
-
-					Vector<POINT> line;
-
-					auto drawGraph = [&](auto& values, auto maxValue, double scale, HPEN pen, bool accumulatingValue, int offsetY = 0)
-						{
-							bool loop = true;
-							u32 reconnectIndex = 0;
-							while (loop)
-							{
-								context.SetPen(pen);
-								bool isFirstUpdate = true;
-								bool isFirstDraw = true;
-								decltype(maxValue) prevValue = 0;
-								u64 prevTime = 0;
-								int prevX = 0;
-								int prevPrevX = 0;
-								int prevY = 0;
-
-								u64 i = 0;
-								u64 e = 0;
-
-								if (reconnectIndex)
-									i = session.reconnectIndices[reconnectIndex - 1];
-
-								if (reconnectIndex < session.reconnectIndices.size())
-									e = session.reconnectIndices[reconnectIndex];
-								else
-								{
-									e = Min(session.updates.size(), values.size());
-									loop = false;
-								}
-
-								line.clear();
-								for (; i!=e; ++i)
-								{
-									u64 updateTime = session.updates[i];
-									auto value = values[i];
-									int x = posX + int(float(TimeToS(updateTime)) * scaleX);
-									int y = graphBaseY;
-
-									double duration = TimeToS(updateTime - prevTime);
-									if (updateTime == 0)
-										isFirstUpdate = true;
-
-									if (accumulatingValue)
-									{
-										double invScaleY = duration * scale;
-										if (invScaleY != 0 && prevValue != 0)
-											y -= int(double(value - prevValue) / invScaleY) + offsetY;
-									}
-									else
-									{
-										double div = 1.0f;
-										if (maxValue != 0)
-											div = double(maxValue);
-										y -= int(double((double)maxValue - (double)value*scale)*graphHeight/div) + offsetY;
-									}
-
-									if (prevValue <= 0)
-										prevY = y;
-
-									if (x > clientRect.left)
-									{
-										if (isFirstUpdate)
-										{
-											isFirstDraw = false;
-										}
-										else
-										{
-											if (isFirstDraw)
-												line.push_back({prevPrevX, prevY});
-											line.push_back({prevX, y});
-											isFirstDraw = false;
-										}
-									}
-									if (prevX > clientRect.right)
-										break;
-									isFirstUpdate = false;
-									prevPrevX = prevX;
-									prevX = x;
-									prevY = y;
-									prevValue = value;
-									prevTime = updateTime;
-								}
-								if (line.size() > 1)
-									context.Polyline(line);
-								++reconnectIndex;
-							}
-						};
-
-					if (m_config.showNetworkStats && session.highestSendPerS != 0 && session.highestRecvPerS != 0)
-					{
-						drawGraph(session.networkSend, 100000000000000ull, double(session.highestSendPerS) / graphHeight, m_sendPen, true);
-						drawGraph(session.networkRecv, 100000000000000ull, double(session.highestRecvPerS) / graphHeight, m_recvPen, true, 1);
-					}
-
-					if (m_config.showCpuMemStats)
-					{
-						drawGraph(session.cpuLoad, 0.0f, -1.0, m_cpuPen, false);
-						drawGraph(session.memAvail, session.memTotal, 1.0, m_memPen, false);
-					}
-
-					if (m_config.showDriveStats)
-						for (auto& pair : session.drives)
-							if (pair.second.busyHighest)
-								drawGraph(pair.second.busyPercent, 0, -0.01, m_drivePen, false);
-
-				}
-				posY += GraphHeight;
-			}
-
-			if (m_config.showDetailedData)
-			{
-				auto drawText = [&](const StringBufferBase& text, RECT& rect, u32* outWidth)
-					{
-						if (rect.top > progressRect.bottom)
-							return;
-						bool selected = m_fetchedFilesSelected == processLocation.sessionIndex && text.StartsWith(TC("Fetched Files"));
-						if (selected)
-							context.SetBkMode(OPAQUE);
-						if (rect.bottom > progressRect.bottom)
-						{
-							rect.bottom = progressRect.bottom;
-							context.ExtTextOutW(rect.left, rect.top, ETO_CLIPPED, &rect, text);
-						}
-						else
-							context.ExtTextOutW(rect.left, rect.top, 0, NULL, text);
-						if (selected)
-							context.SetBkMode(TRANSPARENT);
-						if (outWidth)
-							*outWidth = u32(context.GetTextSize(text).cx);
-					};
-				bool isRemote = processLocation.sessionIndex != 0;
-				PaintDetailedStats(posY, progressRect, session, isRemote, playTime, drawText);
-			}
-
-			context.SetActiveFont(m_processFont);
-
-			bool shouldDrawText = m_processFont.height > 4;
-
-			if (m_config.showProcessBars)
-			{
-				processLocation.processorIndex = 0;
-				for (auto& processor : session.processors)
-				{
-					bool drawProcessorIndex = m_config.showFinishedProcesses;
-
-					if (posY + m_sessionStepY >= progressRect.top && posY < progressRect.bottom)
-					{
-						int barHeight = boxHeight;
-						int textOffsetY = 0;
-						if (posY + boxHeight > progressRect.bottom)
-						{
-							int newBarHeight = Min(barHeight, int(progressRect.bottom - posY));
-							textOffsetY = int(barHeight - newBarHeight);
-							barHeight = newBarHeight;
-						}
-
-						const int textHeight = barHeight;
-						const int rectBottom = posY + textHeight;
-						const int offsetY = (textHeight - m_processFont.height + textOffsetY) / 2;
-
-						processLocation.processIndex = 0;
-						int posX = int(m_scrollPosX) + progressRect.left;
-						for (auto& process : processor.processes)
-						{
-							int left = posX + int(float(TimeToS(process.start)) * scaleX);
-
-							auto pig = MakeGuard([&]() { ++processLocation.processIndex; });
-
-							if (left >= progressRect.right)
-								continue;
-
-							u64 stop = process.stop;
-							bool done = stop != ~u64(0);
-							if (!done)
-								stop = playTime;
-							else if (!m_config.showFinishedProcesses)
-								continue;
-
-							drawProcessorIndex = true;
-
-							RECT rect;
-							rect.left = left;
-							rect.right = Max(posX + int(float(TimeToS(stop)) * scaleX) - 1, left + 1);
-							rect.top = posY;
-							rect.bottom = rectBottom - 1;
-
-							if (rect.right <= progressRect.left)
-								continue;
-
-							if (!m_filterString.empty() && !Contains(process.description.c_str(), m_filterString.c_str()) && !Contains(process.breadcrumbs.c_str(), m_filterString.c_str()))
-								continue;
-							++visibleBoxes;
-
-							bool selected = m_processSelected && m_processSelectedLocation == processLocation;
-							if (selected)
-								process.bitmapDirty = true;
-
-							--rect.top;
-							PaintProcessRect(context, process, rect, progressRect, selected, false);
-							++rect.top;
-
-							int processWidth = rect.right - rect.left;
-							if (shouldDrawText && m_config.ShowProcessText && processWidth > 3)
-							{
-								context.DrawCached(process, processWidth, offsetY, textHeight, left, rect, progressRect, [&](PaintContext& textCacheContext)
-									{
-										RECT rect2{ 0, int(process.bitmapOffset), 256, int(process.bitmapOffset) + m_processFont.height };
-										RECT rect3{ 0, int(process.bitmapOffset), processWidth, int(process.bitmapOffset) + m_processFont.height };
-										if (!done)
-											rect3.right = 256;
-
-										PaintProcessRect(textCacheContext, process, rect3, rect2, selected, true);
-
-										rect2.left += 3; // Move in text a bit
-
-										int textY = rect2.top + m_processFont.offset;
-
-										bool dropShadow = m_config.DarkMode;
-										if (dropShadow)
-										{
-											textCacheContext.SetTextColor(RGB(5, 60, 5));
-											++rect2.left;
-											++textY;
-											textCacheContext.ExtTextOutW(rect2.left, textY, ETO_CLIPPED, &rect2, process.description);
-											--rect2.left;
-											--textY;
-										}
-										textCacheContext.SetTextColor(m_textColor);
-										textCacheContext.ExtTextOutW(rect2.left, textY, ETO_CLIPPED, &rect2, process.description);
-
-										if (!selected)
-											process.bitmapDirty = false;
-									});
-							}
-						}
-
-						if (drawProcessorIndex)
-						{
-							RECT rect;
-							rect.left = 5;
-							rect.right = progressRect.left - 2;
-							rect.top = posY;
-							rect.bottom = rectBottom;
-
-							StringBuffer<> buf;
-							buf.AppendValue(u64(processLocation.processorIndex) + 1);
-							context.ExtTextOutW(5, posY + offsetY, ETO_CLIPPED, &rect, buf);
-						}
-					}
-
-					lastStop = Max(lastStop, processor.processes.rbegin()->stop);
-
-					++processLocation.processorIndex;
-
-					if (drawProcessorIndex)
-						posY += processStepY;
-				}
-			}
-			else
-			{
-				for (auto& processor : session.processors)
-					if (!processor.processes.empty())
-						lastStop = Max(lastStop, processor.processes.rbegin()->stop);
-			}
-
-			if (m_config.showWorkers && isFirst)
-			{
-				u32 trackIndex = 0;
-				for (auto& workTrack : m_traceView.workTracks)
-				{
-					if (posY + m_sessionStepY >= progressRect.top && posY <= progressRect.bottom)
-					{
-						int textOffsetY = 0;
-						int barHeight = boxHeight;
-						if (posY + int(boxHeight) > progressRect.bottom)
-						{
-							int newBarHeight = Min(barHeight, int(progressRect.bottom - posY));
-							textOffsetY = barHeight - newBarHeight;
-							barHeight = newBarHeight;
-						}
-
-						const int textHeight = barHeight;
-						const int rectBottom = posY + textHeight;
-						const int offsetY = (textHeight - m_processFont.height + textOffsetY) / 2;
-
-						if (shouldDrawText)
-						{
-							RECT rect;
-							rect.left = 5;
-							rect.right = progressRect.left - 5;
-							rect.top = posY;
-							rect.bottom = rectBottom;
-
-							StringBuffer<> buf;
-							buf.AppendValue(u64(trackIndex) + 1);
-							context.ExtTextOutW(5, posY + offsetY, ETO_CLIPPED, &rect, buf);
-						}
-
-						int lastDrawnRight = 0;
-
-						u32 workIndex = 0;
-						int posX = int(m_scrollPosX) + progressRect.left;
-						for (auto& work : workTrack.records)
-						{
-							auto inc = MakeGuard([&](){ ++workIndex; });
-
-							if (work.start == work.stop)
-								continue;
-
-							float startTime = TimeToS(work.start);
-
-							int left = posX + int(startTime * scaleX);
-
-							if (left >= progressRect.right)
-								continue;
-
-							if (!m_filterString.empty())
-							{
-								bool keep = Contains(work.description, m_filterString.c_str());
-								if (!keep)
-									for (auto& en : work.entries)
-										keep |= Contains(en.text, m_filterString.c_str());
-								if (!keep)
-									continue;
-							}
-
-							HBRUSH brush;
-
-							u64 stop = work.stop;
-
-							float stopTime = TimeToS(stop);
-
-							RECT rect;
-							rect.left = left;
-							rect.right = posX + int(stopTime * scaleX) - 1;
-							rect.top = posY;
-							rect.bottom = rectBottom - 1;
-
-							if (rect.right <= progressRect.left)
-								continue;
-
-							++visibleBoxes;
-
-							rect.right = Max(int(rect.right), left + 1);
-
-							Color color = work.color;
-							bool selected = m_workSelected && m_workTrack == trackIndex && m_workIndex == workIndex;
-							if (selected)
-							{
-								selectedWork = work;
-								work.bitmapDirty = true;
-								auto c = (u8*)&color;
-								for (u32 j=0; j!=3; ++j)
-									c[j] = u8(Min(int(c[j]) + 40, 255));
-							}
-							else
-							{
-								if (rect.left + 1 == rect.right)
-									if (lastDrawnRight == rect.right)
-										continue;
-							}
-
-							lastDrawnRight = rect.right;
-
-							bool done = stop != ~u64(0);
-							if (done)
-							{
-								auto insres = m_coloredBrushes.try_emplace(color);
-								if (insres.second)
-								{
-									auto c = (u8*)&color;
-									insres.first->second = CreateSolidBrush(RGB(c[2], c[1], c[0]));
-								}
-								brush = insres.first->second;
-							}
-							else
-							{
-								//stop = playTime;
-								brush = m_processBrushes[0].inProgress;
-							}
-
-							--rect.top;
-
-							auto clampRect = [&](RECT& r) { r.left = Min(Max(r.left, progressRect.left), progressRect.right); r.right = Max(Min(r.right, progressRect.right), progressRect.left); };
-
-							clampRect(rect);
-							context.FillRect(rect, brush);
-							++rect.top;
-
-							int processWidth = rect.right - rect.left;
-							if (shouldDrawText && m_config.ShowProcessText && processWidth > 3)
-							{
-								context.DrawCached(work, processWidth, offsetY, textHeight, left, rect, progressRect, [&](PaintContext& textCacheContext)
-									{
-										RECT rect2{ 0,int(work.bitmapOffset), 256, int(work.bitmapOffset) + m_processFont.height };
-										textCacheContext.FillRect(rect2, brush);
-										textCacheContext.ExtTextOutW(rect2.left, rect2.top, ETO_CLIPPED, &rect2, ToView(work.description));
-										if (!selected)
-											work.bitmapDirty = false;
-									});
-							}
-						}
-					}
-					++trackIndex;
-					posY += processStepY;
-				}
-			}
-
-			context.SetActiveFont(m_defaultFont);
-		}
-
-		if (true)
-		{
-			for (auto& kv : m_traceView.tables)
-			{
-				auto& table = kv.second;
-				u32 tableOffsetX = m_activeFont.height*40;
-				u32 tableOffsetY = 30;
-
-				u32 rowOffset = m_activeFont.height + 2;
-				u32 columnOffset = table.rowCharWidth * m_activeFont.height;
-
-				u32 columnIndex = 0;
-				u32 rowIndex = 0;
-				for (const tchar* item : table.items)
-				{
-					u32 y = tableOffsetY + rowIndex*rowOffset;
-					u32 x = tableOffsetX + columnIndex*columnOffset;
-					if (int(x) >= progressRect.right)
-						break;
-
-					context.ExtTextOutW(x, y, ETO_CLIPPED, nullptr, ToView(item));
-
-					++rowIndex;
-					if (rowIndex > table.rowMaxCount || int(y) >= progressRect.bottom)
-					{
-						rowIndex = 0;
-						++columnIndex;
-					}
-				}
-			}
-		}
-
-		m_contentWidth = m_progressRectLeft + Max(0, int(TimeToS((lastStop != 0 && lastStop != ~u64(0)) ? lastStop : playTime) * scaleX));
-
-
-		m_contentHeight = posY - int(m_scrollPosY) + stepY + 14;
-
-		float timelineSelected = m_timelineSelected;
-
-		if (m_config.showTimeline && !m_traceView.sessions.empty())
-			PaintTimeline(context, clientRect);
-
-		if (!m_filterString.empty())
-		{
-			StringBuffer<> str[2];
-			str[0].Append(TCV("Box Filter: ")).Append(m_filterString);
-			str[1].Appendf(TC("Box Count: %u"), visibleBoxes);
-
-			context.SetActiveFont(m_popupFont);
-			context.DrawCenteredText(str, 2, true);
-		}
-
-		if (m_config.showCursorLine && m_mouseOverWindow)
-		{
-			float timeScale = (m_config.horizontalScaleValue * m_zoomValue)*50.0f;
-			float startOffset = -(m_scrollPosX / timeScale);
-			POINT pos;
-			GetCursorPos(&pos);
-			ScreenToClient(m_hwnd, &pos);
-			timelineSelected = startOffset + float(pos.x - m_progressRectLeft) / timeScale;
-		}
-
-		if (timelineSelected != 0)
-		{
-			int posX = int(m_scrollPosX) + progressRect.left;
-			int left = posX + int(timelineSelected * scaleX);
-			int timelineTop = GetTimelineTop(clientRect);
-
-			// TODO: Draw line up
-			context.Line(left, 2, left, timelineTop);
-
-			if (timelineSelected >= 0)
-			{
-				StringBuffer<> b;
-				u32 milliseconds = u32(timelineSelected * 1000.0f);
-				u32 seconds = milliseconds / 1000;
-				milliseconds -= seconds * 1000;
-				u32 minutes = seconds / 60;
-				seconds -= minutes * 60;
-				u32 hours = minutes / 60;
-				minutes -= hours * 60;
-				if (hours)
-				{
-					b.AppendValue(hours).Append('h');
-					if (minutes < 10)
-						b.Append('0');
-				}
-				if (minutes || hours)
-				{
-					b.AppendValue(minutes).Append('m');
-					if (seconds < 10)
-						b.Append('0');
-				}
-				b.AppendValue(seconds).Append('.');
-				if (milliseconds < 100)
-					b.Append('0');
-				if (milliseconds < 10)
-					b.Append('0');
-				b.AppendValue(milliseconds);
-
-				StringBuffer<128> wt = GetWorldTime(timelineSelected);
-
-				context.SetActiveFont(m_popupFont);
-				DrawTextLogger logger(context, m_popupFont.height, m_tooltipBackgroundBrush);
-
-				logger.Info(L"%s (%s)", b.data, wt.data);
-				logger.DrawAtPos(left + 4, timelineTop - 20);
-			}
-		}
-
-		{
-			int boxSide = 8;
-			int boxStride = boxSide + 2;
-			int top = 5;
-			int bottom = top + boxSide;
-			int left = progressRect.right - 7 - boxSide;
-			int right = progressRect.right - 7;
-			bool* values = &m_config.showProgress;
-			for (int i = VisualizerFlag_Count - 1; i >= 0; --i)
-			{
-				context.SetPen(m_buttonSelected == u32(i) ? m_textPen : m_checkboxPen);
-				context.Rectangle(left, top, right, bottom);
-
-				if (values[i])
-				{
-					context.Line(left + 2, top + 2, right - 2, bottom - 2);
-					context.Line(right - 3, top + 2, left + 1, bottom - 2);
-				}
-
-				left -= boxStride;
-				right -= boxStride;
-			}
-
-			top -= 2;
-
-			auto drawText = [&](StringView text, COLORREF color)
-				{
-					context.SetTextColor(color);
-					left -= (context.GetTextSize(text).cx) + 5;
-					context.ExtTextOutW(left, top, 0, NULL, text);
-				};
-
-			if (m_config.showDriveStats)
-			{
-				context.SetActiveFont(m_defaultFont);
-				drawText(TCV("DRV"), m_driveColor);
-			}
-
-			if (m_config.showNetworkStats)
-			{
-				context.SetActiveFont(m_defaultFont);
-				drawText(TCV("SND"), m_sendColor);
-				drawText(TCV("RCV"), m_recvColor);
-			}
-
-			if (m_config.showCpuMemStats)
-			{
-				context.SetActiveFont(m_defaultFont);
-				drawText(TCV("CPU"), m_cpuColor);
-				drawText(TCV("MEM"), m_memColor);
-			}
-
-			context.SetTextColor(m_textColor);
-		}
-
-		if (m_processSelected)
-		{
-			const TraceView::Process& process = m_traceView.GetProcess(m_processSelectedLocation);
-			u64 duration = 0;
-
-			Vector<TString> logLines;
-			u32 maxCharCount = 50u;
-
-			bool hasExited = process.stop != ~u64(0);
-			if (hasExited)
-			{
-				duration = process.stop - process.start;
-
-				if (!process.logLines.empty())
-				{
-					u32 lineMaxCount = 0;
-					for (auto& line : process.logLines)
-					{
-						u32 offset = 0;
-						u32 left = u32(line.text.size());
-						while (left)
-						{
-							u32 toCopy = Min(left, maxCharCount);
-							lineMaxCount = Max(lineMaxCount, toCopy);
-							logLines.push_back(line.text.substr(offset, toCopy));
-							offset += toCopy;
-							left -= toCopy;
-						}
-					}
-				}
-			}
-			else
-			{
-				duration = playTime - process.start;
-			}
-
-			context.SetActiveFont(m_popupFont);
-			DrawTextLogger logger(context, m_popupFont.height, m_tooltipBackgroundBrush);
-
-			logger.AddTextOffset(-10); // Remove spaces in the front
-			logger.AddWidth(3);
-
-			logger.AddSpace(2);
-			logger.Info(L"  %ls", process.description.c_str());
-			if (process.type != TraceView::ProcessType_Task)
-			{
-				logger.Info(L"  Host:        %ls", m_processSelectedLocation.sessionIndex == 0 ? TC("local") : m_traceView.GetSession(m_processSelectedLocation).name.c_str());
-				logger.Info(L"  ProcessId:  %6u", process.id);
-			}
-			logger.Info(L"  Start:     %7ls (%ls)", TimeToText(process.start, true).str, GetWorldTime(process.start).data);
-			logger.Info(L"  Duration:  %7ls", TimeToText(duration, true).str);
-			if (!process.returnedReason.empty())
-				logger.Info(L"  Returned: %7s", process.returnedReason.data());
-			if (hasExited && process.exitCode != 0)
-			{
-				if (process.exitCode == ProcessCancelExitCode)
-					logger.Info(L"  ExitCode: Cancelled");
-				else if (process.exitCode >= 0xC0000000u)
-					logger.Info(L"  ExitCode: 0x%x", process.exitCode);
-				else
-					logger.Info(L"  ExitCode: %7u", process.exitCode);
-			}
-
-			const auto& breadcrumbs = process.breadcrumbs;
-			if (!breadcrumbs.empty())
-			{
-				constexpr TString::size_type maxLineLen = 50;
-				logger.Info(L"");
-				logger.Info(L"  ------------ Breadcrumbs ------------");
-				for (TString::size_type lineStart = 0, lineEnd = 0; lineEnd < breadcrumbs.size(); lineStart = lineEnd + 1)
-				{
-					// Log each individual line
-					lineEnd = breadcrumbs.find(L'\n', lineStart);
-					TString line = (lineEnd == TString::npos ? breadcrumbs.substr(lineStart) : breadcrumbs.substr(lineStart, lineEnd - lineStart));
-
-					// Break each line down into smaller section if they are longer than the maximum allowed length
-					if (line.size() > maxLineLen)
-					{
-						for (TString::size_type sectionStart = 0, sectionEnd = 0; sectionStart < line.size(); sectionStart = sectionEnd)
-						{
-							const TString::size_type maxSectionLen = sectionStart == 0 ? maxLineLen : maxLineLen - 2;
-							sectionEnd = std::min<TString::size_type>(sectionEnd + maxSectionLen, line.size());
-							TString section = (sectionStart == 0 ? L"  " : L"    ") + line.substr(sectionStart, sectionEnd - sectionStart);
-							logger.Info(section.c_str());
-						}
-					}
-					else
-					{
-						line = L"  " + line;
-						logger.Info(line.c_str());
-					}
-				}
-			}
-
-			if (process.stop != ~u64(0) && !process.stats.empty())
-			{
-				BinaryReader reader(process.stats.data(), 0, process.stats.size());
-				ProcessStats processStats;
-				SessionStats sessionStats;
-				StorageStats storageStats;
-				KernelStats kernelStats;
-				CacheFetchStats cacheStats;
-
-				if (process.type == TraceView::ProcessType_CacheFetchTest || process.type == TraceView::ProcessType_CacheFetchDownload)
-				{
-					if (!process.returnedReason.empty())
-						logger.Info(L"  Cache:       Miss");
-					else
-						logger.Info(L"  Cache:        Hit");
-					cacheStats.Read(reader, m_traceView.version);
-					if (reader.GetLeft())
-					{
-						storageStats.Read(reader, m_traceView.version);
-						kernelStats.Read(reader, m_traceView.version);
-					}
-				}
-				else
-				{
-					processStats.Read(reader, m_traceView.version);
-
-					if (reader.GetLeft())
-					{
-						if (process.isRemote || (m_traceView.version >= 36 && !process.isReuse))
-							sessionStats.Read(reader, m_traceView.version);
-						storageStats.Read(reader, m_traceView.version);
-						kernelStats.Read(reader, m_traceView.version);
-					}
-				}
-
-				if (processStats.hostTotalTime)
-				{
-					logger.Info(L"");
-					logger.Info(L"  ----------- Detours stats -----------");
-					processStats.Print(logger, m_traceView.frequency);
-				}
-				else if (processStats.peakMemory)
-				{
-					logger.Info(L"");
-					logger.Info(L"  ----------- Process stats -----------");
-					processStats.Print(logger, m_traceView.frequency);
-				}
-
-				if (!sessionStats.IsEmpty())
-				{
-					logger.Info(L"");
-					logger.Info(L"  ----------- Session stats -----------");
-					sessionStats.Print(logger, m_traceView.frequency);
-				}
-
-				if (!cacheStats.IsEmpty())
-				{
-					logger.Info(L"");
-					logger.Info(L"  ------------ Cache stats ------------");
-					cacheStats.Print(logger, m_traceView.frequency);
-				}
-
-				if (!storageStats.IsEmpty())
-				{
-					logger.Info(L"");
-					logger.Info(L"  ----------- Storage stats -----------");
-					storageStats.Print(logger, m_traceView.frequency);
-				}
-
-				if (!kernelStats.IsEmpty())
-				{
-					logger.Info(L"");
-					logger.Info(L"  ----------- Kernel stats ------------");
-					kernelStats.Print(logger, false, m_traceView.frequency);
-				}
-
-				PrintCacheWriteStats(logger, process.id);
-
-				if (!logLines.empty())
-				{
-					logger.Info(L"");
-					logger.Info(L"  ---------------- Log ----------------");
-					logger.AddTextOffset(14);
-					for (auto& line : logLines)
-						logger.Log(LogEntryType_Info, line);
-				}
-			}
-			logger.AddSpace(3);
-			logger.DrawAtCursor();
-		}
-		else if (m_workSelected && selectedWork.description)
-		{
-			u64 duration;
-			if (selectedWork.stop != ~u64(0))
-				duration = selectedWork.stop - selectedWork.start;
-			else
-				duration = playTime - selectedWork.start;
-
-			context.SetActiveFont(m_popupFont);
-			DrawTextLogger logger(context, m_popupFont.height, m_tooltipBackgroundBrush);
-
-			logger.AddSpace();
-			logger.Info(L"  %ls", selectedWork.description);
-			logger.Info(L"  Start:     %ls (%ls)", TimeToText(selectedWork.start, true).str, GetWorldTime(selectedWork.start).data);
-			logger.Info(L"  Duration:  %ls", TimeToText(duration, true).str);
-			if (!selectedWork.entries.empty())
-				logger.AddSpace();
-			for (u32 i=0,e=u32(selectedWork.entries.size()); i!=e; ++i)
-			{
-				auto& entry = selectedWork.entries[i];
-				u64 time = 0;
-				if (i == 0)
-				{
-					u64 start = entry.startTime ? entry.startTime : entry.time;
-					if (TimeToMs(start - selectedWork.start) > 1)
-						logger.Info(L"   Start (%ls)", TimeToText(start - selectedWork.start).str);
-				}
-
-				if (entry.startTime)
-					time = entry.time - entry.startTime;
-
-				if (!time)
-					for (u32 j=i+1;j<e;++j)
-					{
-						auto& next = selectedWork.entries[j];
-						if (next.startTime)
-							continue;
-						time = next.time - entry.time;
-						break;
-					}
-
-				if (!time && selectedWork.stop != ~0u)
-					time = selectedWork.stop - entry.time;
-				if (entry.count == 1)
-					logger.Info(L"%s  %ls (%ls)", (entry.startTime ? L" " : L""), entry.text, TimeToText(time).str);
-				else
-					logger.Info(L"%s  %ls (%ls %u)", (entry.startTime ? L" " : L""), entry.text, TimeToText(time).str, entry.count);
-			}
-			logger.AddSpace();
-			logger.DrawAtCursor();
-		}
-		else if (m_sessionSelectedIndex != ~0u)
-		{
-			context.SetActiveFont(m_popupFont);
-			DrawTextLogger logger(context, m_popupFont.height, m_tooltipBackgroundBrush);
-			logger.AddWidth(3);
-			logger.AddSpace(2);
-
-			auto& session = m_traceView.sessions[m_sessionSelectedIndex];
-			Vector<TString>& summary = session.summary;
-			if (summary.empty())
-			{
-				if (m_traceView.finished)
-					logger.Info(L" Session summary was never received");
-				else
-					logger.Info(L" Session summary not available until session is done");
-			}
-
-			logger.AddTextOffset(-10); // Remove spaces in the front
-			for (auto& line : summary)
-				logger.Log(LogEntryType_Info, line);
-			logger.AddTextOffset(0);
-
-			if (m_sessionSelectedIndex == 0)
-			{
-				for (auto& cacheSummary : m_traceView.cacheSummaries)
-					for (auto& line : cacheSummary.lines)
-						logger.Log(LogEntryType_Info, line);
-			}
-
-			for (auto& pair : session.drives)
-			{
-				auto& d = pair.second;
-				logger.Info(L"  %c: Read: %s (%u) Write: %s (%u)", pair.first, BytesToText(d.totalReadBytes).str, d.totalReadCount, BytesToText(d.totalWriteBytes).str, d.totalWriteCount);
-			}
-
-			logger.AddSpace(3);
-			logger.DrawAtCursor();
-		}
-		else if (m_statsSelected)
-		{
-			context.SetActiveFont(m_popupFont);
-			DrawTextLogger logger(context, m_popupFont.height, m_tooltipBackgroundBrush);
-			logger.AddSpace(3);
-			logger.SetColor(m_cpuColor).Info(L"  Cpu: %.1f%%", m_stats.cpuLoad * 100.0f);
-			logger.SetColor(m_memColor).Info(L"  Mem: %ls/%ls", BytesToText(m_stats.memTotal - m_stats.memAvail).str, BytesToText(m_stats.memTotal).str);
-			if (m_stats.recvBytes || m_stats.sendBytes)
-			{
-				logger.SetColor(m_recvColor).Info(L"  Recv: %lsit/s", BytesToText(m_stats.recvBytesPerSecond*8).str);
-				logger.SetColor(m_sendColor).Info(L"  Send: %lsit/s", BytesToText(m_stats.sendBytesPerSecond*8).str);
-			}
-			logger.SetColor(m_textColor);
-
-			if (m_stats.ping)
-				logger.Info(L"  Ping: %ls", TimeToText(m_stats.ping, false, m_traceView.frequency).str);
-			if (m_stats.connectionCount)
-				logger.Info(L"  Conn: %u", u32(m_stats.connectionCount));
-
-			for (auto& pair : m_stats.drives)
-				logger.SetColor(m_driveColor).Info(L"  %c: %u%% R:%s/s W:%s/s", pair.first, pair.second.busyPercent, BytesToText(pair.second.readPerSecond).str, BytesToText(pair.second.writePerSecond).str);
-
-			logger.AddSpace(3);
-			logger.DrawAtCursor();
-		}
-		else if (m_activeProcessGraphSelected)
-		{
-			context.SetActiveFont(m_popupFont);
-			DrawTextLogger logger(context, m_popupFont.height, m_tooltipBackgroundBrush);
-			logger.AddSpace(3);
-			logger.AddSpace(3);
-			logger.SetColor(m_activeProcColor).Info(L"  Active Processes: %u", m_activeProcessCount);
-			logger.DrawAtCursor();
-		}
-		else if (m_buttonSelected != ~0u)
-		{
-			const wchar_t* tooltip[] =
-			{
-#define UBA_VISUALIZER_FLAG(name, defaultValue, desc) desc,
-				UBA_VISUALIZER_FLAGS1
-#undef UBA_VISUALIZER_FLAG
-			};
-
-			context.SetActiveFont(m_popupFont);
-			DrawTextLogger logger(context, m_popupFont.height, m_tooltipBackgroundBrush);
-			logger.Info(L"%ls %ls", L"Show", tooltip[m_buttonSelected]);
-			logger.DrawAtCursor();
-		}
-		else if (m_fetchedFilesSelected != ~0u)
-		{
-			auto& session = m_traceView.sessions[m_fetchedFilesSelected];
-			auto& fetchedFiles = session.fetchedFiles;
-			if (!fetchedFiles.empty() && !fetchedFiles[0].hint.empty())
-			{
-				/*
-				int colWidth = 500;
-				int width = colWidth * 2;
-				int height = Min(int(clientRect.bottom), int(fetchedFiles.size() * m_popupFont.height));
-
-				context.SetActiveFont(m_defaultFont);
-				DrawTextLogger logger(context, r, m_font.height, m_tooltipBackgroundBrush);
-				for (auto& f : fetchedFiles)
-				{
-				if (f.hint == TC("KnownInput"))
-				continue;
-				if (logger.rect.top >= r.bottom - m_font.height)
-				{
-				if (logger.rect.left + colWidth >= r.right)
-				{
-				logger.Info(L"...");
-				break;
-				}
-				logger.rect.top = r.top;
-				logger.rect.left += colWidth;
-				}
-				logger.Info(L"%s", f.hint.c_str());
-				}
-				logger.DrawAtCursor();
-				*/
-			}
-		}
-	}
-
-	u64 ConvertTime(const TraceView& view, u64 time);
-
-	void VisualizerWin::PaintActiveProcesses(PaintContext& context, int& posY, const RECT& clientRect, const Function<void(TraceView::ProcessLocation&, u32, bool)>& drawProcess)
-	{
-		context.SetActiveFont(m_processFont);
-		int startPosY = posY;
-
-		Map<u64, TraceView::ProcessLocation*> activeProcesses;
-		u32 remoteCount = 0;
-		for (auto& kv : m_trace.m_activeProcesses)
-		{
-			auto& active = kv.second;
-			TraceView::Process& process = m_trace.GetSession(m_traceView, active.sessionIndex).processors[active.processorIndex].processes[active.processIndex];
-			u64 start = process.start; //~0llu - process.start;
-			activeProcesses.try_emplace(start, &active);
-			if (process.isRemote)
-				++remoteCount;
-		}
-
-		u32 maxHeight = clientRect.bottom;
-
-		bool fillHeight = ActiveProcessesShouldFillHeight();
-		if (fillHeight)
-		{
-			maxHeight = clientRect.bottom - posY;
-			if (m_config.showTimeline)
-				maxHeight -= GetTimelineHeight();
-		}
-		else
-		{
-			u32 maxHeight2 = m_config.maxActiveVisible * (m_activeFont.height + 2);
-			maxHeight = Min(maxHeight, maxHeight2);
-		}
-
-		u32 maxSize = Min(m_config.maxActiveProcessHeight, 32u);
-		u32 maxSizeMinusOne = maxSize - 1;
-
-		u32 counts[128] = { 0 };
-
-		u32 highestHistoryCount = 0;
-		for (u32 i=0; i!= sizeof_array(m_activeProcessCountHistory) - 1; ++i)
-			highestHistoryCount = Max(highestHistoryCount, m_activeProcessCountHistory[i]);
-
-		u32 activeProcessCount = highestHistoryCount;
-		counts[0] = activeProcessCount;
-		u32 totalHeight = counts[0]*2;
-		while (totalHeight < maxHeight && counts[maxSizeMinusOne] != activeProcessCount)
-		{
-			bool changed = false;
-			for (u32 i=0; i!=maxSizeMinusOne; ++i)
-			{
-				if (counts[i] && counts[i] > counts[i+1]*2 + 1)
-				{
-					--counts[i];
-					++counts[i+1];
-					++totalHeight;
-					changed = true;
-				}
-			}
-			if (!changed)
-			{
-				for (u32 j=0; j!=maxSizeMinusOne; ++j)
-				{
-					if (!counts[j])
-						continue;
-					++counts[j + 1];
-					--counts[j];
-					++totalHeight;
-					break;
-				}
-			}
-		}
-
-		auto it = activeProcesses.begin();
-		auto itEnd = activeProcesses.end();
-		int endY = int(startPosY + maxHeight);
-		for (u32 i=0;i!=maxSize;++i)
-		{
-			u32 v = maxSizeMinusOne - i;
-			u32 boxHeight = v + 1;
-
-			for (u32 j=0;j!=counts[v];++j)
-			{
-				if (it == itEnd || posY >= endY)
-					break;
-				auto& active = *it->second;
-				++it;
-				drawProcess(active, boxHeight, j == 0);
-				posY += boxHeight + 1;
-			}
-		}
-
-		//if (fillHeight)
-		//	posY = clientRect.bottom-1;//startPosY; // To prevent vertical scrollbar
-		if (fillHeight || counts[maxSizeMinusOne] != activeProcessCount)
-			posY = startPosY + maxHeight;
-		else
-			posY += 3;
-
-		context.SetActiveFont(m_defaultFont);
-	}
-
-	void VisualizerWin::PaintProcessRect(PaintContext& context, TraceView::Process& process, RECT rect, const RECT& progressRect, bool selected, bool writingBitmap)
-	{
-		auto clampRect = [&](RECT& r) { r.left = Min(Max(r.left, progressRect.left), progressRect.right); r.right = Max(Min(r.right, progressRect.right), progressRect.left); };
-		bool done = process.stop != ~u64(0);
-
-		HBRUSH brush = m_processBrushes[selected].success;
-		if (!process.returnedReason.empty())
-			brush = m_processBrushes[selected].returned;
-		else if (!done)
-			brush = process.isRacing ? m_processBrushes[selected].inProgressRace : m_processBrushes[selected].inProgress;
-		else if (process.type == TraceView::ProcessType_CacheFetchTest)
-			brush = m_processBrushes[selected].cacheFetchTest;
-		else if (process.type == TraceView::ProcessType_CacheFetchDownload)
-			brush = m_processBrushes[selected].cacheFetchDownload;
-		else if (process.exitCode == ProcessCancelExitCode)
-			brush = m_processBrushes[selected].returned;
-		else if (process.exitCode != 0)
-			brush = m_processBrushes[selected].error;
-
-		u64 writeFilesTime = process.writeFilesTime;
-
-		if (!done || process.exitCode != 0 || !m_config.ShowReadWriteColors)
-		{
-			if (writingBitmap)
-				rect.right = 256;
-			clampRect(rect);
-			context.FillRect(rect, brush);
-			return;
-		}
-
-		double duration = double(process.stop - process.start);
-
-		RECT main = rect;
-		int width = rect.right - rect.left;
-
-		double recvPart = (double(ConvertTime(m_traceView, process.createFilesTime)) / duration);
-		if (int headSize = int(recvPart * width))
-		{
-			main.left += headSize;
-			RECT r2 = rect;
-			r2.right = r2.left + headSize;
-			clampRect(r2);
-			if (r2.left != r2.right)
-				context.FillRect(r2, m_processBrushes[selected].recv);
-		}
-
-		double sendPart = (double(ConvertTime(m_traceView, writeFilesTime)) / duration);
-		if (int tailSize = int(sendPart * width))
-		{
-			main.right -= tailSize;
-			RECT r2 = rect;
-			r2.left = r2.right - tailSize;
-			clampRect(r2);
-			if (r2.left != r2.right)
-				context.FillRect(r2, m_processBrushes[selected].send);
-		}
-
-		clampRect(main);
-		if (main.left != main.right)
-			context.FillRect(main, brush);
-	}
-
-	void VisualizerWin::PaintTimeline(PaintContext& context, const RECT& clientRect)
-	{
-		context.SetActiveFont(m_timelineFont);
-		int top = GetTimelineTop(clientRect);
-		float timeScale = (m_config.horizontalScaleValue*m_zoomValue)*50.0f;
-		float startOffset = ((m_scrollPosX/timeScale) - float(int(m_scrollPosX/timeScale))) * timeScale;
-		int index = -int(startOffset/timeScale);
-
-		int number = -int(float(m_scrollPosX)/timeScale);
-
-		int textStepSize = int((5.0f / timeScale) + 1) * 5;
-		if (textStepSize > 150)
-			textStepSize = 600;
-		else if (textStepSize > 120)
-			textStepSize = 300;
-		else if (textStepSize > 90)
-			textStepSize = 240;
-		else if (textStepSize > 45)
-			textStepSize = 120;
-		else if (textStepSize > 30)
-			textStepSize = 60;
-		else if (textStepSize > 10)
-			textStepSize = 30;
-
-		int lineStepSize = textStepSize / 5;
-
-		RECT progressRect = clientRect;
-		progressRect.left += m_progressRectLeft;
-
-		context.SetPen(m_textPen);
-		Vector<DWORD> lines;
-		Vector<POINT> points;
-
-		while (true)
-		{
-			int pos = progressRect.left + int(startOffset + float(index)*timeScale);
-			if (pos >= clientRect.right)
-				break;
-
-			int lineBottom = top + 5;
-			if (!(number % textStepSize))
-			{
-				bool shouldDraw = true;
-				int seconds = number;
-				StringBuffer<> buffer;
-				if (seconds >= 60)
-				{
-					int min = seconds / 60;
-					seconds -= min * 60;
-					if (!seconds)
-					{
-						buffer.Appendf(L"%um", min);
-						lineBottom += 4;
-					}
-				}
-				if (!number || seconds)
-					buffer.Appendf(L"%u", seconds);
-
-				if (shouldDraw)
-				{
-					SIZE s = context.GetTextSize(buffer);
-					RECT textRect;
-					textRect.top = top + 8;
-					textRect.bottom = textRect.top + m_activeFont.height;
-					textRect.right = pos + s.cx/2;
-					textRect.left = pos - s.cx/2;
-					context.ExtTextOutW(textRect.left, textRect.top, 0, NULL, buffer);
-				}
-			}
-			if (!(number % lineStepSize))
-			{
-				points.push_back({pos, top});
-				points.push_back({pos, lineBottom});
-				lines.push_back(2);
-			}
-
-			++number;
-			++index;
-		}
-
-		points.push_back({m_contentWidth, top - 25});
-		points.push_back({m_contentWidth, top});
-		lines.push_back(2);
-
-		context.PolyPolyline(points.data(), lines.data(), int(lines.size()));
-	}
-
-	void VisualizerWin::PaintDetailedStats(int& posY, const RECT& progressRect, TraceView::Session& session, bool isRemote, u64 playTime, const DrawTextFunc& drawTextFunc)
-	{
-		int stepY = m_activeFont.height;
-		int startPosY = posY;
-		int posX = progressRect.left + 5;
-		int maxY = posY + stepY;
-		u32 maxTextWidth = 0;
-
-		RECT textRect;
-		auto drawText = [&](const wchar_t* format, ...)
-			{
-				textRect.top = posY;
-				textRect.bottom = posY + 20;
-				textRect.left = posX;
-				textRect.right = posX + 1000;
-				posY += stepY;
-				StringBuffer<> str;
-				va_list arg;
-				va_start(arg, format);
-				str.Append(format, arg);
-				va_end(arg);
-				u32 lastTextWidth = 0;
-				drawTextFunc(str, textRect, &lastTextWidth);
-				maxTextWidth = Max(maxTextWidth, lastTextWidth);
-				maxY = Max(maxY, posY);
-			};
-
-		if (isRemote)
-		{
-			drawText(L"Finished Processes: %u", session.processExitedCount);
-			drawText(L"Active Processes: %u", session.processActiveCount);
-			drawText(L"ClientId: %u  TcpCount: %u", session.clientUid.data1, session.connectionCount.empty() ? 1 : session.connectionCount.back());
-
-			if (session.disconnectTime == ~u64(0))
-			{
-				if (session.proxyCreated)
-					drawText(L"Proxy(HOSTED): %ls", session.proxyName.c_str());
-				else if (!session.proxyName.empty())
-					drawText(L"Proxy: %ls", session.proxyName.c_str());
-				else
-					drawText(L"Proxy: None");
-			}
-
-			bool hasFileDetails = !session.fetchedFiles.empty() || !session.storedFiles.empty();
-
-			if (!hasFileDetails)
-			{
-				posY = startPosY;
-				posX += maxTextWidth + 15;
-			}
-
-			if (!session.updates.empty())
-			{
-				auto updateTime = session.updates.back();
-				auto updateSend = session.networkSend.back();
-				auto updateRecv = session.networkRecv.back();
-				u64 sendPerS = 0;
-				u64 recvPerS = 0;
-				float duration = TimeToS(updateTime - session.prevUpdateTime);
-				if (duration != 0)
-				{
-					sendPerS = u64(float(updateSend - session.prevSend) / duration);
-					recvPerS = u64(float(updateRecv - session.prevRecv) / duration);
-				}
-				drawText(L"Recv: %ls (%sit/s)", BytesToText(updateRecv), BytesToText(recvPerS*8));
-				drawText(L"Send: %ls (%sit/s)", BytesToText(updateSend), BytesToText(sendPerS*8));
-			}
-
-			int fileWidth = 700;
-
-			auto drawFiles = [&](const wchar_t* fileType, Vector<TraceView::FileTransfer>& files, u64 count, u64 bytes, u64 activeCount, u32& maxVisibleFiles)
-				{
-					textRect.left = posX;
-					textRect.right = posX + fileWidth;
-					drawText(L"%ls Files: %u (%u) %s", fileType, u32(count), u32(activeCount), BytesToText(bytes));
-					u32 fileCount = 0;
-					for (auto rit = files.rbegin(), rend = files.rend(); rit != rend; ++rit)
-					{
-						TraceView::FileTransfer& file = *rit;
-						if (file.stop != ~u64(0))
-							continue;
-						u64 time = 0;
-						if (file.start < playTime)
-							time = playTime - file.start;
-
-						if ((((u8*)&file.key)[19] & IsProxyMask) != 0)
-							drawText(L"%s (proxy) %s", file.hint.c_str(), TimeToText(time, true).str);
-						else if (file.size == 0)
-							drawText(L"%s (calc) %s", file.hint.c_str(), TimeToText(time, true).str);
-						else
-							drawText(L"%s (%s) %s", file.hint.c_str(), BytesToText(file.size).str, TimeToText(time, true).str);
-
-						if (fileCount++ > 5)
-							break;
-					}
-					posY += stepY * (maxVisibleFiles - fileCount);
-					maxVisibleFiles = Max(maxVisibleFiles, fileCount);
-				};
-
-			Vector<TraceView::FileTransfer> fetchedFiles;
-			for (auto& kv : session.fetchedFilesActive)
-				fetchedFiles.push_back(session.fetchedFiles[kv.second]);
-			std::sort(fetchedFiles.begin(), fetchedFiles.end(), [](const TraceView::FileTransfer& a, const TraceView::FileTransfer& b) { return a.start > b.start; });
-
-			if (hasFileDetails)
-			{
-				posY = startPosY;
-				posX += maxTextWidth + 15;
-			}
-			drawFiles(L"Fetched", fetchedFiles, session.fetchedFilesCount, session.fetchedFilesBytes, session.fetchedFilesActive.size(), session.maxVisibleFiles);
-			if (hasFileDetails)
-			{
-				posY = startPosY;
-				posX += fileWidth;
-			}
-			drawFiles(L"Stored", session.storedFiles, session.storedFilesCount, session.storedFilesBytes, session.storedFilesActive.size(), session.maxVisibleFiles);
-		}
-		else
-		{
-			if (m_traceView.totalProcessActiveCount || m_traceView.totalProcessExitedCount)
-			{
-				drawText(L"Finished Processes: %u (local: %u)", m_traceView.totalProcessExitedCount, session.processExitedCount);
-				drawText(L"Active Processes: %u (local: %u)", m_traceView.totalProcessActiveCount, session.processActiveCount);
-				drawText(L"Active Helpers: %u", Max(1u, m_traceView.activeSessionCount) - 1);
-			}
-
-			if (session.highestSendPerS > 0.0f || session.highestRecvPerS > 0.0f)
-			{
-				auto updateTime = session.updates.back();
-				auto updateSend = session.networkSend.back();
-				auto updateRecv = session.networkRecv.back();
-				if (updateSend || updateRecv)
-				{
-					u64 sendPerS = 0;
-					u64 recvPerS = 0;
-					float duration = TimeToS(updateTime - session.prevUpdateTime);
-					if (duration != 0)
-					{
-						sendPerS = u64(float(updateSend - session.prevSend) / duration);
-						recvPerS = u64(float(updateRecv - session.prevRecv) / duration);
-					}
-
-					drawText(L"Recv: %ls (%sit/s)", BytesToText(updateRecv), BytesToText(recvPerS*8));
-					drawText(L"Send: %ls (%sit/s)", BytesToText(updateSend), BytesToText(sendPerS*8));
-				}
-			}
-
-			if (!session.updates.empty())
-			{
-				posY = startPosY;
-				posX += maxTextWidth + 10;
-
-				for (auto& pair : session.drives)
-				{
-					auto& drive = pair.second;
-
-					u64 readPerS = 0;
-					u64 writePerS = 0;
-
-					u64 updateCount = session.updates.size();
-					float duration = TimeToS(session.updates.back() - session.prevUpdateTime);
-					if (duration != 0 && updateCount > 1)
-					{
-						readPerS = u64(float(drive.readBytes.back()) / duration);
-						writePerS = u64(float(drive.writeBytes.back()) / duration);
-					}
-
-					drawText(L"%c: Rd %ls (%s/s) Wr %ls (%s/s) ", pair.first, BytesToText(drive.totalReadBytes).str, BytesToText(readPerS).str, BytesToText(drive.totalWriteBytes).str, BytesToText(writePerS).str);
-				}
-			}
-		}
-
-		posY = maxY;
-	}
-
-	int VisualizerWin::GetTimelineHeight()
-	{
-		return m_timelineFont.height + 8;
-	}
-
-	int VisualizerWin::GetTimelineTop(const RECT& clientRect)
-	{
-		int timelineHeight = GetTimelineHeight();
-		int posY = m_contentHeight - timelineHeight;
-		int maxY = int(clientRect.bottom - timelineHeight);
-		return m_config.LockTimelineToBottom ? maxY : Min(posY, maxY);
-	}
-
-	void VisualizerWin::HitTest(HitTestResult& outResult, const POINT& pos)
-	{
-		u64 playTime = GetPlayTime();
-
-		RECT clientRect;
-		GetClientRect(m_hwnd, &clientRect);
-
-		PaintContext context { *this, NULL, NULL, clientRect };
-		context.SetActiveFont(m_defaultFont);
-
-		int posY = int(m_scrollPosY);
-		int boxHeight = m_config.boxHeight;
-		int processStepY = boxHeight + 1;
-		float scaleX = 50.0f*m_zoomValue*m_config.horizontalScaleValue;
-
-		RECT progressRect = clientRect;
-		progressRect.left += m_progressRectLeft;
-		progressRect.bottom -= 30;
-
-		{
-			int boxSide = 8;
-			int boxStride = boxSide + 2;
-			int top = 5;
-			int bottom = top + boxSide;
-			int left = progressRect.right - 7 - boxSide;
-			int right = progressRect.right - 7;
-			for (int i = VisualizerFlag_Count - 1; i >= 0; --i)
-			{
-				if (pos.x >= left && pos.x <= right && pos.y >= top && pos.y <= bottom)
-				{
-					outResult.buttonSelected = i;
-					return;
-				}
-				left -= boxStride;
-				right -= boxStride;
-			}
-		}
-
-		outResult.section = 0;
-
-		if (m_config.showTimeline && !m_traceView.sessions.empty())
-		{
-			int timelineTop = GetTimelineTop(clientRect);
-			if (pos.y >= timelineTop)// && pos.y < timelineTop + 40)
-			{
-				outResult.section = 3;
-				float timeScale = (m_config.horizontalScaleValue * m_zoomValue)*50.0f;
-				float startOffset = -(m_scrollPosX / timeScale);
-				outResult.timelineSelected = startOffset + float(pos.x - m_progressRectLeft) / timeScale;
-				return;
-			}
-		}
-
-		u64 lastStop = 0;
-
-		if (m_config.showProgress && m_traceView.progressProcessesTotal)
-			posY += m_activeFont.height + 2;
-
-		if (m_config.showStatus && !m_traceView.statusMap.empty())
-		{
-			u32 lastRow = ~0u;
-			u32 row = ~0u;
-			for (auto& kv : m_traceView.statusMap)
-			{
-				if (kv.second.text.empty())
-					continue;
-
-				row = u32(kv.first >> 32);
-				if (lastRow != ~0u && lastRow != row)
-					posY += m_activeFont.height + 2;
-				lastRow = row;
-
-				if (!kv.second.link.empty())
-				{
-					if (pos.y >= posY && pos.y < posY + m_activeFont.height && pos.x > 20 && pos.x < 80) // TODO: x is just hard coded to fit horde for now
-					{
-						outResult.hyperLink = kv.second.link;
-						return;
-					}
-				}
-			}
-			if (row != ~0u)
-				posY += m_activeFont.height + 2;
-			posY += 3;
-		}
-
-		if (pos.y < posY)
-			return;
-
-		outResult.section = 1;
-
-		if (m_config.showActiveProcessGraph)
-		{
-			if (pos.y > posY && pos.y < posY + GraphHeight)
-			{
-				float timeScale = (m_config.horizontalScaleValue * m_zoomValue)*50.0f;
-				float startOffset = -(m_scrollPosX / timeScale);
-				u64 selectedTimeMs = u64(1000.0f*(startOffset + float(pos.x - m_progressRectLeft) / timeScale));
-
-				u64 lastTime = m_traceView.activeProcessCounts.empty() ? 0 : m_traceView.activeProcessCounts[m_traceView.activeProcessCounts.size()-1].time;
-				if (selectedTimeMs < TimeToMs(lastTime))
-				{
-					u16 count = 0;
-					for (auto& entry : m_traceView.activeProcessCounts)
-					{
-						count = entry.count;
-						if (TimeToMs(entry.time) > selectedTimeMs)
-							break;
-					}
-
-					outResult.activeProcessGraphSelected = true;
-					outResult.activeProcessCount = count;
-				}
-			}
-			posY += GraphHeight;
-		}
-
-		TraceView::ProcessLocation& outLocation = outResult.processLocation;
-
-		if (m_config.showActiveProcesses && !m_trace.m_activeProcesses.empty())
-		{
-			PaintActiveProcesses(context, posY, clientRect, [&](TraceView::ProcessLocation& processLocation, u32 boxHeight, bool firstWithHeight)
-				{
-					if (pos.y < posY || pos.y > posY + int(boxHeight))
-						return;
-					auto& session = m_trace.GetSession(m_traceView, processLocation.sessionIndex);
-					TraceView::Process& process = session.processors[processLocation.processorIndex].processes[processLocation.processIndex];
-
-					int posX = int(m_scrollPosX) + progressRect.left;
-					u64 stop = process.stop;
-					bool done = stop != ~u64(0);
-					if (!done)
-						stop = playTime;
-					int left = posX + int(float(TimeToS(process.start)) * scaleX);
-					int right = posX + int(float(TimeToS(stop)) * scaleX) - 1;
-
-					if (pos.x >= left && pos.x <= right)
-					{
-						outLocation.sessionIndex = processLocation.sessionIndex;
-						outLocation.processorIndex = processLocation.processorIndex;
-						outLocation.processIndex = processLocation.processIndex;
-						outResult.processSelected = true;
-						return;
-					}
-				});
-			if (outResult.processSelected)
-				return;
-		}
-
-		if (pos.y < posY)
-			return;
-
-		outResult.section = 2;
-
-		SessionRec sortedSessions[1024];
-		Populate(sortedSessions, m_traceView, m_config.SortActiveRemoteSessions);
-
-		for (u64 sessionIt = 0, sessionEnd = m_traceView.sessions.size(); sessionIt != sessionEnd; ++sessionIt)
-		{
-			bool isFirst = sessionIt == 0;
-			auto& session = *sortedSessions[sessionIt].session;
-			bool hasUpdates = !session.updates.empty();
-			if (!isFirst)
-			{
-				if (!hasUpdates && session.processors.empty())
-					continue;
-				if (!m_config.showFinishedProcesses && session.disconnectTime != ~u64(0))
-					continue;
-			}
-
-			u32 sessionIndex = sortedSessions[sessionIt].index;
-			if (!isFirst)
-				posY += 3;
-
-			if (m_config.showTitleBars)
-			{
-				if (pos.y >= posY && pos.y < posY + m_sessionStepY)
-				{
-					if (pos.x < int(session.fullNameWidth) + 5)
-					{
-						outResult.sessionSelectedIndex = sessionIndex;
-						return;
-					}
-				}
-
-				posY += m_sessionStepY;
-			}
-
-			bool showGraph = m_config.showNetworkStats || m_config.showCpuMemStats || m_config.showDriveStats;
-			if (showGraph && !session.updates.empty())
-			{
-				int posX = int(m_scrollPosX) + progressRect.left;
-
-				if (pos.y >= posY && pos.y < posY + GraphHeight && pos.x >= posX)
-				{
-					bool loop = true;
-					u32 reconnectIndex = 0;
-					while (loop)
-					{
-						u64 i = 0;
-						u64 e = 0;
-
-						if (reconnectIndex)
-							i = session.reconnectIndices[reconnectIndex - 1];
-
-						if (reconnectIndex < session.reconnectIndices.size())
-							e = session.reconnectIndices[reconnectIndex];
-						else
-						{
-							e = session.updates.size();
-							loop = false;
-						}
-
-						int prevX = -1;
-						int prevCenterX = -1;
-						for (; i!=e; ++i)
-						{
-							u64 updateTime = session.updates[i];
-							auto updateSend = session.networkSend[i];
-							auto updateRecv = session.networkRecv[i];
-
-							int x = posX + int(float(TimeToS(updateTime)) * scaleX);
-							int centerX = prevX != -1 ? (prevX + (x - prevX)/2) : x;
-
-							if (prevCenterX != -1 && pos.x >= prevCenterX && pos.x < centerX)
-							{
-								u64 prevTime = 0;
-								u64 prevSend = 0;
-								u64 prevRecv = 0;
-								if (i > 0)
-								{
-									prevTime = session.updates[i - 1];
-									prevSend = Min(session.networkSend[i - 1], updateSend);
-									prevRecv = Min(session.networkRecv[i - 1], updateRecv);
-								}
-
-								double duration = TimeToS(updateTime - prevTime);
-								outResult.stats.recvBytes = updateRecv;
-								outResult.stats.sendBytes = updateSend;
-								outResult.stats.recvBytesPerSecond = u64(float(updateRecv - prevRecv) / duration);
-								outResult.stats.sendBytesPerSecond = u64(float(updateSend - prevSend) / duration);
-								outResult.stats.ping = session.ping[i];
-								outResult.stats.memAvail = session.memAvail[i];
-								outResult.stats.cpuLoad = session.cpuLoad[i];
-								outResult.stats.memTotal = session.memTotal;
-								outResult.stats.connectionCount = session.connectionCount[i];
-								outResult.statsSelected = true;
-
-								for (auto& pair : session.drives)
-								{
-									auto& updateBusy = pair.second.busyPercent[i];
-									u64 updateRead = pair.second.readBytes[i];
-									u64 updateWrite = pair.second.writeBytes[i];
-									auto& statsDrive = outResult.stats.drives[pair.first];
-									statsDrive.busyPercent = updateBusy;
-									statsDrive.readPerSecond = u64(float(updateRead) / duration);
-									statsDrive.writePerSecond = u64(float(updateWrite) / duration);
-								}
-								return;
-							}
-
-							prevCenterX = centerX;
-							prevX = x;
-						}
-						++reconnectIndex;
-					}
-					posY += GraphHeight;
-				}
-
-				posY += GraphHeight;
-			}
-
-			if (m_config.showDetailedData)
-			{
-				auto drawText = [&](const StringBufferBase& text, RECT& rect, u32* outWidth)
-					{
-						if (pos.x >= rect.left && pos.x < rect.right && pos.y >= rect.top && pos.y < rect.bottom && text.StartsWith(TC("Fetched Files")))
-							outResult.fetchedFilesSelected = sessionIndex;
-					};
-				PaintDetailedStats(posY, progressRect, session, sessionIt != 0, playTime, drawText);
-			}
-
-			if (m_config.showProcessBars)
-			{
-				u32 processorIndex = 0;
-				for (auto& processor : session.processors)
-				{
-					bool drawProcessorIndex = m_config.showFinishedProcesses;
-
-					if (pos.y < progressRect.bottom && posY + processStepY >= progressRect.top && posY <= progressRect.bottom && pos.y >= posY-1 && pos.y < posY-1 + processStepY)
-					{
-						u32 processIndex = 0;
-						int posX = int(m_scrollPosX) + progressRect.left;
-						for (auto& process : processor.processes)
-						{
-							int left = posX + int(float(TimeToS(process.start)) * scaleX);
-
-							auto pig = MakeGuard([&]() { ++processIndex; });
-
-							if (left >= progressRect.right)
-								continue;
-
-							if (left < progressRect.left)
-								left = progressRect.left;
-
-							u64 stopTime = process.stop;
-							bool done = stopTime != ~u64(0);
-							if (!done)
-								stopTime = playTime;
-							else if (!m_config.showFinishedProcesses)
-								continue;
-
-							drawProcessorIndex = true;
-
-							RECT rect;
-							rect.left = left;
-							rect.right = posX + int(float(TimeToS(stopTime)) * scaleX);
-
-							if (rect.right <= progressRect.left)
-								continue;
-
-							if (!m_filterString.empty() && !Contains(process.description.c_str(), m_filterString.c_str()) && !Contains(process.breadcrumbs.c_str(), m_filterString.c_str()))
-								continue;
-
-							rect.right = Max(int(rect.right), left + 1);
-							rect.top = posY;
-							rect.bottom = posY + int(float(18) * m_zoomValue);
-
-							if (pos.x >= rect.left && pos.x <= rect.right)
-							{
-								outLocation.sessionIndex = sessionIndex;
-								outLocation.processorIndex = processorIndex;
-								outLocation.processIndex = processIndex;
-								outResult.processSelected = true;
-								return;
-							}
-						}
-					}
-
-					if (!processor.processes.empty())
-						lastStop = Max(lastStop, processor.processes.rbegin()->stop);
-
-					if (drawProcessorIndex)
-						posY += processStepY;
-
-					++processorIndex;
-				}
-			}
-			else
-			{
-				for (auto& processor : session.processors)
-					if (!processor.processes.empty())
-						lastStop = Max(lastStop, processor.processes.rbegin()->stop);
-			}
-
-			if (m_config.showWorkers && isFirst)
-			{
-				int trackIndex = 0;
-				for (auto& workTrack : m_traceView.workTracks)
-				{
-					if (pos.y < progressRect.bottom && posY + processStepY >= progressRect.top && posY <= progressRect.bottom && pos.y >= posY-1 && pos.y < posY-1 + processStepY)
-					{
-						u32 workIndex = 0;
-						int posX = int(m_scrollPosX) + progressRect.left;
-						for (auto& work : workTrack.records)
-						{
-							auto inc = MakeGuard([&](){ ++workIndex; });
-
-							int left = posX + int(float(TimeToS(work.start)) * scaleX);
-
-							if (left >= progressRect.right)
-								continue;
-
-							if (!m_filterString.empty())
-							{
-								bool keep = Contains(work.description, m_filterString.c_str());
-								if (!keep)
-									for (auto& en : work.entries)
-										keep |= Contains(en.text, m_filterString.c_str());
-								if (!keep)
-									continue;
-							}
-
-							if (left < progressRect.left)
-								left = progressRect.left;
-
-							u64 stopTime = work.stop;
-							bool done = stopTime != ~u64(0);
-							if (!done)
-								stopTime = playTime;
-
-							RECT rect;
-							rect.left = left;
-							rect.right = posX + int(float(TimeToS(stopTime)) * scaleX);
-
-							if (rect.right <= progressRect.left)
-								continue;
-
-							rect.right = Max(int(rect.right), left + 1);
-							rect.top = posY;
-							rect.bottom = posY + int(float(18) * m_zoomValue);
-
-							if (pos.x >= rect.left && pos.x <= rect.right)
-							{
-								outResult.workTrack = trackIndex;
-								outResult.workIndex = workIndex;
-								outResult.workSelected = true;
-								return;
-							}
-						}
-					}
-					++trackIndex;
-					posY += processStepY;
-				}
-			}
-		}
-
-		m_contentWidth = m_progressRectLeft + Max(0, int(TimeToS((lastStop != 0 && lastStop != ~u64(0)) ? lastStop : playTime) * scaleX));
-		m_contentHeight = posY - int(m_scrollPosY) + processStepY + 14;
-	}
-
-	void VisualizerWin::WriteProcessStats(Logger& out, const TraceView::Process& process)
-	{
-		bool hasExited = process.stop != ~u64(0);
-		out.Info(L"  %ls", process.description.c_str());
-		out.Info(L"  ProcessId: %u", process.id);
-		out.Info(L"  Start:     %ls", TimeToText(process.start, true).str);
-		if (hasExited)
-			out.Info(L"  Duration:  %ls", TimeToText(process.stop - process.start, true).str);
-		if (hasExited && process.exitCode != 0)
-			out.Info(L"  ExitCode:  %u", process.exitCode);
-
-		if (process.stop != ~u64(0))
-		{
-			out.Info(L"");
-
-			BinaryReader reader(process.stats.data(), 0, process.stats.size());
-			ProcessStats processStats;
-			SessionStats sessionStats;
-			StorageStats storageStats;
-			KernelStats kernelStats;
-
-			processStats.Read(reader, m_traceView.version);
-			if (reader.GetLeft())
-			{
-				if (process.isRemote || (m_traceView.version >= 36 && !process.isReuse))
-					sessionStats.Read(reader, m_traceView.version);
-				storageStats.Read(reader, m_traceView.version);
-				kernelStats.Read(reader, m_traceView.version);
-			}
-
-			out.Info(L"  ----------- Detours stats -----------");
-			processStats.Print(out, m_traceView.frequency);
-
-			if (!sessionStats.IsEmpty())
-			{
-				out.Info(L"");
-				out.Info(L"  ----------- Session stats -----------");
-				sessionStats.Print(out, m_traceView.frequency);
-			}
-
-			if (!storageStats.IsEmpty())
-			{
-				out.Info(L"");
-				out.Info(L"  ----------- Storage stats -----------");
-				storageStats.Print(out, m_traceView.frequency);
-			}
-
-			if (!kernelStats.IsEmpty())
-			{
-				out.Info(L"");
-				out.Info(L"  ----------- Kernel stats ------------");
-				kernelStats.Print(out, false, m_traceView.frequency);
-			}
-
-			PrintCacheWriteStats(out, process.id);
-		}
-	}
-
-	void VisualizerWin::WriteWorkStats(Logger& out, const TraceView::WorkRecord& record)
-	{
-		out.Info(L"  %ls", record.description);
-		out.Info(L"  Start:     %ls", TimeToText(record.start, true).str);
-		if (record.stop != ~u64(0))
-			out.Info(L"  Duration:  %ls", TimeToText(record.stop - record.start, true).str);
-		for (auto& e : record.entries)
-			out.Info(L"   %ls (%ls)", e.text, TimeToText(e.time - record.start).str);
-	}
-
-	void VisualizerWin::CopyTextToClipboard(const TString& str)
-	{
-		if (!OpenClipboard(m_hwnd))
-			return;
-		if (auto hglbCopy = GlobalAlloc(GMEM_MOVEABLE, (str.size() + 1) * sizeof(TCHAR)))
-		{
-			if (auto lptstrCopy = GlobalLock(hglbCopy))
-			{
-				memcpy(lptstrCopy, str.data(), (str.size() + 1) * sizeof(TCHAR));
-				GlobalUnlock(hglbCopy);
-				EmptyClipboard();
-				SetClipboardData(CF_UNICODETEXT, hglbCopy);
-			}
-		}
-		CloseClipboard();
-	}
-
-	void VisualizerWin::UnselectAndRedraw()
-	{
-		if (Unselect() || m_config.showCursorLine)
-			Redraw(false);
-	}
-
-	bool VisualizerWin::UpdateAutoscroll()
-	{
-		if (!m_autoScroll)
-			return false;
-
-		u64 playTime = GetPlayTime();
-
-
-		RECT rect;
-		GetClientRect(m_hwnd, &rect);
-		if (rect.right == 0)
-			return false;
-
-		float timeS = TimeToS(playTime);
-
-		if (m_config.AutoScaleHorizontal)
-		{
-			m_scrollPosX = 0;
-			timeS = Max(timeS, 20.0f/m_zoomValue);
-			m_config.horizontalScaleValue = Max(float(rect.right - m_progressRectLeft - 2)/(m_zoomValue*timeS*50.0f), 0.001f);
-			return true;
-		}
-		else
-		{
-			float oldScrollPosX = m_scrollPosX;
-			m_scrollPosX = Min(0.0f, (float)rect.right - timeS*50.0f*m_config.horizontalScaleValue*m_zoomValue - float(m_progressRectLeft));
-			return oldScrollPosX != m_scrollPosX;
+			RECT clientRect;
+			GetClientRect(m_hwnd, &clientRect);
+			PaintContextWin context(*this, NULL, NULL, clientRect);
+			HitTest(res, context, pos);
 		}
-	}
-
-	bool VisualizerWin::UpdateSelection()
-	{
-		if (!m_mouseOverWindow || m_dragToScrollCounter > 0)
-			return false;
-		POINT pos;
-		GetCursorPos(&pos);
-		ScreenToClient(m_hwnd, &pos);
-
-		HitTestResult res;
-		HitTest(res, pos);
-
 		m_activeSection = res.section;
 
 		if (res.processSelected == m_processSelected && res.processLocation == m_processSelectedLocation &&
@@ -3517,17 +1203,17 @@ namespace uba
 		DwmSetWindowAttribute(m_hwnd, attribute, &useDarkMode, sizeof(useDarkMode));
 	}
 
-	bool VisualizerWin::ActiveProcessesShouldFillHeight()
+	bool Visualizer::ActiveProcessesShouldFillHeight()
 	{
 		return !m_config.showDetailedData && !m_config.showTitleBars && !m_config.showCpuMemStats && !m_config.showNetworkStats && !m_config.showDriveStats && !m_config.showProcessBars;
 	}
 
-	StringBuffer<128> VisualizerWin::GetWorldTime(u64 time)
+	StringBuffer<128> Visualizer::GetWorldTime(u64 time)
 	{
 		return GetWorldTime(TimeToS(time));
 	}
 
-	StringBuffer<128> VisualizerWin::GetWorldTime(float seconds)
+	StringBuffer<128> Visualizer::GetWorldTime(float seconds)
 	{
 		return StringBuffer<128>(DateToText(float(m_traceView.traceSystemStartTimeUs / 1'000'000) + seconds).str);
 	}
@@ -3642,7 +1328,7 @@ namespace uba
 					HBITMAP nullBmp = CreateCompatibleBitmap(memDC, 1, 1);
 					HBITMAP oldBmp = (HBITMAP)SelectObject(textDC, nullBmp);
 
-					PaintContext context { *this, m_hwnd, memDC, rect, textDC };
+					PaintContextWin context(*this, m_hwnd, memDC, rect, textDC);
 					PaintAll(context, rect);
 
 					SelectObject(textDC, oldBmp);
@@ -3923,8 +1609,12 @@ namespace uba
 			{
 				bool* values = &m_config.showProgress;
 				values[m_buttonSelected] = !values[m_buttonSelected];
+
+				RECT clientRect;
+				GetClientRect(m_hwnd, &clientRect);
+				PaintContextWin context(*this, NULL, NULL, clientRect);
 				HitTestResult res;
-				HitTest(res, { -1, -1 });
+				HitTest(res, context, { -1, -1 });
 				UpdateScrollbars(true);
 				Redraw(false);
 			}
@@ -3971,11 +1661,13 @@ namespace uba
 						m_pauseTime = 0;
 					}
 
-					HitTestResult res;
-					HitTest(res, { -1, -1 });
-
 					RECT r;
 					GetClientRect(hWnd, &r);
+
+					HitTestResult res;
+					PaintContextWin context(*this, NULL, NULL, r);
+					HitTest(res, context, { -1, -1 });
+
 					m_scrollPosX = Min(Max(m_scrollPosX, float(r.right - m_contentWidth)), 0.0f);
 					m_scrollPosY = Min(0.0f, Max(m_scrollPosY, float(r.bottom - m_contentHeight)));
 
diff --git a/Engine/Source/Programs/UnrealBuildAccelerator/Visualizer/Private/Windows/UbaVisualizerWin.h b/Engine/Source/Programs/UnrealBuildAccelerator/Visualizer/Private/Windows/UbaVisualizerWin.h
--- a/Engine/Source/Programs/UnrealBuildAccelerator/Visualizer/Private/Windows/UbaVisualizerWin.h
+++ b/Engine/Source/Programs/UnrealBuildAccelerator/Visualizer/Private/Windows/UbaVisualizerWin.h
@@ -28,20 +28,8 @@ namespace uba
 		bool Unselect();
 		void Reset();
 		void PaintClient(const Function<void(HDC hdc, HDC memDC, RECT& clientRect)>& paintFunc);
-
-		struct PaintContext;
-		struct DrawTextLogger;
-
-		void PaintAll(PaintContext& context, const RECT& clientRect);
-		void PaintActiveProcesses(PaintContext& context, int& posY, const RECT& clientRect, const Function<void(TraceView::ProcessLocation&, u32, bool)>& drawProcess);
-		void PaintProcessRect(PaintContext& context, TraceView::Process& process, RECT rect, const RECT& progressRect, bool selected, bool writingBitmap);
-		void PaintTimeline(PaintContext& context, const RECT& clientRect);
-		using DrawTextFunc = Function<void(const StringBufferBase& text, RECT& rect, u32* outWidth)>;
-		void PaintDetailedStats(int& posY, const RECT& progressRect, TraceView::Session& session, bool isRemote, u64 playTime, const DrawTextFunc& drawTextFunc);
-
-		int GetTimelineHeight();
-		int GetTimelineTop(const RECT& clientRect);
-		void HitTest(HitTestResult& outResult, const POINT& pos);
+		virtual void PaintEmptyScreen(PaintContext& context) override;
+		virtual void FinishIncompatible() override;
 
 		void WriteProcessStats(Logger& out, const TraceView::Process& process);
 		void WriteWorkStats(Logger& out, const TraceView::WorkRecord& record);
@@ -58,87 +46,21 @@ namespace uba
 		void StopDragToScroll();
 		void SaveSettings();
 		void DirtyBitmaps(bool full);
-		void PrintCacheWriteStats(Logger& logger, u32 processId);
 
-		struct Font
-		{
-			HFONT handle = 0;
-			HFONT handleUnderlined = 0;
-			int height = 0;
-			int offset = 0;
-		};
-
-		void UpdateFont(Font& font, int height, bool createUnderline);
 		void UpdateDefaultFont();
 		void UpdateProcessFont();
 		void ChangeFontSize(int offset);
 		void Redraw(bool now);
 		void UpdateTheme();
-		bool ActiveProcessesShouldFillHeight();
-		StringBuffer<128> GetWorldTime(u64 time);
-		StringBuffer<128> GetWorldTime(float seconds);
 
 		StringBuffer<256> m_namedTrace;
 		StringBuffer<256> m_fileName;
 
 		u64 m_lastPaintTimeMs;
 
-		struct ProcessBrushes
-		{
-			HBRUSH inProgress = 0;
-			HBRUSH inProgressRace = 0;
-			HBRUSH success = 0;
-			HBRUSH error = 0;
-			HBRUSH returned = 0;
-			HBRUSH recv = 0;
-			HBRUSH send = 0;
-			HBRUSH cacheFetchTest = 0;
-			HBRUSH cacheFetchDownload = 0;
-			HBRUSH warning = 0;
-		};
-
-		ProcessBrushes m_processBrushes[2]; // Non-selected and selected
-
 		Atomic<bool> m_looping = false;
 		HWND m_hwnd = 0;
 		HWND m_parentHwnd = 0;
-		COLORREF m_textColor = {};
-		COLORREF m_textWarningColor = {};
-		COLORREF m_textErrorColor = {};
-		COLORREF m_sendColor = {};
-		COLORREF m_recvColor = {};
-		COLORREF m_cpuColor = {};
-		COLORREF m_memColor = {};
-		COLORREF m_driveColor = {};
-		COLORREF m_activeProcColor = {};
-		HBRUSH m_backgroundBrush = 0;
-		HBRUSH m_tooltipBackgroundBrush = 0;
-		HPEN m_textPen = 0;
-		HPEN m_separatorPen = 0;
-		HPEN m_sendPen = 0;
-		HPEN m_recvPen = 0;
-		HPEN m_cpuPen = 0;
-		HPEN m_memPen = 0;
-		HPEN m_drivePen = 0;
-		HPEN m_activeProcPen = 0;
-		HPEN m_processUpdatePen = 0;
-		HPEN m_checkboxPen = 0;
-		int m_sessionStepY = 0;
-
-		Font m_defaultFont;
-		Font m_processFont;
-		Font m_timelineFont;
-		Font m_popupFont;
-
-		int m_processFontOffsetY = 0;
-
-		Font m_activeProcessFont[32];
-		u32 m_activeProcessCountHistory[5];
-		u32 m_activeProcessCountHistoryIterator = 0;
-
-		Font m_activeFont;
-
-		int m_progressRectLeft = 30;
 
 		NetworkClient* m_client = nullptr;
 		Event m_clientDisconnect;
@@ -147,15 +69,8 @@ namespace uba
 		StringBuffer<256> m_newTraceName;
 		Event m_listenTimeout;
 
-		int m_contentWidth = 0;
-		int m_contentHeight = 0;
-
 		int m_contentWidthWhenThumbTrack = 0;
 
-		float m_scrollPosX = 0;
-		float m_scrollPosY = 0;
-		float m_zoomValue = 0.5f;
-		float m_horizontalScaleValue = 0.5f;
 		bool m_autoScroll = true;
 
 		bool m_isInPaint = false;
@@ -164,25 +79,7 @@ namespace uba
 		HBITMAP m_lastBitmap = 0;
 		int m_lastBitmapOffset = BitmapCacheHeight;
 
-		u32 m_activeSection = ~0u;
-		TraceView::ProcessLocation m_processSelectedLocation;
-		bool m_processSelected = false;
-		u32 m_sessionSelectedIndex = ~0u;
-		bool m_statsSelected = false;
-		bool m_activeProcessGraphSelected = false;
-		u64 m_activeProcessCount = 0u;
-		Stats m_stats;
-		u32 m_buttonSelected = ~0u;
-		float m_timelineSelected = 0;
-		u32 m_fetchedFilesSelected = ~0u;
-		TString m_hyperLinkSelected;
-
-		bool m_workSelected = false;
-		u32 m_workTrack = ~0u;
-		u32 m_workIndex = ~0u;
-		
 		bool m_usingNamed = false;
-		bool m_mouseOverWindow = false;
 		bool m_showPopup = false;
 		bool m_locked = false;
 		bool m_handlingTimer = false;
@@ -200,16 +97,14 @@ namespace uba
 		bool m_horizontalScrollBarEnabled = true;
 		bool m_verticalScrollBarEnabled = true;
 
-		TString m_filterString;
-
 		Thread m_thread;
 
-		UnorderedMap<Color, HBRUSH> m_coloredBrushes;
-
 		void PostNewTrace(u32 replay, bool paused);
 		void PostNewTitle(const StringView& title);
 		void PostQuit();
 		LRESULT WinProc(HWND hWnd, UINT Msg, WPARAM wParam, LPARAM lParam);
 		static LRESULT CALLBACK StaticWinProc(HWND hWnd, UINT Msg, WPARAM wParam, LPARAM lParam);
+
+		struct PaintContextWin;
 	};
 }
