[UBA]
* Fixed so UbaVisualizer compiles on linux

diff --git a/Engine/Source/Programs/UnrealBuildAccelerator/Visualizer/Private/UbaVisualizer.cpp b/Engine/Source/Programs/UnrealBuildAccelerator/Visualizer/Private/UbaVisualizer.cpp
--- a/Engine/Source/Programs/UnrealBuildAccelerator/Visualizer/Private/UbaVisualizer.cpp
+++ b/Engine/Source/Programs/UnrealBuildAccelerator/Visualizer/Private/UbaVisualizer.cpp
@@ -1,6 +1,7 @@
 // Copyright Epic Games, Inc. All Rights Reserved.
 
 #include "UbaVisualizer.h"
+#include "UbaLogger.h"
 #include <algorithm>
 
 namespace uba
@@ -23,7 +24,7 @@ namespace uba
 
 		virtual void BeginScope() override {}
 		virtual void EndScope() override {}
-		virtual void Log(LogEntryType type, const wchar_t* str, u32 strLen) override
+		virtual void Log(LogEntryType type, const tchar* str, u32 strLen) override
 		{
 			u32 textWidth = context.GetTextSize(StringView(str, strLen)).cx;
 
@@ -57,6 +58,7 @@ namespace uba
 
 			RECT clientRect = context.GetClientRect();
 
+#if PLATFORM_WINDOWS
 			if (r.right > clientRect.right)
 				OffsetRect(&r, -width - 15, 0);
 			if (r.bottom > clientRect.bottom)
@@ -65,6 +67,7 @@ namespace uba
 				if (r.top < 0)
 					OffsetRect(&r, 0, -r.top);
 			}
+#endif
 
 			RECT fillRect = r;
 			fillRect.right += 2 + extraWidth;
@@ -162,6 +165,21 @@ namespace uba
 		return m_config.LockTimelineToBottom ? maxY : Min(posY, maxY);
 	}
 
+	bool Visualizer::ActiveProcessesShouldFillHeight()
+	{
+		return !m_config.showDetailedData && !m_config.showTitleBars && !m_config.showCpuMemStats && !m_config.showNetworkStats && !m_config.showDriveStats && !m_config.showProcessBars;
+	}
+
+	StringBuffer<128> Visualizer::GetWorldTime(u64 time)
+	{
+		return GetWorldTime(TimeToS(time));
+	}
+
+	StringBuffer<128> Visualizer::GetWorldTime(float seconds)
+	{
+		return StringBuffer<128>(DateToText(float(m_traceView.traceSystemStartTimeUs / 1'000'000) + seconds).str);
+	}
+
 	void Visualizer::HitTest(HitTestResult& outResult, PaintContext& context, const POINT& pos)
 	{
 		u64 playTime = GetPlayTime();
@@ -638,7 +656,7 @@ namespace uba
 				rect.bottom = posY + activeFont.height + 2;
 				context.SetTextColor(type == LogEntryType_Info ? m_textColor : (type == LogEntryType_Error ? m_textErrorColor : m_textWarningColor));
 
-				u32 align = 0;
+				u32 align = TA_LEFT;
 				if (centered)
 				{
 					posX += (endX - posX)/2;
@@ -663,7 +681,7 @@ namespace uba
 
 		if (m_config.showProgress && m_traceView.progressProcessesTotal)
 		{
-			drawIndentedText(AsView(L"Progress"), LogEntryType_Info, 1, false);
+			drawIndentedText(TCV("Progress"), LogEntryType_Info, 1, false);
 
 			Font& activeFont = context.GetActiveFont();
 
@@ -870,17 +888,17 @@ namespace uba
 						if (firstPar != -1 && *process.description.rbegin() == ')')
 						{
 							str.Append(process.description.c_str() + firstPar + 1).Resize(str.count-1);
-							str.Append(L"  ").Append(process.description.c_str(), firstPar);
+							str.Append(TCV("  ")).Append(process.description.c_str(), firstPar);
 						}
 						else
 							str.Append(process.description);
 
 						if (process.isRemote)
-							str.Append(L" [").Append(session.name).Append(']');
+							str.Append(TCV(" [")).Append(session.name).Append(']');
 						else if (process.type == TraceView::ProcessType_CacheFetchTest || process.type == TraceView::ProcessType_CacheFetchDownload)
-							str.Append(L" [cache]");
+							str.Append(TCV(" [cache]"));
 						if (boxRect.left < 0)
-							str.Appendf(L"   %s", TimeToText(playTime - process.start, true).str);
+							str.Appendf(TC("   %s"), TimeToText(playTime - process.start, true).str);
 						//drawIndentedText(str, LogEntryType_Info, 1, true);
 						drawStatusText(str, LogEntryType_Info, Max(int(boxRect.left+1), 1), boxRect.right, false);
 					}
@@ -933,19 +951,19 @@ namespace uba
 						//u64 memAvail = session.updates.back().memAvail;
 						//float cpuLoad = session.updates.back().cpuLoad;
 
-						//text.Appendf(L" - Cpu: %.1f%%", cpuLoad * 100.0f);
+						//text.Appendf(TC(" - Cpu: %.1f%%"), cpuLoad * 100.0f);
 						//if (memAvail)
-						//	text.Appendf(L" Mem: %ls/%ls", BytesToText(session.memTotal - memAvail).str, BytesToText(session.memTotal).str);
+						//	text.Appendf(TC(" Mem: %ls/%ls"), BytesToText(session.memTotal - memAvail).str, BytesToText(session.memTotal).str);
 						if (ping)
-							text.Appendf(L" Ping: %ls", TimeToText(ping, false, m_traceView.frequency).str);
+							text.Appendf(TC(" Ping: %ls"), TimeToText(ping, false, m_traceView.frequency).str);
 						if (!session.notification.empty())
-							text.Append(L" - ").Append(session.notification);
+							text.Append(TCV(" - ")).Append(session.notification);
 					}
 					else if (!isFirst)
 					{
-						text.Append(L" - Disconnected");
+						text.Append(TCV(" - Disconnected"));
 						if (!session.notification.empty())
-							text.Append(L" (").Append(session.notification).Append(')');
+							text.Append(TCV(" (")).Append(session.notification).Append(')');
 					}
 
 					bool selected = m_sessionSelectedIndex == processLocation.sessionIndex;
@@ -1496,7 +1514,7 @@ namespace uba
 				context.SetActiveFont(m_popupFont);
 				DrawTextLogger logger(context, m_popupFont.height, m_tooltipBackgroundBrush);
 
-				logger.Info(L"%s (%s)", b.data, wt.data);
+				logger.Info(TC("%s (%s)"), b.data, wt.data);
 				logger.DrawAtPos(left + 4, timelineTop - 20);
 			}
 		}
@@ -1599,32 +1617,32 @@ namespace uba
 			logger.AddWidth(3);
 
 			logger.AddSpace(2);
-			logger.Info(L"  %ls", process.description.c_str());
+			logger.Info(TC("  %ls"), process.description.c_str());
 			if (process.type != TraceView::ProcessType_Task)
 			{
-				logger.Info(L"  Host:        %ls", m_processSelectedLocation.sessionIndex == 0 ? TC("local") : m_traceView.GetSession(m_processSelectedLocation).name.c_str());
-				logger.Info(L"  ProcessId:  %6u", process.id);
+				logger.Info(TC("  Host:        %ls"), m_processSelectedLocation.sessionIndex == 0 ? TC("local") : m_traceView.GetSession(m_processSelectedLocation).name.c_str());
+				logger.Info(TC("  ProcessId:  %6u"), process.id);
 			}
-			logger.Info(L"  Start:     %7ls (%ls)", TimeToText(process.start, true).str, GetWorldTime(process.start).data);
-			logger.Info(L"  Duration:  %7ls", TimeToText(duration, true).str);
+			logger.Info(TC("  Start:     %7ls (%ls)"), TimeToText(process.start, true).str, GetWorldTime(process.start).data);
+			logger.Info(TC("  Duration:  %7ls"), TimeToText(duration, true).str);
 			if (!process.returnedReason.empty())
-				logger.Info(L"  Returned: %7s", process.returnedReason.data());
+				logger.Info(TC("  Returned: %7s"), process.returnedReason.data());
 			if (hasExited && process.exitCode != 0)
 			{
 				if (process.exitCode == ProcessCancelExitCode)
-					logger.Info(L"  ExitCode: Cancelled");
+					logger.Info(TC("  ExitCode: Cancelled"));
 				else if (process.exitCode >= 0xC0000000u)
-					logger.Info(L"  ExitCode: 0x%x", process.exitCode);
+					logger.Info(TC("  ExitCode: 0x%x"), process.exitCode);
 				else
-					logger.Info(L"  ExitCode: %7u", process.exitCode);
+					logger.Info(TC("  ExitCode: %7u"), process.exitCode);
 			}
 
 			const auto& breadcrumbs = process.breadcrumbs;
 			if (!breadcrumbs.empty())
 			{
 				constexpr TString::size_type maxLineLen = 50;
-				logger.Info(L"");
-				logger.Info(L"  ------------ Breadcrumbs ------------");
+				logger.Info(TC(""));
+				logger.Info(TC("  ------------ Breadcrumbs ------------"));
 				for (TString::size_type lineStart = 0, lineEnd = 0; lineEnd < breadcrumbs.size(); lineStart = lineEnd + 1)
 				{
 					// Log each individual line
@@ -1638,13 +1656,13 @@ namespace uba
 						{
 							const TString::size_type maxSectionLen = sectionStart == 0 ? maxLineLen : maxLineLen - 2;
 							sectionEnd = std::min<TString::size_type>(sectionEnd + maxSectionLen, line.size());
-							TString section = (sectionStart == 0 ? L"  " : L"    ") + line.substr(sectionStart, sectionEnd - sectionStart);
+							TString section = (sectionStart == 0 ? TC("  ") : TC("    ")) + line.substr(sectionStart, sectionEnd - sectionStart);
 							logger.Info(section.c_str());
 						}
 					}
 					else
 					{
-						line = L"  " + line;
+						line = TC("  ") + line;
 						logger.Info(line.c_str());
 					}
 				}
@@ -1662,9 +1680,9 @@ namespace uba
 				if (process.type == TraceView::ProcessType_CacheFetchTest || process.type == TraceView::ProcessType_CacheFetchDownload)
 				{
 					if (!process.returnedReason.empty())
-						logger.Info(L"  Cache:       Miss");
+						logger.Info(TC("  Cache:       Miss"));
 					else
-						logger.Info(L"  Cache:        Hit");
+						logger.Info(TC("  Cache:        Hit"));
 					cacheStats.Read(reader, m_traceView.version);
 					if (reader.GetLeft())
 					{
@@ -1687,42 +1705,42 @@ namespace uba
 
 				if (processStats.hostTotalTime)
 				{
-					logger.Info(L"");
-					logger.Info(L"  ----------- Detours stats -----------");
+					logger.Info(TC(""));
+					logger.Info(TC("  ----------- Detours stats -----------"));
 					processStats.Print(logger, m_traceView.frequency);
 				}
 				else if (processStats.peakMemory)
 				{
-					logger.Info(L"");
-					logger.Info(L"  ----------- Process stats -----------");
+					logger.Info(TC(""));
+					logger.Info(TC("  ----------- Process stats -----------"));
 					processStats.Print(logger, m_traceView.frequency);
 				}
 
 				if (!sessionStats.IsEmpty())
 				{
-					logger.Info(L"");
-					logger.Info(L"  ----------- Session stats -----------");
+					logger.Info(TC(""));
+					logger.Info(TC("  ----------- Session stats -----------"));
 					sessionStats.Print(logger, m_traceView.frequency);
 				}
 
 				if (!cacheStats.IsEmpty())
 				{
-					logger.Info(L"");
-					logger.Info(L"  ------------ Cache stats ------------");
+					logger.Info(TC(""));
+					logger.Info(TC("  ------------ Cache stats ------------"));
 					cacheStats.Print(logger, m_traceView.frequency);
 				}
 
 				if (!storageStats.IsEmpty())
 				{
-					logger.Info(L"");
-					logger.Info(L"  ----------- Storage stats -----------");
+					logger.Info(TC(""));
+					logger.Info(TC("  ----------- Storage stats -----------"));
 					storageStats.Print(logger, m_traceView.frequency);
 				}
 
 				if (!kernelStats.IsEmpty())
 				{
-					logger.Info(L"");
-					logger.Info(L"  ----------- Kernel stats ------------");
+					logger.Info(TC(""));
+					logger.Info(TC("  ----------- Kernel stats ------------"));
 					kernelStats.Print(logger, false, m_traceView.frequency);
 				}
 
@@ -1730,8 +1748,8 @@ namespace uba
 
 				if (!logLines.empty())
 				{
-					logger.Info(L"");
-					logger.Info(L"  ---------------- Log ----------------");
+					logger.Info(TC(""));
+					logger.Info(TC("  ---------------- Log ----------------"));
 					logger.AddTextOffset(14);
 					for (auto& line : logLines)
 						logger.Log(LogEntryType_Info, line);
@@ -1752,9 +1770,9 @@ namespace uba
 			DrawTextLogger logger(context, m_popupFont.height, m_tooltipBackgroundBrush);
 
 			logger.AddSpace();
-			logger.Info(L"  %ls", selectedWork.description);
-			logger.Info(L"  Start:     %ls (%ls)", TimeToText(selectedWork.start, true).str, GetWorldTime(selectedWork.start).data);
-			logger.Info(L"  Duration:  %ls", TimeToText(duration, true).str);
+			logger.Info(TC("  %ls"), selectedWork.description);
+			logger.Info(TC("  Start:     %ls (%ls)"), TimeToText(selectedWork.start, true).str, GetWorldTime(selectedWork.start).data);
+			logger.Info(TC("  Duration:  %ls"), TimeToText(duration, true).str);
 			if (!selectedWork.entries.empty())
 				logger.AddSpace();
 			for (u32 i=0,e=u32(selectedWork.entries.size()); i!=e; ++i)
@@ -1765,7 +1783,7 @@ namespace uba
 				{
 					u64 start = entry.startTime ? entry.startTime : entry.time;
 					if (TimeToMs(start - selectedWork.start) > 1)
-						logger.Info(L"   Start (%ls)", TimeToText(start - selectedWork.start).str);
+						logger.Info(TC("   Start (%ls)"), TimeToText(start - selectedWork.start).str);
 				}
 
 				if (entry.startTime)
@@ -1784,9 +1802,9 @@ namespace uba
 				if (!time && selectedWork.stop != ~0u)
 					time = selectedWork.stop - entry.time;
 				if (entry.count == 1)
-					logger.Info(L"%s  %ls (%ls)", (entry.startTime ? L" " : L""), entry.text, TimeToText(time).str);
+					logger.Info(TC("%s  %ls (%ls)"), (entry.startTime ? TC(" ") : TC("")), entry.text, TimeToText(time).str);
 				else
-					logger.Info(L"%s  %ls (%ls %u)", (entry.startTime ? L" " : L""), entry.text, TimeToText(time).str, entry.count);
+					logger.Info(TC("%s  %ls (%ls %u)"), (entry.startTime ? TC(" ") : TC("")), entry.text, TimeToText(time).str, entry.count);
 			}
 			logger.AddSpace();
 			logger.DrawAtCursor();
@@ -1803,9 +1821,9 @@ namespace uba
 			if (summary.empty())
 			{
 				if (m_traceView.finished)
-					logger.Info(L" Session summary was never received");
+					logger.Info(TC(" Session summary was never received"));
 				else
-					logger.Info(L" Session summary not available until session is done");
+					logger.Info(TC(" Session summary not available until session is done"));
 			}
 
 			logger.AddTextOffset(-10); // Remove spaces in the front
@@ -1823,7 +1841,7 @@ namespace uba
 			for (auto& pair : session.drives)
 			{
 				auto& d = pair.second;
-				logger.Info(L"  %c: Read: %s (%u) Write: %s (%u)", pair.first, BytesToText(d.totalReadBytes).str, d.totalReadCount, BytesToText(d.totalWriteBytes).str, d.totalWriteCount);
+				logger.Info(TC("  %c: Read: %s (%u) Write: %s (%u)"), pair.first, BytesToText(d.totalReadBytes).str, d.totalReadCount, BytesToText(d.totalWriteBytes).str, d.totalWriteCount);
 			}
 
 			logger.AddSpace(3);
@@ -1834,22 +1852,22 @@ namespace uba
 			context.SetActiveFont(m_popupFont);
 			DrawTextLogger logger(context, m_popupFont.height, m_tooltipBackgroundBrush);
 			logger.AddSpace(3);
-			logger.SetColor(m_cpuColor).Info(L"  Cpu: %.1f%%", m_stats.cpuLoad * 100.0f);
-			logger.SetColor(m_memColor).Info(L"  Mem: %ls/%ls", BytesToText(m_stats.memTotal - m_stats.memAvail).str, BytesToText(m_stats.memTotal).str);
+			logger.SetColor(m_cpuColor).Info(TC("  Cpu: %.1f%%"), m_stats.cpuLoad * 100.0f);
+			logger.SetColor(m_memColor).Info(TC("  Mem: %ls/%ls"), BytesToText(m_stats.memTotal - m_stats.memAvail).str, BytesToText(m_stats.memTotal).str);
 			if (m_stats.recvBytes || m_stats.sendBytes)
 			{
-				logger.SetColor(m_recvColor).Info(L"  Recv: %lsit/s", BytesToText(m_stats.recvBytesPerSecond*8).str);
-				logger.SetColor(m_sendColor).Info(L"  Send: %lsit/s", BytesToText(m_stats.sendBytesPerSecond*8).str);
+				logger.SetColor(m_recvColor).Info(TC("  Recv: %lsit/s"), BytesToText(m_stats.recvBytesPerSecond*8).str);
+				logger.SetColor(m_sendColor).Info(TC("  Send: %lsit/s"), BytesToText(m_stats.sendBytesPerSecond*8).str);
 			}
 			logger.SetColor(m_textColor);
 
 			if (m_stats.ping)
-				logger.Info(L"  Ping: %ls", TimeToText(m_stats.ping, false, m_traceView.frequency).str);
+				logger.Info(TC("  Ping: %ls"), TimeToText(m_stats.ping, false, m_traceView.frequency).str);
 			if (m_stats.connectionCount)
-				logger.Info(L"  Conn: %u", u32(m_stats.connectionCount));
+				logger.Info(TC("  Conn: %u"), u32(m_stats.connectionCount));
 
 			for (auto& pair : m_stats.drives)
-				logger.SetColor(m_driveColor).Info(L"  %c: %u%% R:%s/s W:%s/s", pair.first, pair.second.busyPercent, BytesToText(pair.second.readPerSecond).str, BytesToText(pair.second.writePerSecond).str);
+				logger.SetColor(m_driveColor).Info(TC("  %c: %u%% R:%s/s W:%s/s"), pair.first, pair.second.busyPercent, BytesToText(pair.second.readPerSecond).str, BytesToText(pair.second.writePerSecond).str);
 
 			logger.AddSpace(3);
 			logger.DrawAtCursor();
@@ -1860,7 +1878,7 @@ namespace uba
 			DrawTextLogger logger(context, m_popupFont.height, m_tooltipBackgroundBrush);
 			logger.AddSpace(3);
 			logger.AddSpace(3);
-			logger.SetColor(m_activeProcColor).Info(L"  Active Processes: %u", m_activeProcessCount);
+			logger.SetColor(m_activeProcColor).Info(TC("  Active Processes: %u"), m_activeProcessCount);
 			logger.DrawAtCursor();
 		}
 		else if (m_buttonSelected != ~0u)
@@ -1874,7 +1892,7 @@ namespace uba
 
 			context.SetActiveFont(m_popupFont);
 			DrawTextLogger logger(context, m_popupFont.height, m_tooltipBackgroundBrush);
-			logger.Info(L"%ls %ls", L"Show", tooltip[m_buttonSelected]);
+			logger.Info(TC("%ls %ls"), TC("Show"), tooltip[m_buttonSelected]);
 			logger.DrawAtCursor();
 		}
 		else if (m_fetchedFilesSelected != ~0u)
@@ -1898,13 +1916,13 @@ namespace uba
 				{
 				if (logger.rect.left + colWidth >= r.right)
 				{
-				logger.Info(L"...");
+				logger.Info(TC("..."));
 				break;
 				}
 				logger.rect.top = r.top;
 				logger.rect.left += colWidth;
 				}
-				logger.Info(L"%s", f.hint.c_str());
+				logger.Info(TC("%s"), f.hint.c_str());
 				}
 				logger.DrawAtCursor();
 				*/
@@ -2106,7 +2124,7 @@ namespace uba
 		progressRect.left += m_progressRectLeft;
 
 		context.SetPen(m_textPen);
-		Vector<DWORD> lines;
+		Vector<u32> lines;
 		Vector<POINT> points;
 
 		while (true)
@@ -2127,12 +2145,12 @@ namespace uba
 					seconds -= min * 60;
 					if (!seconds)
 					{
-						buffer.Appendf(L"%um", min);
+						buffer.Appendf(TC("%um"), min);
 						lineBottom += 4;
 					}
 				}
 				if (!number || seconds)
-					buffer.Appendf(L"%u", seconds);
+					buffer.Appendf(TC("%u"), seconds);
 
 				if (shouldDraw)
 				{
@@ -2172,7 +2190,7 @@ namespace uba
 		u32 maxTextWidth = 0;
 
 		RECT textRect;
-		auto drawText = [&](const wchar_t* format, ...)
+		auto drawText = [&](const tchar* format, ...)
 			{
 				textRect.top = posY;
 				textRect.bottom = posY + 20;
@@ -2192,18 +2210,18 @@ namespace uba
 
 		if (isRemote)
 		{
-			drawText(L"Finished Processes: %u", session.processExitedCount);
-			drawText(L"Active Processes: %u", session.processActiveCount);
-			drawText(L"ClientId: %u  TcpCount: %u", session.clientUid.data1, session.connectionCount.empty() ? 1 : session.connectionCount.back());
+			drawText(TC("Finished Processes: %u"), session.processExitedCount);
+			drawText(TC("Active Processes: %u"), session.processActiveCount);
+			drawText(TC("ClientId: %u  TcpCount: %u"), session.clientUid.data1, session.connectionCount.empty() ? 1 : session.connectionCount.back());
 
 			if (session.disconnectTime == ~u64(0))
 			{
 				if (session.proxyCreated)
-					drawText(L"Proxy(HOSTED): %ls", session.proxyName.c_str());
+					drawText(TC("Proxy(HOSTED): %ls"), session.proxyName.c_str());
 				else if (!session.proxyName.empty())
-					drawText(L"Proxy: %ls", session.proxyName.c_str());
+					drawText(TC("Proxy: %ls"), session.proxyName.c_str());
 				else
-					drawText(L"Proxy: None");
+					drawText(TC("Proxy: None"));
 			}
 
 			bool hasFileDetails = !session.fetchedFiles.empty() || !session.storedFiles.empty();
@@ -2227,17 +2245,17 @@ namespace uba
 					sendPerS = u64(float(updateSend - session.prevSend) / duration);
 					recvPerS = u64(float(updateRecv - session.prevRecv) / duration);
 				}
-				drawText(L"Recv: %ls (%sit/s)", BytesToText(updateRecv), BytesToText(recvPerS*8));
-				drawText(L"Send: %ls (%sit/s)", BytesToText(updateSend), BytesToText(sendPerS*8));
+				drawText(TC("Recv: %ls (%sit/s)"), BytesToText(updateRecv), BytesToText(recvPerS*8));
+				drawText(TC("Send: %ls (%sit/s)"), BytesToText(updateSend), BytesToText(sendPerS*8));
 			}
 
 			int fileWidth = 700;
 
-			auto drawFiles = [&](const wchar_t* fileType, Vector<TraceView::FileTransfer>& files, u64 count, u64 bytes, u64 activeCount, u32& maxVisibleFiles)
+			auto drawFiles = [&](const tchar* fileType, Vector<TraceView::FileTransfer>& files, u64 count, u64 bytes, u64 activeCount, u32& maxVisibleFiles)
 				{
 					textRect.left = posX;
 					textRect.right = posX + fileWidth;
-					drawText(L"%ls Files: %u (%u) %s", fileType, u32(count), u32(activeCount), BytesToText(bytes));
+					drawText(TC("%ls Files: %u (%u) %s"), fileType, u32(count), u32(activeCount), BytesToText(bytes));
 					u32 fileCount = 0;
 					for (auto rit = files.rbegin(), rend = files.rend(); rit != rend; ++rit)
 					{
@@ -2249,11 +2267,11 @@ namespace uba
 							time = playTime - file.start;
 
 						if ((((u8*)&file.key)[19] & IsProxyMask) != 0)
-							drawText(L"%s (proxy) %s", file.hint.c_str(), TimeToText(time, true).str);
+							drawText(TC("%s (proxy) %s"), file.hint.c_str(), TimeToText(time, true).str);
 						else if (file.size == 0)
-							drawText(L"%s (calc) %s", file.hint.c_str(), TimeToText(time, true).str);
+							drawText(TC("%s (calc) %s"), file.hint.c_str(), TimeToText(time, true).str);
 						else
-							drawText(L"%s (%s) %s", file.hint.c_str(), BytesToText(file.size).str, TimeToText(time, true).str);
+							drawText(TC("%s (%s) %s"), file.hint.c_str(), BytesToText(file.size).str, TimeToText(time, true).str);
 
 						if (fileCount++ > 5)
 							break;
@@ -2272,21 +2290,21 @@ namespace uba
 				posY = startPosY;
 				posX += maxTextWidth + 15;
 			}
-			drawFiles(L"Fetched", fetchedFiles, session.fetchedFilesCount, session.fetchedFilesBytes, session.fetchedFilesActive.size(), session.maxVisibleFiles);
+			drawFiles(TC("Fetched"), fetchedFiles, session.fetchedFilesCount, session.fetchedFilesBytes, session.fetchedFilesActive.size(), session.maxVisibleFiles);
 			if (hasFileDetails)
 			{
 				posY = startPosY;
 				posX += fileWidth;
 			}
-			drawFiles(L"Stored", session.storedFiles, session.storedFilesCount, session.storedFilesBytes, session.storedFilesActive.size(), session.maxVisibleFiles);
+			drawFiles(TC("Stored"), session.storedFiles, session.storedFilesCount, session.storedFilesBytes, session.storedFilesActive.size(), session.maxVisibleFiles);
 		}
 		else
 		{
 			if (m_traceView.totalProcessActiveCount || m_traceView.totalProcessExitedCount)
 			{
-				drawText(L"Finished Processes: %u (local: %u)", m_traceView.totalProcessExitedCount, session.processExitedCount);
-				drawText(L"Active Processes: %u (local: %u)", m_traceView.totalProcessActiveCount, session.processActiveCount);
-				drawText(L"Active Helpers: %u", Max(1u, m_traceView.activeSessionCount) - 1);
+				drawText(TC("Finished Processes: %u (local: %u)"), m_traceView.totalProcessExitedCount, session.processExitedCount);
+				drawText(TC("Active Processes: %u (local: %u)"), m_traceView.totalProcessActiveCount, session.processActiveCount);
+				drawText(TC("Active Helpers: %u"), Max(1u, m_traceView.activeSessionCount) - 1);
 			}
 
 			if (session.highestSendPerS > 0.0f || session.highestRecvPerS > 0.0f)
@@ -2305,8 +2323,8 @@ namespace uba
 						recvPerS = u64(float(updateRecv - session.prevRecv) / duration);
 					}
 
-					drawText(L"Recv: %ls (%sit/s)", BytesToText(updateRecv), BytesToText(recvPerS*8));
-					drawText(L"Send: %ls (%sit/s)", BytesToText(updateSend), BytesToText(sendPerS*8));
+					drawText(TC("Recv: %ls (%sit/s)"), BytesToText(updateRecv), BytesToText(recvPerS*8));
+					drawText(TC("Send: %ls (%sit/s)"), BytesToText(updateSend), BytesToText(sendPerS*8));
 				}
 			}
 
@@ -2330,7 +2348,7 @@ namespace uba
 						writePerS = u64(float(drive.writeBytes.back()) / duration);
 					}
 
-					drawText(L"%c: Rd %ls (%s/s) Wr %ls (%s/s) ", pair.first, BytesToText(drive.totalReadBytes).str, BytesToText(readPerS).str, BytesToText(drive.totalWriteBytes).str, BytesToText(writePerS).str);
+					drawText(TC("%c: Rd %ls (%s/s) Wr %ls (%s/s) "), pair.first, BytesToText(drive.totalReadBytes).str, BytesToText(readPerS).str, BytesToText(drive.totalWriteBytes).str, BytesToText(writePerS).str);
 				}
 			}
 		}
@@ -2344,12 +2362,12 @@ namespace uba
 		if (findIt == m_traceView.cacheWrites.end())
 			return;
 		TraceView::CacheWrite& write = findIt->second;
-		logger.Info(L"");
-		logger.Info(L"  -------- Cache write stats ----------");
+		logger.Info(TC(""));
+		logger.Info(TC("  -------- Cache write stats ----------"));
 		if (write.end)
 		{
-			logger.Info(L"  Duration                    %9s", TimeToText(write.end - write.start).str);
-			logger.Info(L"  Success                     %9s", write.success ? L"true" : L"false");
+			logger.Info(TC("  Duration                    %9s"), TimeToText(write.end - write.start).str);
+			logger.Info(TC("  Success                     %9s"), write.success ? TC("true") : TC("false"));
 
 			BinaryReader reader(write.stats.data(), 0, write.stats.size());
 			CacheSendStats stats;
@@ -2364,7 +2382,7 @@ namespace uba
 		}
 		else
 		{
-			logger.Info(L"  In progress...");
+			logger.Info(TC("  In progress..."));
 		}
 	}
 
@@ -2376,5 +2394,10 @@ namespace uba
 	{
 	}
 
+	HBRUSH Visualizer::CreateSolidBrush(COLORREF color)
+	{
+		return 0;
+	}
+
 	////////////////////////////////////////////////////////////////////////////////////////////////////
 }
diff --git a/Engine/Source/Programs/UnrealBuildAccelerator/Visualizer/Private/UbaVisualizer.h b/Engine/Source/Programs/UnrealBuildAccelerator/Visualizer/Private/UbaVisualizer.h
--- a/Engine/Source/Programs/UnrealBuildAccelerator/Visualizer/Private/UbaVisualizer.h
+++ b/Engine/Source/Programs/UnrealBuildAccelerator/Visualizer/Private/UbaVisualizer.h
@@ -9,6 +9,25 @@ namespace uba
 {
 	////////////////////////////////////////////////////////////////////////////////////////////////////
 
+	#if !PLATFORM_WINDOWS
+	struct POINT { int x; int y; };
+	struct SIZE { int cx; int cy; };
+	struct RECT { int left; int top; int right; int bottom; };
+	using COLORREF = u32;
+	using HFONT = void*;
+	using HPEN = void*;
+	using HBRUSH = void*;
+	#define TRANSPARENT  1
+	#define OPAQUE       2
+	#define TA_LEFT      0
+	#define TA_RIGHT     2
+	#define TA_CENTER    6
+	#define ETO_CLIPPED  0x0004
+	#define RGB(r,g,b)          ((COLORREF)(((u8)(r)|((u16)((u8)(g))<<8))|(((u32)(u8)(b))<<16)))
+	#endif
+
+	////////////////////////////////////////////////////////////////////////////////////////////////////
+
 	class Visualizer
 	{
 	public:
@@ -39,7 +58,7 @@ namespace uba
 			virtual void SetPen(HPEN pen) = 0;
 			virtual void SetFont(HFONT font) = 0;
 			virtual u32 SetTextAlign(u32 align) = 0;
-			virtual void ExtTextOutW(int x, int y, UINT options, CONST RECT * lprect, StringView str) = 0;
+			virtual void ExtTextOutW(int x, int y, u32 options, const RECT* lprect, StringView str) = 0;
 			virtual void SetActiveFont(const Font& font) = 0;
 			virtual Font& GetActiveFont() = 0;
 			virtual void SetBrush(HBRUSH b) = 0;
@@ -47,7 +66,7 @@ namespace uba
 			virtual void Rectangle(int left, int top, int right, int bottom) = 0;
 			virtual void Line(int fromX, int fromY, int toX, int toY) = 0;
 			virtual void Polyline(const Vector<POINT>& polyLine) = 0;
-			virtual void PolyPolyline(const POINT* points, const DWORD* lines, u32 lineCount) = 0;
+			virtual void PolyPolyline(const POINT* points, const u32* lines, u32 lineCount) = 0;
 			virtual void DrawCenteredText(const StringBuffer<>* rows, u32 rowCount, bool drawBackground) = 0;
 			virtual SIZE GetTextSize(StringView str) = 0;
 			virtual void DrawCached(TraceView::BitmapCache& process, RECT rect, int offsetY, int textHeight, RECT clippingRect, const Function<void(PaintContext& textCacheContext)>& func) = 0;
@@ -65,8 +84,10 @@ namespace uba
 		using DrawTextFunc = Function<void(const StringBufferBase& text, RECT& rect, u32* outWidth)>;
 		void PaintDetailedStats(PaintContext& context, int& posY, const RECT& progressRect, TraceView::Session& session, bool isRemote, u64 playTime, const DrawTextFunc& drawTextFunc);
 		void PrintCacheWriteStats(Logger& logger, u32 processId);
+
 		virtual void PaintEmptyScreen(PaintContext& context);
 		virtual void FinishIncompatible();
+		virtual HBRUSH CreateSolidBrush(COLORREF color);
 
 		void UpdateFont(Font& font, int height, bool createUnderline);
 
diff --git a/Engine/Source/Programs/UnrealBuildAccelerator/Visualizer/Private/Windows/UbaVisualizerWin.cpp b/Engine/Source/Programs/UnrealBuildAccelerator/Visualizer/Private/Windows/UbaVisualizerWin.cpp
--- a/Engine/Source/Programs/UnrealBuildAccelerator/Visualizer/Private/Windows/UbaVisualizerWin.cpp
+++ b/Engine/Source/Programs/UnrealBuildAccelerator/Visualizer/Private/Windows/UbaVisualizerWin.cpp
@@ -54,7 +54,7 @@ namespace uba
 		virtual void SetPen(HPEN pen) override { ::SelectObject(hdc, pen); }
 		virtual void SetFont(HFONT font) override { ::SelectObject(hdc, font); }
 		virtual u32 SetTextAlign(u32 align) override { return ::SetTextAlign(hdc, align); }
-		virtual void ExtTextOutW(int x, int y, UINT options, CONST RECT * lprect, StringView str) override { ::ExtTextOutW(hdc, x, y, options, lprect, str.data, str.count, NULL); }
+		virtual void ExtTextOutW(int x, int y, u32 options, CONST RECT * lprect, StringView str) override { ::ExtTextOutW(hdc, x, y, options, lprect, str.data, str.count, NULL); }
 		
 		virtual void SetActiveFont(const Font& font) override
 		{
@@ -99,9 +99,9 @@ namespace uba
 			::Polyline(hdc, polyLine.data(), int(polyLine.size()));
 		}
 
-		virtual void PolyPolyline(const POINT* points, const DWORD* lines, u32 lineCount) override
+		virtual void PolyPolyline(const POINT* points, const u32* lines, u32 lineCount) override
 		{
-			::PolyPolyline(hdc, points, lines, int(lineCount));
+			::PolyPolyline(hdc, points, (const DWORD*)lines, int(lineCount));
 		}
 
 		virtual void DrawCenteredText(const StringBuffer<>* rows, u32 rowCount, bool drawBackground) override
@@ -966,6 +966,11 @@ namespace uba
 		SetWindowTextW(m_hwnd, GetTitlePrefix(title).Appendf(L" (Listening for new sessions on channel '%s')", m_listenChannel.data).data);
 	}
 
+	HBRUSH VisualizerWin::CreateSolidBrush(COLORREF color)
+	{
+		return ::CreateSolidBrush(color);
+	}
+
 	void VisualizerWin::WriteProcessStats(Logger& out, const TraceView::Process& process)
 	{
 		bool hasExited = process.stop != ~u64(0);
@@ -1203,21 +1208,6 @@ namespace uba
 		DwmSetWindowAttribute(m_hwnd, attribute, &useDarkMode, sizeof(useDarkMode));
 	}
 
-	bool Visualizer::ActiveProcessesShouldFillHeight()
-	{
-		return !m_config.showDetailedData && !m_config.showTitleBars && !m_config.showCpuMemStats && !m_config.showNetworkStats && !m_config.showDriveStats && !m_config.showProcessBars;
-	}
-
-	StringBuffer<128> Visualizer::GetWorldTime(u64 time)
-	{
-		return GetWorldTime(TimeToS(time));
-	}
-
-	StringBuffer<128> Visualizer::GetWorldTime(float seconds)
-	{
-		return StringBuffer<128>(DateToText(float(m_traceView.traceSystemStartTimeUs / 1'000'000) + seconds).str);
-	}
-
 	void VisualizerWin::PostNewTrace(u32 replay, bool paused)
 	{
 		KillTimer(m_hwnd, 0);
diff --git a/Engine/Source/Programs/UnrealBuildAccelerator/Visualizer/Private/Windows/UbaVisualizerWin.h b/Engine/Source/Programs/UnrealBuildAccelerator/Visualizer/Private/Windows/UbaVisualizerWin.h
--- a/Engine/Source/Programs/UnrealBuildAccelerator/Visualizer/Private/Windows/UbaVisualizerWin.h
+++ b/Engine/Source/Programs/UnrealBuildAccelerator/Visualizer/Private/Windows/UbaVisualizerWin.h
@@ -30,6 +30,7 @@ namespace uba
 		void PaintClient(const Function<void(HDC hdc, HDC memDC, RECT& clientRect)>& paintFunc);
 		virtual void PaintEmptyScreen(PaintContext& context) override;
 		virtual void FinishIncompatible() override;
+		virtual HBRUSH CreateSolidBrush(COLORREF color) override;
 
 		void WriteProcessStats(Logger& out, const TraceView::Process& process);
 		void WriteWorkStats(Logger& out, const TraceView::WorkRecord& record);
