[UBA]
* More visualizer code cleanup

diff --git a/Engine/Source/Programs/UnrealBuildAccelerator/Visualizer/Private/UbaVisualizer.cpp b/Engine/Source/Programs/UnrealBuildAccelerator/Visualizer/Private/UbaVisualizer.cpp
--- a/Engine/Source/Programs/UnrealBuildAccelerator/Visualizer/Private/UbaVisualizer.cpp
+++ b/Engine/Source/Programs/UnrealBuildAccelerator/Visualizer/Private/UbaVisualizer.cpp
@@ -10,6 +10,10 @@ namespace uba
 
 	u64 ConvertTime(const TraceView& view, u64 time);
 
+	#if !PLATFORM_WINDOWS
+	void OffsetRect(RECT* r, int x, int y) { r->left -= x; r->right -= x; r->top -= y; r->bottom -= y; }
+	#endif
+
 	////////////////////////////////////////////////////////////////////////////////////////////////////
 
 	struct Visualizer::DrawTextLogger : public Logger
@@ -58,7 +62,6 @@ namespace uba
 
 			RECT clientRect = context.GetClientRect();
 
-#if PLATFORM_WINDOWS
 			if (r.right > clientRect.right)
 				OffsetRect(&r, -width - 15, 0);
 			if (r.bottom > clientRect.bottom)
@@ -67,7 +70,6 @@ namespace uba
 				if (r.top < 0)
 					OffsetRect(&r, 0, -r.top);
 			}
-#endif
 
 			RECT fillRect = r;
 			fillRect.right += 2 + extraWidth;
@@ -79,7 +81,7 @@ namespace uba
 				tr.left += line.left;
 				tr.top += line.top;
 				context.SetTextColor(line.color);
-				context.ExtTextOutW(tr.left, tr.top, 0, NULL, line.str);
+				context.Text(tr.left, tr.top, nullptr, line.str);
 			}
 		}
 
@@ -185,7 +187,7 @@ namespace uba
 		u64 playTime = GetPlayTime();
 
 		RECT clientRect = context.GetClientRect();
-		context.SetActiveFont(m_defaultFont);
+		context.SetFont(m_defaultFont);
 
 		int posY = int(m_scrollPosY);
 		int boxHeight = m_config.boxHeight;
@@ -232,7 +234,7 @@ namespace uba
 
 		u64 lastStop = 0;
 
-		Font& activeFont = context.GetActiveFont();
+		Font& activeFont = context.GetFont();
 
 		if (m_config.showProgress && m_traceView.progressProcessesTotal)
 			posY += activeFont.height + 2;
@@ -644,11 +646,11 @@ namespace uba
 
 		u64 lastStop = 0;
 
-		context.SetActiveFont(m_defaultFont);
+		context.SetFont(m_defaultFont);
 
 		auto drawStatusText = [&](StringView text, LogEntryType type, int posX, int endX, bool moveY, bool underlined = false, bool centered = false)
 			{
-				Font& activeFont = context.GetActiveFont();
+				Font& activeFont = context.GetFont();
 				RECT rect;
 				rect.left = posX;
 				rect.right = endX;
@@ -663,10 +665,10 @@ namespace uba
 					align = context.SetTextAlign(TA_CENTER);
 				}
 				if (underlined)
-					context.SetFont(activeFont.handleUnderlined);
-				context.ExtTextOutW(posX, posY, ETO_CLIPPED, &rect, text);
+					context.SetFont(activeFont, true);
+				context.Text(posX, posY, &rect, text);
 				if (underlined)
-					context.SetFont(activeFont.handle);
+					context.SetFont(activeFont);
 				if (centered)
 					context.SetTextAlign(align);
 				if (moveY)
@@ -683,7 +685,7 @@ namespace uba
 		{
 			drawIndentedText(TCV("Progress"), LogEntryType_Info, 1, false);
 
-			Font& activeFont = context.GetActiveFont();
+			Font& activeFont = context.GetFont();
 
 			float progress = float(m_traceView.progressProcessesDone) / float(m_traceView.progressProcessesTotal);
 			u32 start = 3 + 6*activeFont.height;
@@ -730,7 +732,7 @@ namespace uba
 		{
 			u32 lastRow = ~0u;
 			u32 row = ~0u;
-			Font& activeFont = context.GetActiveFont();
+			Font& activeFont = context.GetFont();
 			for (auto& kv : m_traceView.statusMap)
 			{
 				auto& status = kv.second;
@@ -881,7 +883,7 @@ namespace uba
 							m_activeProcessFont[fontIndex].offset += 1;
 						}
 						if (firstWithHeight)
-							context.SetActiveFont(m_activeProcessFont[fontIndex]);
+							context.SetFont(m_activeProcessFont[fontIndex]);
 
 						StringBuffer<> str;
 						auto firstPar = -1;//process.description.find_first_of('(');
@@ -982,7 +984,7 @@ namespace uba
 					SIZE textSize = context.GetTextSize(text);
 					if (textSize.cx)
 						session.fullNameWidth = u32(textSize.cx);
-					context.ExtTextOutW(5, posY+2, ETO_CLIPPED, &rect, text);
+					context.Text(5, posY+2, &rect, text);
 					if (selected)
 						context.SetBkMode(TRANSPARENT);
 				}
@@ -1120,10 +1122,10 @@ namespace uba
 						if (rect.bottom > progressRect.bottom)
 						{
 							rect.bottom = progressRect.bottom;
-							context.ExtTextOutW(rect.left, rect.top, ETO_CLIPPED, &rect, text);
+							context.Text(rect.left, rect.top, &rect, text);
 						}
 						else
-							context.ExtTextOutW(rect.left, rect.top, 0, NULL, text);
+							context.Text(rect.left, rect.top, nullptr, text);
 						if (selected)
 							context.SetBkMode(TRANSPARENT);
 						if (outWidth)
@@ -1133,7 +1135,7 @@ namespace uba
 				PaintDetailedStats(context, posY, progressRect, session, isRemote, playTime, drawText);
 			}
 
-			context.SetActiveFont(m_processFont);
+			context.SetFont(m_processFont);
 
 			bool shouldDrawText = m_processFont.height > 4;
 
@@ -1222,12 +1224,12 @@ namespace uba
 											textCacheContext.SetTextColor(RGB(5, 60, 5));
 											++rect2.left;
 											++textY;
-											textCacheContext.ExtTextOutW(rect2.left, textY, ETO_CLIPPED, &rect2, process.description);
+											textCacheContext.Text(rect2.left, textY, &rect2, process.description);
 											--rect2.left;
 											--textY;
 										}
 										textCacheContext.SetTextColor(m_textColor);
-										textCacheContext.ExtTextOutW(rect2.left, textY, ETO_CLIPPED, &rect2, process.description);
+										textCacheContext.Text(rect2.left, textY, &rect2, process.description);
 
 										if (!selected)
 											process.bitmapDirty = false;
@@ -1245,7 +1247,7 @@ namespace uba
 
 							StringBuffer<> buf;
 							buf.AppendValue(u64(processLocation.processorIndex) + 1);
-							context.ExtTextOutW(5, posY + offsetY, ETO_CLIPPED, &rect, buf);
+							context.Text(5, posY + offsetY, &rect, buf);
 						}
 					}
 
@@ -1294,7 +1296,7 @@ namespace uba
 
 							StringBuffer<> buf;
 							buf.AppendValue(u64(trackIndex) + 1);
-							context.ExtTextOutW(5, posY + offsetY, ETO_CLIPPED, &rect, buf);
+							context.Text(5, posY + offsetY, &rect, buf);
 						}
 
 						int lastDrawnRight = 0;
@@ -1395,7 +1397,7 @@ namespace uba
 									{
 										RECT rect2{ 0,int(work.bitmapOffset), 256, int(work.bitmapOffset) + m_processFont.height };
 										textCacheContext.FillRect(rect2, brush);
-										textCacheContext.ExtTextOutW(rect2.left, rect2.top, ETO_CLIPPED, &rect2, ToView(work.description));
+										textCacheContext.Text(rect2.left, rect2.top, &rect2, ToView(work.description));
 										if (!selected)
 											work.bitmapDirty = false;
 									});
@@ -1407,12 +1409,12 @@ namespace uba
 				}
 			}
 
-			context.SetActiveFont(m_defaultFont);
+			context.SetFont(m_defaultFont);
 		}
 
 		if (true)
 		{
-			Font& activeFont = context.GetActiveFont();
+			Font& activeFont = context.GetFont();
 			for (auto& kv : m_traceView.tables)
 			{
 				auto& table = kv.second;
@@ -1431,7 +1433,7 @@ namespace uba
 					if (int(x) >= progressRect.right)
 						break;
 
-					context.ExtTextOutW(x, y, ETO_CLIPPED, nullptr, ToView(item));
+					context.Text(x, y, nullptr, ToView(item));
 
 					++rowIndex;
 					if (rowIndex > table.rowMaxCount || int(y) >= progressRect.bottom)
@@ -1459,7 +1461,7 @@ namespace uba
 			str[0].Append(TCV("Box Filter: ")).Append(m_filterString);
 			str[1].Appendf(TC("Box Count: %u"), visibleBoxes);
 
-			context.SetActiveFont(m_popupFont);
+			context.SetFont(m_popupFont);
 			context.DrawCenteredText(str, 2, true);
 		}
 
@@ -1511,7 +1513,7 @@ namespace uba
 
 				StringBuffer<128> wt = GetWorldTime(timelineSelected);
 
-				context.SetActiveFont(m_popupFont);
+				context.SetFont(m_popupFont);
 				DrawTextLogger logger(context, m_popupFont.height, m_tooltipBackgroundBrush);
 
 				logger.Info(TC("%s (%s)"), b.data, wt.data);
@@ -1548,25 +1550,25 @@ namespace uba
 				{
 					context.SetTextColor(color);
 					left -= (context.GetTextSize(text).cx) + 5;
-					context.ExtTextOutW(left, top, 0, NULL, text);
+					context.Text(left, top, nullptr, text);
 				};
 
 			if (m_config.showDriveStats)
 			{
-				context.SetActiveFont(m_defaultFont);
+				context.SetFont(m_defaultFont);
 				drawText(TCV("DRV"), m_driveColor);
 			}
 
 			if (m_config.showNetworkStats)
 			{
-				context.SetActiveFont(m_defaultFont);
+				context.SetFont(m_defaultFont);
 				drawText(TCV("SND"), m_sendColor);
 				drawText(TCV("RCV"), m_recvColor);
 			}
 
 			if (m_config.showCpuMemStats)
 			{
-				context.SetActiveFont(m_defaultFont);
+				context.SetFont(m_defaultFont);
 				drawText(TCV("CPU"), m_cpuColor);
 				drawText(TCV("MEM"), m_memColor);
 			}
@@ -1610,7 +1612,7 @@ namespace uba
 				duration = playTime - process.start;
 			}
 
-			context.SetActiveFont(m_popupFont);
+			context.SetFont(m_popupFont);
 			DrawTextLogger logger(context, m_popupFont.height, m_tooltipBackgroundBrush);
 
 			logger.AddTextOffset(-10); // Remove spaces in the front
@@ -1766,7 +1768,7 @@ namespace uba
 			else
 				duration = playTime - selectedWork.start;
 
-			context.SetActiveFont(m_popupFont);
+			context.SetFont(m_popupFont);
 			DrawTextLogger logger(context, m_popupFont.height, m_tooltipBackgroundBrush);
 
 			logger.AddSpace();
@@ -1811,7 +1813,7 @@ namespace uba
 		}
 		else if (m_sessionSelectedIndex != ~0u)
 		{
-			context.SetActiveFont(m_popupFont);
+			context.SetFont(m_popupFont);
 			DrawTextLogger logger(context, m_popupFont.height, m_tooltipBackgroundBrush);
 			logger.AddWidth(3);
 			logger.AddSpace(2);
@@ -1849,7 +1851,7 @@ namespace uba
 		}
 		else if (m_statsSelected)
 		{
-			context.SetActiveFont(m_popupFont);
+			context.SetFont(m_popupFont);
 			DrawTextLogger logger(context, m_popupFont.height, m_tooltipBackgroundBrush);
 			logger.AddSpace(3);
 			logger.SetColor(m_cpuColor).Info(TC("  Cpu: %.1f%%"), m_stats.cpuLoad * 100.0f);
@@ -1874,7 +1876,7 @@ namespace uba
 		}
 		else if (m_activeProcessGraphSelected)
 		{
-			context.SetActiveFont(m_popupFont);
+			context.SetFont(m_popupFont);
 			DrawTextLogger logger(context, m_popupFont.height, m_tooltipBackgroundBrush);
 			logger.AddSpace(3);
 			logger.AddSpace(3);
@@ -1890,7 +1892,7 @@ namespace uba
 #undef UBA_VISUALIZER_FLAG
 			};
 
-			context.SetActiveFont(m_popupFont);
+			context.SetFont(m_popupFont);
 			DrawTextLogger logger(context, m_popupFont.height, m_tooltipBackgroundBrush);
 			logger.Info(TC("%ls %ls"), TC("Show"), tooltip[m_buttonSelected]);
 			logger.DrawAtCursor();
@@ -1906,7 +1908,7 @@ namespace uba
 				int width = colWidth * 2;
 				int height = Min(int(clientRect.bottom), int(fetchedFiles.size() * m_popupFont.height));
 
-				context.SetActiveFont(m_defaultFont);
+				context.SetFont(m_defaultFont);
 				DrawTextLogger logger(context, r, m_font.height, m_tooltipBackgroundBrush);
 				for (auto& f : fetchedFiles)
 				{
@@ -1994,7 +1996,7 @@ namespace uba
 
 	void Visualizer::PaintActiveProcesses(PaintContext& context, int& posY, const RECT& clientRect, const Function<void(TraceView::ProcessLocation&, u32, bool)>& drawProcess)
 	{
-		context.SetActiveFont(m_processFont);
+		context.SetFont(m_processFont);
 		int startPosY = posY;
 
 		Map<u64, TraceView::ProcessLocation*> activeProcesses;
@@ -2009,7 +2011,7 @@ namespace uba
 				++remoteCount;
 		}
 
-		Font& activeFont = context.GetActiveFont();
+		Font& activeFont = context.GetFont();
 
 		u32 maxHeight = clientRect.bottom;
 
@@ -2091,12 +2093,12 @@ namespace uba
 		else
 			posY += 3;
 
-		context.SetActiveFont(m_defaultFont);
+		context.SetFont(m_defaultFont);
 	}
 
 	void Visualizer::PaintTimeline(PaintContext& context, const RECT& clientRect)
 	{
-		context.SetActiveFont(m_timelineFont);
+		context.SetFont(m_timelineFont);
 		int top = GetTimelineTop(clientRect);
 		float timeScale = (m_config.horizontalScaleValue*m_zoomValue)*50.0f;
 		float startOffset = ((m_scrollPosX/timeScale) - float(int(m_scrollPosX/timeScale))) * timeScale;
@@ -2160,7 +2162,7 @@ namespace uba
 					textRect.bottom = textRect.top + m_timelineFont.height;
 					textRect.right = pos + s.cx/2;
 					textRect.left = pos - s.cx/2;
-					context.ExtTextOutW(textRect.left, textRect.top, 0, NULL, buffer);
+					context.Text(textRect.left, textRect.top, nullptr, buffer);
 				}
 			}
 			if (!(number % lineStepSize))
@@ -2183,7 +2185,7 @@ namespace uba
 
 	void Visualizer::PaintDetailedStats(PaintContext& context, int& posY, const RECT& progressRect, TraceView::Session& session, bool isRemote, u64 playTime, const DrawTextFunc& drawTextFunc)
 	{
-		int stepY = context.GetActiveFont().height;
+		int stepY = context.GetFont().height;
 		int startPosY = posY;
 		int posX = progressRect.left + 5;
 		int maxY = posY + stepY;
diff --git a/Engine/Source/Programs/UnrealBuildAccelerator/Visualizer/Private/UbaVisualizer.h b/Engine/Source/Programs/UnrealBuildAccelerator/Visualizer/Private/UbaVisualizer.h
--- a/Engine/Source/Programs/UnrealBuildAccelerator/Visualizer/Private/UbaVisualizer.h
+++ b/Engine/Source/Programs/UnrealBuildAccelerator/Visualizer/Private/UbaVisualizer.h
@@ -22,8 +22,7 @@ namespace uba
 	#define TA_LEFT      0
 	#define TA_RIGHT     2
 	#define TA_CENTER    6
-	#define ETO_CLIPPED  0x0004
-	#define RGB(r,g,b)          ((COLORREF)(((u8)(r)|((u16)((u8)(g))<<8))|(((u32)(u8)(b))<<16)))
+	#define RGB(r,g,b)          ((COLORREF)((u8(r)|(u16(u8(g))<<8))|(u32(u8(b))<<16)))
 	#endif
 
 	////////////////////////////////////////////////////////////////////////////////////////////////////
@@ -52,16 +51,16 @@ namespace uba
 		struct PaintContext
 		{
 			virtual RECT GetClientRect() = 0;
+
 			virtual void SetBkMode(int mode) = 0;
-			virtual void SetTextColor(COLORREF color) = 0;
-			virtual COLORREF GetTextColor() = 0;
 			virtual void SetPen(HPEN pen) = 0;
-			virtual void SetFont(HFONT font) = 0;
-			virtual u32 SetTextAlign(u32 align) = 0;
-			virtual void ExtTextOutW(int x, int y, u32 options, const RECT* lprect, StringView str) = 0;
-			virtual void SetActiveFont(const Font& font) = 0;
-			virtual Font& GetActiveFont() = 0;
 			virtual void SetBrush(HBRUSH b) = 0;
+			virtual void SetFont(const Font& font, bool underlined = false) = 0;
+			virtual void SetTextColor(COLORREF color) = 0;
+			virtual u32 SetTextAlign(u32 align) = 0;
+			virtual void Text(int x, int y, const RECT* lprect, StringView str) = 0;
+			virtual COLORREF GetTextColor() = 0;
+			virtual Font& GetFont() = 0;
 			virtual void FillRect(const RECT& r, HBRUSH b) = 0;
 			virtual void Rectangle(int left, int top, int right, int bottom) = 0;
 			virtual void Line(int fromX, int fromY, int toX, int toY) = 0;
@@ -99,7 +98,6 @@ namespace uba
 			VisualizerFlag_Count
 		};
 
-
 		struct Stats
 		{
 			u64 recvBytesPerSecond = 0;
diff --git a/Engine/Source/Programs/UnrealBuildAccelerator/Visualizer/Private/Windows/UbaVisualizerWin.cpp b/Engine/Source/Programs/UnrealBuildAccelerator/Visualizer/Private/Windows/UbaVisualizerWin.cpp
--- a/Engine/Source/Programs/UnrealBuildAccelerator/Visualizer/Private/Windows/UbaVisualizerWin.cpp
+++ b/Engine/Source/Programs/UnrealBuildAccelerator/Visualizer/Private/Windows/UbaVisualizerWin.cpp
@@ -52,17 +52,27 @@ namespace uba
 		virtual void SetTextColor(COLORREF color) override { ::SetTextColor(hdc, color); }
 		virtual COLORREF GetTextColor() override { return ::GetTextColor(hdc); }
 		virtual void SetPen(HPEN pen) override { ::SelectObject(hdc, pen); }
-		virtual void SetFont(HFONT font) override { ::SelectObject(hdc, font); }
 		virtual u32 SetTextAlign(u32 align) override { return ::SetTextAlign(hdc, align); }
-		virtual void ExtTextOutW(int x, int y, u32 options, CONST RECT * lprect, StringView str) override { ::ExtTextOutW(hdc, x, y, options, lprect, str.data, str.count, NULL); }
+		virtual void Text(int x, int y, const RECT* clipRect, StringView str) override { ::ExtTextOutW(hdc, x, y, clipRect ? ETO_CLIPPED : 0, clipRect, str.data, str.count, NULL); }
 		
-		virtual void SetActiveFont(const Font& font) override
+		virtual void SetFont(const Font& font, bool underlined) override
 		{
+			if (underlined)
+			{
+				if (font.handleUnderlined == activeFont.handleUnderlined)
+					return;
+				::SelectObject(hdc, font.handleUnderlined);
+			}
+			else
+			{
+				if (font.handle == activeFont.handle)
+					return;
+				::SelectObject(hdc, font.handle);
+			}
 			activeFont = font;
-			::SelectObject(hdc, font.handle);
 		}
 
-		virtual Font& GetActiveFont() override
+		virtual Font& GetFont() override
 		{
 			return activeFont;
 		}
@@ -123,7 +133,7 @@ namespace uba
 					InflateRect(&r, 4, 2);
 					FillRect(r, owner.m_tooltipBackgroundBrush);
 				}
-				ExtTextOutW(rect.left, rect.top, 0, NULL, rows[i]);
+				Text(rect.left, rect.top, nullptr, rows[i]);
 			}
 		}
 
@@ -956,7 +966,7 @@ namespace uba
 			else
 				str[0].Append(TCV("No trace active"));
 		}
-		context.SetActiveFont(m_popupFont);
+		context.SetFont(m_popupFont);
 		context.DrawCenteredText(str, 2, false);
 	}
 
