[UBA]
* Moved out Visualizer backend to its own file
* Created VisualizerBackendGDI and moved all cached bitmap stuff to that file

diff --git a/Engine/Source/Programs/UnrealBuildAccelerator/Visualizer/Private/UbaVisualizer.cpp b/Engine/Source/Programs/UnrealBuildAccelerator/Visualizer/Private/UbaVisualizer.cpp
--- a/Engine/Source/Programs/UnrealBuildAccelerator/Visualizer/Private/UbaVisualizer.cpp
+++ b/Engine/Source/Programs/UnrealBuildAccelerator/Visualizer/Private/UbaVisualizer.cpp
@@ -112,10 +112,11 @@ namespace uba
 
 	////////////////////////////////////////////////////////////////////////////////////////////////////
 
-	Visualizer::Visualizer(VisualizerConfig& config, Logger& logger)
+	Visualizer::Visualizer(VisualizerConfig& config, Logger& logger, VisualizerBackend& backend)
 	:	m_logger(logger)
 	,	m_config(config)
 	,	m_trace(logger)
+	,	m_backend(backend)
 	{
 	}
 
@@ -623,7 +624,7 @@ namespace uba
 				StringBuffer<> str;
 				str.Appendf(TC("Unsupported trace version %u (Versions supported are %u to %u)"), m_traceView.version, TraceReadCompatibilityVersion, TraceVersion);
 				context.SetTextColor(m_textWarningColor);
-				context.DrawCenteredText(&str, 1, false);
+				context.DrawCenteredText(&str, 1, nullptr);
 			}
 			else
 			{
@@ -1205,7 +1206,7 @@ namespace uba
 							int processWidth = rect.right - rect.left;
 							if (shouldDrawText && m_config.ShowProcessText && processWidth > 3)
 							{
-								context.DrawCached(process, rect, offsetY, textHeight, progressRect, [&](PaintContext& textCacheContext)
+								context.DrawCached(process, rect, offsetY, textHeight, m_processFont, progressRect, [&](PaintContext& textCacheContext)
 									{
 										RECT rect2{ 0, int(process.bitmapOffset), 256, int(process.bitmapOffset) + m_processFont.height };
 										RECT rect3{ 0, int(process.bitmapOffset), processWidth, int(process.bitmapOffset) + m_processFont.height };
@@ -1372,7 +1373,7 @@ namespace uba
 								if (insres.second)
 								{
 									auto c = (u8*)&color;
-									insres.first->second = CreateSolidBrush(RGB(c[2], c[1], c[0]));
+									insres.first->second = m_backend.CreateSolidBrush(RGB(c[2], c[1], c[0]));
 								}
 								brush = insres.first->second;
 							}
@@ -1393,7 +1394,7 @@ namespace uba
 							int processWidth = rect.right - rect.left;
 							if (shouldDrawText && m_config.ShowProcessText && processWidth > 3)
 							{
-								context.DrawCached(work, rect, offsetY, textHeight, progressRect, [&](PaintContext& textCacheContext)
+								context.DrawCached(work, rect, offsetY, textHeight, m_processFont, progressRect, [&](PaintContext& textCacheContext)
 									{
 										RECT rect2{ 0,int(work.bitmapOffset), 256, int(work.bitmapOffset) + m_processFont.height };
 										textCacheContext.FillRect(rect2, brush);
@@ -1462,7 +1463,7 @@ namespace uba
 			str[1].Appendf(TC("Box Count: %u"), visibleBoxes);
 
 			context.SetFont(m_popupFont);
-			context.DrawCenteredText(str, 2, true);
+			context.DrawCenteredText(str, 2, m_tooltipBackgroundBrush);
 		}
 
 		if (m_config.showCursorLine && m_mouseOverWindow)
@@ -2388,17 +2389,80 @@ namespace uba
 		}
 	}
 
-	void Visualizer::PaintEmptyScreen(PaintContext& context)
+	void Visualizer::WriteProcessStats(Logger& out, const TraceView::Process& process)
+	{
+		bool hasExited = process.stop != ~u64(0);
+		out.Info(TC("  %ls"), process.description.c_str());
+		out.Info(TC("  ProcessId: %u"), process.id);
+		out.Info(TC("  Start:     %ls"), TimeToText(process.start, true).str);
+		if (hasExited)
+			out.Info(TC("  Duration:  %ls"), TimeToText(process.stop - process.start, true).str);
+		if (hasExited && process.exitCode != 0)
+			out.Info(TC("  ExitCode:  %u"), process.exitCode);
+
+		if (process.stop != ~u64(0))
+		{
+			out.Info(TC(""));
+
+			BinaryReader reader(process.stats.data(), 0, process.stats.size());
+			ProcessStats processStats;
+			SessionStats sessionStats;
+			StorageStats storageStats;
+			KernelStats kernelStats;
+
+			processStats.Read(reader, m_traceView.version);
+			if (reader.GetLeft())
+			{
+				if (process.isRemote || (m_traceView.version >= 36 && !process.isReuse))
+					sessionStats.Read(reader, m_traceView.version);
+				storageStats.Read(reader, m_traceView.version);
+				kernelStats.Read(reader, m_traceView.version);
+			}
+
+			out.Info(TC("  ----------- Detours stats -----------"));
+			processStats.Print(out, m_traceView.frequency);
+
+			if (!sessionStats.IsEmpty())
+			{
+				out.Info(TC(""));
+				out.Info(TC("  ----------- Session stats -----------"));
+				sessionStats.Print(out, m_traceView.frequency);
+			}
+
+			if (!storageStats.IsEmpty())
+			{
+				out.Info(TC(""));
+				out.Info(TC("  ----------- Storage stats -----------"));
+				storageStats.Print(out, m_traceView.frequency);
+			}
+
+			if (!kernelStats.IsEmpty())
+			{
+				out.Info(TC(""));
+				out.Info(TC("  ----------- Kernel stats ------------"));
+				kernelStats.Print(out, false, m_traceView.frequency);
+			}
+
+			PrintCacheWriteStats(out, process.id);
+		}
+	}
+
+	void Visualizer::WriteWorkStats(Logger& out, const TraceView::WorkRecord& record)
 	{
+		out.Info(TC("  %ls"), record.description);
+		out.Info(TC("  Start:     %ls"), TimeToText(record.start, true).str);
+		if (record.stop != ~u64(0))
+			out.Info(TC("  Duration:  %ls"), TimeToText(record.stop - record.start, true).str);
+		for (auto& e : record.entries)
+			out.Info(TC("   %ls (%ls)"), e.text, TimeToText(e.time - record.start).str);
 	}
 
-	void Visualizer::FinishIncompatible()
+	void Visualizer::PaintEmptyScreen(PaintContext& context)
 	{
 	}
 
-	HBRUSH Visualizer::CreateSolidBrush(COLORREF color)
+	void Visualizer::FinishIncompatible()
 	{
-		return 0;
 	}
 
 	////////////////////////////////////////////////////////////////////////////////////////////////////
diff --git a/Engine/Source/Programs/UnrealBuildAccelerator/Visualizer/Private/UbaVisualizer.h b/Engine/Source/Programs/UnrealBuildAccelerator/Visualizer/Private/UbaVisualizer.h
--- a/Engine/Source/Programs/UnrealBuildAccelerator/Visualizer/Private/UbaVisualizer.h
+++ b/Engine/Source/Programs/UnrealBuildAccelerator/Visualizer/Private/UbaVisualizer.h
@@ -3,34 +3,17 @@
 #pragma once
 
 #include "UbaTraceReader.h"
+#include "UbaVisualizerBackend.h"
 #include "UbaVisualizerConfig.h"
 
 namespace uba
 {
 	////////////////////////////////////////////////////////////////////////////////////////////////////
 
-	#if !PLATFORM_WINDOWS
-	struct POINT { int x; int y; };
-	struct SIZE { int cx; int cy; };
-	struct RECT { int left; int top; int right; int bottom; };
-	using COLORREF = u32;
-	using HFONT = void*;
-	using HPEN = void*;
-	using HBRUSH = void*;
-	#define TRANSPARENT  1
-	#define OPAQUE       2
-	#define TA_LEFT      0
-	#define TA_RIGHT     2
-	#define TA_CENTER    6
-	#define RGB(r,g,b)          ((COLORREF)((u8(r)|(u16(u8(g))<<8))|(u32(u8(b))<<16)))
-	#endif
-
-	////////////////////////////////////////////////////////////////////////////////////////////////////
-
 	class Visualizer
 	{
 	public:
-		Visualizer(VisualizerConfig& config, Logger& logger);
+		Visualizer(VisualizerConfig& config, Logger& logger, VisualizerBackend& backend);
 
 	protected:
 		u64 GetPlayTime();
@@ -40,39 +23,8 @@ namespace uba
 		StringBuffer<128> GetWorldTime(float seconds);
 		bool ActiveProcessesShouldFillHeight();
 
-		struct Font
-		{
-			HFONT handle = 0;
-			HFONT handleUnderlined = 0;
-			int height = 0;
-			int offset = 0;
-		};
-
-		struct PaintContext
-		{
-			virtual RECT GetClientRect() = 0;
-
-			virtual void SetBkMode(int mode) = 0;
-			virtual void SetPen(HPEN pen) = 0;
-			virtual void SetBrush(HBRUSH b) = 0;
-			virtual void SetFont(const Font& font, bool underlined = false) = 0;
-			virtual void SetTextColor(COLORREF color) = 0;
-			virtual u32 SetTextAlign(u32 align) = 0;
-			virtual void Text(int x, int y, const RECT* lprect, StringView str) = 0;
-			virtual COLORREF GetTextColor() = 0;
-			virtual Font& GetFont() = 0;
-			virtual void FillRect(const RECT& r, HBRUSH b) = 0;
-			virtual void Rectangle(int left, int top, int right, int bottom) = 0;
-			virtual void Line(int fromX, int fromY, int toX, int toY) = 0;
-			virtual void Polyline(const Vector<POINT>& polyLine) = 0;
-			virtual void PolyPolyline(const POINT* points, const u32* lines, u32 lineCount) = 0;
-			virtual void DrawCenteredText(const StringBuffer<>* rows, u32 rowCount, bool drawBackground) = 0;
-			virtual SIZE GetTextSize(StringView str) = 0;
-			virtual void DrawCached(TraceView::BitmapCache& process, RECT rect, int offsetY, int textHeight, RECT clippingRect, const Function<void(PaintContext& textCacheContext)>& func) = 0;
-
-			virtual POINT GetCursorClientPos() = 0;
-		};
-
+		using Font = VisualizerBackend::Font;
+		using PaintContext = VisualizerBackend::PaintContext;
 		struct HitTestResult;
 
 		void HitTest(HitTestResult& outResult, PaintContext& context, const POINT& pos);
@@ -83,10 +35,11 @@ namespace uba
 		using DrawTextFunc = Function<void(const StringBufferBase& text, RECT& rect, u32* outWidth)>;
 		void PaintDetailedStats(PaintContext& context, int& posY, const RECT& progressRect, TraceView::Session& session, bool isRemote, u64 playTime, const DrawTextFunc& drawTextFunc);
 		void PrintCacheWriteStats(Logger& logger, u32 processId);
+		void WriteProcessStats(Logger& out, const TraceView::Process& process);
+		void WriteWorkStats(Logger& out, const TraceView::WorkRecord& record);
 
 		virtual void PaintEmptyScreen(PaintContext& context);
 		virtual void FinishIncompatible();
-		virtual HBRUSH CreateSolidBrush(COLORREF color);
 
 		void UpdateFont(Font& font, int height, bool createUnderline);
 
@@ -156,6 +109,7 @@ namespace uba
 		VisualizerConfig m_config;
 		TraceReader m_trace;
 		TraceView m_traceView;
+		VisualizerBackend& m_backend;
 
 		ProcessBrushes m_processBrushes[2]; // Non-selected and selected
 
diff --git a/Engine/Source/Programs/UnrealBuildAccelerator/Visualizer/Private/UbaVisualizerBackend.h b/Engine/Source/Programs/UnrealBuildAccelerator/Visualizer/Private/UbaVisualizerBackend.h
new file mode 100644
--- /dev/null
+++ b/Engine/Source/Programs/UnrealBuildAccelerator/Visualizer/Private/UbaVisualizerBackend.h
@@ -0,0 +1,69 @@
+// Copyright Epic Games, Inc. All Rights Reserved.
+
+#pragma once
+
+namespace uba
+{
+	////////////////////////////////////////////////////////////////////////////////////////////////////
+
+	#if !PLATFORM_WINDOWS
+	struct POINT { int x; int y; };
+	struct SIZE { int cx; int cy; };
+	struct RECT { int left; int top; int right; int bottom; };
+	using COLORREF = u32;
+	using HFONT = void*;
+	using HPEN = void*;
+	using HBRUSH = void*;
+	#define TRANSPARENT  1
+	#define OPAQUE       2
+	#define TA_LEFT      0
+	#define TA_RIGHT     2
+	#define TA_CENTER    6
+	#define RGB(r,g,b)          ((COLORREF)((u8(r)|(u16(u8(g))<<8))|(u32(u8(b))<<16)))
+	#endif
+
+	////////////////////////////////////////////////////////////////////////////////////////////////////
+
+	class VisualizerBackend
+	{
+	public:
+		struct Font
+		{
+			HFONT handle = 0;
+			HFONT handleUnderlined = 0;
+			int height = 0;
+			int offset = 0;
+		};
+
+		struct PaintContext
+		{
+			virtual RECT GetClientRect() = 0;
+
+			virtual void SetBkMode(int mode) = 0;
+			virtual void SetPen(HPEN pen) = 0;
+			virtual void SetBrush(HBRUSH b) = 0;
+			virtual void SetFont(const Font& font, bool underlined = false) = 0;
+			virtual void SetTextColor(COLORREF color) = 0;
+			virtual u32 SetTextAlign(u32 align) = 0;
+			virtual void Text(int x, int y, const RECT* lprect, StringView str) = 0;
+			virtual COLORREF GetTextColor() = 0;
+			virtual Font& GetFont() = 0;
+			virtual void FillRect(const RECT& r, HBRUSH b) = 0;
+			virtual void Rectangle(int left, int top, int right, int bottom) = 0;
+			virtual void Line(int fromX, int fromY, int toX, int toY) = 0;
+			virtual void Polyline(const Vector<POINT>& polyLine) = 0;
+			virtual void PolyPolyline(const POINT* points, const u32* lines, u32 lineCount) = 0;
+			virtual void DrawCenteredText(const StringBuffer<>* rows, u32 rowCount, HBRUSH backgroundBrush) = 0;
+			virtual SIZE GetTextSize(StringView str) = 0;
+			virtual void DrawCached(TraceView::BitmapCache& process, RECT rect, int offsetY, int textHeight, Font& processFont, RECT clippingRect, const Function<void(PaintContext& textCacheContext)>& func) = 0;
+
+			virtual POINT GetCursorClientPos() = 0;
+		};
+
+		virtual void ResetCaches() = 0;
+
+		virtual HBRUSH CreateSolidBrush(COLORREF color) = 0;
+	};
+
+	////////////////////////////////////////////////////////////////////////////////////////////////////
+}
diff --git a/Engine/Source/Programs/UnrealBuildAccelerator/Visualizer/Private/Windows/UbaVisualizerBackendGDI.h b/Engine/Source/Programs/UnrealBuildAccelerator/Visualizer/Private/Windows/UbaVisualizerBackendGDI.h
new file mode 100644
--- /dev/null
+++ b/Engine/Source/Programs/UnrealBuildAccelerator/Visualizer/Private/Windows/UbaVisualizerBackendGDI.h
@@ -0,0 +1,223 @@
+// Copyright Epic Games, Inc. All Rights Reserved.
+
+#pragma once
+
+#include "UbaVisualizerBackend.h"
+
+namespace uba
+{
+	class VisualizerWin;
+
+	////////////////////////////////////////////////////////////////////////////////////////////////////
+
+	class VisualizerBackendGDI : public VisualizerBackend
+	{
+	public:
+		static constexpr int BitmapCacheHeight = 1024*1024;
+
+		struct PaintContextGDI : public PaintContext
+		{
+			virtual RECT GetClientRect() override { return clientRect; }
+			virtual void SetBkMode(int mode) override { ::SetBkMode(hdc, mode); }
+			virtual void SetTextColor(COLORREF color) override { ::SetTextColor(hdc, color); }
+			virtual COLORREF GetTextColor() override { return ::GetTextColor(hdc); }
+			virtual void SetPen(HPEN pen) override { ::SelectObject(hdc, pen); }
+			virtual u32 SetTextAlign(u32 align) override { return ::SetTextAlign(hdc, align); }
+			virtual void Text(int x, int y, const RECT* clipRect, StringView str) override { ::ExtTextOutW(hdc, x, y, clipRect ? ETO_CLIPPED : 0, clipRect, str.data, str.count, NULL); }
+
+			virtual void SetFont(const Font& font, bool underlined) override
+			{
+				if (underlined)
+				{
+					if (font.handleUnderlined == activeFont.handleUnderlined)
+						return;
+					::SelectObject(hdc, font.handleUnderlined);
+				}
+				else
+				{
+					if (font.handle == activeFont.handle)
+						return;
+					::SelectObject(hdc, font.handle);
+				}
+				activeFont = font;
+			}
+
+			virtual Font& GetFont() override
+			{
+				return activeFont;
+			}
+
+			virtual void SetBrush(HBRUSH b) override
+			{
+				if (lastSelectedBrush != b)
+				{
+					SelectObject(hdc, b);
+					lastSelectedBrush = b;
+				}
+			}
+
+			virtual void FillRect(const RECT& r, HBRUSH b) override
+			{
+				SetBrush(b);
+				PatBlt(hdc, r.left, r.top, r.right - r.left, r.bottom - r.top, PATCOPY);
+			}
+
+			virtual void Rectangle(int left, int top, int right, int bottom) override
+			{
+				SetBrush((HBRUSH)GetStockObject(NULL_BRUSH));
+				::Rectangle(hdc, left, top, right, bottom);
+			}
+
+			virtual void Line(int fromX, int fromY, int toX, int toY) override
+			{
+				::MoveToEx(hdc, fromX, fromY, NULL);
+				::LineTo(hdc, toX, toY);
+			}
+
+			virtual void Polyline(const Vector<POINT>& polyLine) override
+			{
+				::Polyline(hdc, polyLine.data(), int(polyLine.size()));
+			}
+
+			virtual void PolyPolyline(const POINT* points, const u32* lines, u32 lineCount) override
+			{
+				::PolyPolyline(hdc, points, (const DWORD*)lines, int(lineCount));
+			}
+
+			virtual void DrawCenteredText(const StringBuffer<>* rows, u32 rowCount, HBRUSH backgroundBrush) override
+			{
+				for (u32 i=0; i!=rowCount; ++i)
+				{ 
+					SIZE textSize = GetTextSize(rows[i]);
+
+					RECT rect = { 0, 0, 0, 0 };
+					int top = (clientRect.bottom - textSize.cy)/2 + i*20;
+					rect.left = (clientRect.right - textSize.cx)/2;
+					rect.right = (clientRect.right + textSize.cx)/2;
+					rect.bottom = top + textSize.cy;
+					rect.top = top;
+
+					if (backgroundBrush)
+					{
+						RECT r = rect;
+						InflateRect(&r, 4, 2);
+						FillRect(r, backgroundBrush);
+					}
+					Text(rect.left, rect.top, nullptr, rows[i]);
+				}
+			}
+
+			virtual SIZE GetTextSize(StringView str) override
+			{
+				SIZE s;
+				return GetTextExtentPoint32W(hdc, str.data, str.count, &s) ? s : SIZE();
+			}
+
+			virtual void DrawCached(TraceView::BitmapCache& process, RECT rect, int offsetY, int textHeight, Font& processFont, RECT clippingRect, const Function<void(PaintContext& textCacheContext)>& func) override
+			{
+				if (!process.bitmap || process.bitmapDirty)
+				{
+					if (!process.bitmap)
+					{
+						if (owner.m_lastBitmapOffset == BitmapCacheHeight)
+						{
+							if (owner.m_lastBitmap)
+								owner.m_textBitmaps.push_back(owner.m_lastBitmap);
+
+							owner.m_lastBitmapOffset = 0;
+							owner.m_lastBitmap = CreateCompatibleBitmap(hdc, 256, BitmapCacheHeight);
+						}
+						process.bitmap = owner.m_lastBitmap;
+						process.bitmapOffset = owner.m_lastBitmapOffset;
+						owner.m_lastBitmapOffset += processFont.height;
+					}
+					if (lastSelectedBitmap != process.bitmap)
+					{
+						SelectObject(textDC, process.bitmap);
+						lastSelectedBitmap = process.bitmap;
+					}
+
+					PaintContextGDI textCacheContext{owner, NULL, textDC, clientRect };
+					func(textCacheContext);
+				}
+
+				if (lastSelectedBitmap != process.bitmap)
+				{
+					SelectObject(textDC, process.bitmap);
+					lastSelectedBitmap = process.bitmap;
+				}
+
+				int processWidth = rect.right - rect.left;
+				int width = Min(processWidth, 256);
+				int bitmapOffsetY = process.bitmapOffset;
+				int bltOffsetY = offsetY;
+				if (bltOffsetY < 0)
+				{
+					bitmapOffsetY -= bltOffsetY;
+					bltOffsetY = 0;
+				}
+				int height = Min(textHeight, processFont.height);
+				if (bltOffsetY + height > textHeight)
+					height = textHeight - bltOffsetY;
+
+				int left = rect.left;
+				if (left <= -256 || height <= 0)
+					return;
+
+				int bitmapOffsetX = rect.left - left;
+
+				if (left < clippingRect.left)
+				{
+					int diff = clippingRect.left - left;
+					rect.left = clippingRect.left;
+					width -= diff;
+					bitmapOffsetX += diff;
+				}
+				BitBlt(hdc, rect.left, rect.top + bltOffsetY, width, height, textDC, bitmapOffsetX, bitmapOffsetY, SRCCOPY);
+			}
+
+			virtual POINT GetCursorClientPos() override
+			{
+				POINT p;
+				GetCursorPos(&p);
+				ScreenToClient(hwnd, &p);
+				return p;
+			}
+
+			PaintContextGDI(VisualizerBackendGDI& o, HWND hw, HDC dc, const RECT& cr, HDC tdc = 0) : owner(o), hwnd(hw), hdc(dc), clientRect(cr), textDC(tdc) {}
+
+
+			VisualizerBackendGDI& owner;
+			HWND hwnd;
+			HDC hdc;
+			const RECT& clientRect;
+
+			Font activeFont;
+
+			HDC textDC = NULL;
+			HBRUSH lastSelectedBrush = 0;
+			HBITMAP lastSelectedBitmap = 0;
+		};
+
+		virtual void ResetCaches() override
+		{
+			for (HBITMAP bm : m_textBitmaps)
+				DeleteObject(bm);
+			m_textBitmaps.clear();
+			DeleteObject(m_lastBitmap);
+			m_lastBitmap = 0;
+			m_lastBitmapOffset = BitmapCacheHeight;
+		}
+
+		virtual HBRUSH CreateSolidBrush(COLORREF color) override
+		{
+			return ::CreateSolidBrush(color);
+		}
+
+		HBITMAP m_lastBitmap = 0;
+		int m_lastBitmapOffset = BitmapCacheHeight;
+		Vector<HBITMAP> m_textBitmaps;
+	};
+
+	////////////////////////////////////////////////////////////////////////////////////////////////////
+}
diff --git a/Engine/Source/Programs/UnrealBuildAccelerator/Visualizer/Private/Windows/UbaVisualizerWin.cpp b/Engine/Source/Programs/UnrealBuildAccelerator/Visualizer/Private/Windows/UbaVisualizerWin.cpp
--- a/Engine/Source/Programs/UnrealBuildAccelerator/Visualizer/Private/Windows/UbaVisualizerWin.cpp
+++ b/Engine/Source/Programs/UnrealBuildAccelerator/Visualizer/Private/Windows/UbaVisualizerWin.cpp
@@ -2,8 +2,9 @@
 
 #include "UbaVisualizerWin.h"
 #include "UbaConfig.h"
-#include <algorithm>
+#include "UbaVisualizerBackendGDI.h"
 
+#include <algorithm>
 #include <uxtheme.h>
 #include <dwmapi.h>
 #include <shellscalingapi.h>
@@ -45,190 +46,6 @@ namespace uba
 		Popup_Quit,
 	};
 
-	struct VisualizerWin::PaintContextWin : public Visualizer::PaintContext 
-	{
-		virtual RECT GetClientRect() override { return clientRect; }
-		virtual void SetBkMode(int mode) override { ::SetBkMode(hdc, mode); }
-		virtual void SetTextColor(COLORREF color) override { ::SetTextColor(hdc, color); }
-		virtual COLORREF GetTextColor() override { return ::GetTextColor(hdc); }
-		virtual void SetPen(HPEN pen) override { ::SelectObject(hdc, pen); }
-		virtual u32 SetTextAlign(u32 align) override { return ::SetTextAlign(hdc, align); }
-		virtual void Text(int x, int y, const RECT* clipRect, StringView str) override { ::ExtTextOutW(hdc, x, y, clipRect ? ETO_CLIPPED : 0, clipRect, str.data, str.count, NULL); }
-		
-		virtual void SetFont(const Font& font, bool underlined) override
-		{
-			if (underlined)
-			{
-				if (font.handleUnderlined == activeFont.handleUnderlined)
-					return;
-				::SelectObject(hdc, font.handleUnderlined);
-			}
-			else
-			{
-				if (font.handle == activeFont.handle)
-					return;
-				::SelectObject(hdc, font.handle);
-			}
-			activeFont = font;
-		}
-
-		virtual Font& GetFont() override
-		{
-			return activeFont;
-		}
-
-		virtual void SetBrush(HBRUSH b) override
-		{
-			if (lastSelectedBrush != b)
-			{
-				SelectObject(hdc, b);
-				lastSelectedBrush = b;
-			}
-		}
-
-		virtual void FillRect(const RECT& r, HBRUSH b) override
-		{
-			SetBrush(b);
-			PatBlt(hdc, r.left, r.top, r.right - r.left, r.bottom - r.top, PATCOPY);
-		}
-
-		virtual void Rectangle(int left, int top, int right, int bottom) override
-		{
-			SetBrush((HBRUSH)GetStockObject(NULL_BRUSH));
-			::Rectangle(hdc, left, top, right, bottom);
-		}
-
-		virtual void Line(int fromX, int fromY, int toX, int toY) override
-		{
-			::MoveToEx(hdc, fromX, fromY, NULL);
-			::LineTo(hdc, toX, toY);
-		}
-
-		virtual void Polyline(const Vector<POINT>& polyLine) override
-		{
-			::Polyline(hdc, polyLine.data(), int(polyLine.size()));
-		}
-
-		virtual void PolyPolyline(const POINT* points, const u32* lines, u32 lineCount) override
-		{
-			::PolyPolyline(hdc, points, (const DWORD*)lines, int(lineCount));
-		}
-
-		virtual void DrawCenteredText(const StringBuffer<>* rows, u32 rowCount, bool drawBackground) override
-		{
-			for (u32 i=0; i!=rowCount; ++i)
-			{ 
-				SIZE textSize = GetTextSize(rows[i]);
-
-				RECT rect = { 0, 0, 0, 0 };
-				int top = (clientRect.bottom - textSize.cy)/2 + i*20;
-				rect.left = (clientRect.right - textSize.cx)/2;
-				rect.right = (clientRect.right + textSize.cx)/2;
-				rect.bottom = top + textSize.cy;
-				rect.top = top;
-
-				if (drawBackground)
-				{
-					RECT r = rect;
-					InflateRect(&r, 4, 2);
-					FillRect(r, owner.m_tooltipBackgroundBrush);
-				}
-				Text(rect.left, rect.top, nullptr, rows[i]);
-			}
-		}
-
-		virtual SIZE GetTextSize(StringView str) override
-		{
-			SIZE s;
-			return GetTextExtentPoint32W(hdc, str.data, str.count, &s) ? s : SIZE();
-		}
-
-		virtual void DrawCached(TraceView::BitmapCache& process, RECT rect, int offsetY, int textHeight, RECT clippingRect, const Function<void(PaintContext& textCacheContext)>& func) override
-		{
-			if (!process.bitmap || process.bitmapDirty)
-			{
-				if (!process.bitmap)
-				{
-					if (owner.m_lastBitmapOffset == BitmapCacheHeight)
-					{
-						if (owner.m_lastBitmap)
-							owner.m_textBitmaps.push_back(owner.m_lastBitmap);
-
-						owner.m_lastBitmapOffset = 0;
-						owner.m_lastBitmap = CreateCompatibleBitmap(hdc, 256, BitmapCacheHeight);
-					}
-					process.bitmap = owner.m_lastBitmap;
-					process.bitmapOffset = owner.m_lastBitmapOffset;
-					owner.m_lastBitmapOffset += owner.m_processFont.height;
-				}
-				if (lastSelectedBitmap != process.bitmap)
-				{
-					SelectObject(textDC, process.bitmap);
-					lastSelectedBitmap = process.bitmap;
-				}
-
-				PaintContextWin textCacheContext{owner, NULL, textDC, clientRect };
-				func(textCacheContext);
-			}
-
-			if (lastSelectedBitmap != process.bitmap)
-			{
-				SelectObject(textDC, process.bitmap);
-				lastSelectedBitmap = process.bitmap;
-			}
-
-			int processWidth = rect.right - rect.left;
-			int width = Min(processWidth, 256);
-			int bitmapOffsetY = process.bitmapOffset;
-			int bltOffsetY = offsetY;
-			if (bltOffsetY < 0)
-			{
-				bitmapOffsetY -= bltOffsetY;
-				bltOffsetY = 0;
-			}
-			int height = Min(textHeight, owner.m_processFont.height);
-			if (bltOffsetY + height > textHeight)
-				height = textHeight - bltOffsetY;
-
-			int left = rect.left;
-			if (left <= -256 || height <= 0)
-				return;
-
-			int bitmapOffsetX = rect.left - left;
-
-			if (left < clippingRect.left)
-			{
-				int diff = clippingRect.left - left;
-				rect.left = clippingRect.left;
-				width -= diff;
-				bitmapOffsetX += diff;
-			}
-			BitBlt(hdc, rect.left, rect.top + bltOffsetY, width, height, textDC, bitmapOffsetX, bitmapOffsetY, SRCCOPY);
-		}
-
-		virtual POINT GetCursorClientPos() override
-		{
-			POINT p;
-			GetCursorPos(&p);
-			ScreenToClient(hwnd, &p);
-			return p;
-		}
-
-		PaintContextWin(VisualizerWin& o, HWND hw, HDC dc, const RECT& cr, HDC tdc = 0) : owner(o), hwnd(hw), hdc(dc), clientRect(cr), textDC(tdc) {}
-
-
-		VisualizerWin& owner;
-		HWND hwnd;
-		HDC hdc;
-		const RECT& clientRect;
-
-		Font activeFont;
-
-		HDC textDC = NULL;
-		HBRUSH lastSelectedBrush = 0;
-		HBITMAP lastSelectedBitmap = 0;
-	};
-
 	class WriteTextLogger : public Logger
 	{
 	public:
@@ -240,7 +57,7 @@ namespace uba
 	};
 
 	VisualizerWin::VisualizerWin(VisualizerConfig& config, Logger& logger)
-		:	Visualizer(config, logger)
+	:	Visualizer(config, logger, *new VisualizerBackendGDI())
 	{
 		memset(m_activeProcessFont, 0, sizeof(m_activeProcessFont));
 		memset(m_activeProcessCountHistory, 0, sizeof(m_activeProcessCountHistory));
@@ -467,14 +284,9 @@ namespace uba
 
 	void VisualizerWin::Reset()
 	{
-		for (HBITMAP bm : m_textBitmaps)
-			DeleteObject(bm);
-		DeleteObject(m_lastBitmap);
+		m_backend.ResetCaches();
 		m_contentWidth = 0;
 		m_contentHeight = 0;
-		m_textBitmaps.clear();
-		m_lastBitmap = 0;
-		m_lastBitmapOffset = BitmapCacheHeight;
 		//m_autoScroll = true;
 		//m_scrollPosX = 0;
 		//m_scrollPosY = 0;
@@ -682,7 +494,7 @@ namespace uba
 		{
 			RECT clientRect;
 			GetClientRect(m_hwnd, &clientRect);
-			PaintContextWin context(*this, NULL, NULL, clientRect);
+			VisualizerBackendGDI::PaintContextGDI context((VisualizerBackendGDI&)m_backend, NULL, NULL, clientRect);
 			HitTestResult res;
 			HitTest(res, context, { -1, -1 });
 		}
@@ -832,12 +644,7 @@ namespace uba
 		if (!full)
 			return;
 
-		for (HBITMAP bm : m_textBitmaps)
-			DeleteObject(bm);
-		DeleteObject(m_lastBitmap);
-		m_textBitmaps.clear();
-		m_lastBitmapOffset = BitmapCacheHeight;
-		m_lastBitmap = 0;
+		m_backend.ResetCaches();
 	}
 
 	void Visualizer::UpdateFont(Font& font, int height, bool createUnderline)
@@ -967,7 +774,7 @@ namespace uba
 				str[0].Append(TCV("No trace active"));
 		}
 		context.SetFont(m_popupFont);
-		context.DrawCenteredText(str, 2, false);
+		context.DrawCenteredText(str, 2, nullptr);
 	}
 
 	void VisualizerWin::FinishIncompatible()
@@ -976,79 +783,6 @@ namespace uba
 		SetWindowTextW(m_hwnd, GetTitlePrefix(title).Appendf(L" (Listening for new sessions on channel '%s')", m_listenChannel.data).data);
 	}
 
-	HBRUSH VisualizerWin::CreateSolidBrush(COLORREF color)
-	{
-		return ::CreateSolidBrush(color);
-	}
-
-	void VisualizerWin::WriteProcessStats(Logger& out, const TraceView::Process& process)
-	{
-		bool hasExited = process.stop != ~u64(0);
-		out.Info(L"  %ls", process.description.c_str());
-		out.Info(L"  ProcessId: %u", process.id);
-		out.Info(L"  Start:     %ls", TimeToText(process.start, true).str);
-		if (hasExited)
-			out.Info(L"  Duration:  %ls", TimeToText(process.stop - process.start, true).str);
-		if (hasExited && process.exitCode != 0)
-			out.Info(L"  ExitCode:  %u", process.exitCode);
-
-		if (process.stop != ~u64(0))
-		{
-			out.Info(L"");
-
-			BinaryReader reader(process.stats.data(), 0, process.stats.size());
-			ProcessStats processStats;
-			SessionStats sessionStats;
-			StorageStats storageStats;
-			KernelStats kernelStats;
-
-			processStats.Read(reader, m_traceView.version);
-			if (reader.GetLeft())
-			{
-				if (process.isRemote || (m_traceView.version >= 36 && !process.isReuse))
-					sessionStats.Read(reader, m_traceView.version);
-				storageStats.Read(reader, m_traceView.version);
-				kernelStats.Read(reader, m_traceView.version);
-			}
-
-			out.Info(L"  ----------- Detours stats -----------");
-			processStats.Print(out, m_traceView.frequency);
-
-			if (!sessionStats.IsEmpty())
-			{
-				out.Info(L"");
-				out.Info(L"  ----------- Session stats -----------");
-				sessionStats.Print(out, m_traceView.frequency);
-			}
-
-			if (!storageStats.IsEmpty())
-			{
-				out.Info(L"");
-				out.Info(L"  ----------- Storage stats -----------");
-				storageStats.Print(out, m_traceView.frequency);
-			}
-
-			if (!kernelStats.IsEmpty())
-			{
-				out.Info(L"");
-				out.Info(L"  ----------- Kernel stats ------------");
-				kernelStats.Print(out, false, m_traceView.frequency);
-			}
-
-			PrintCacheWriteStats(out, process.id);
-		}
-	}
-
-	void VisualizerWin::WriteWorkStats(Logger& out, const TraceView::WorkRecord& record)
-	{
-		out.Info(L"  %ls", record.description);
-		out.Info(L"  Start:     %ls", TimeToText(record.start, true).str);
-		if (record.stop != ~u64(0))
-			out.Info(L"  Duration:  %ls", TimeToText(record.stop - record.start, true).str);
-		for (auto& e : record.entries)
-			out.Info(L"   %ls (%ls)", e.text, TimeToText(e.time - record.start).str);
-	}
-
 	void VisualizerWin::CopyTextToClipboard(const TString& str)
 	{
 		if (!OpenClipboard(m_hwnd))
@@ -1114,7 +848,7 @@ namespace uba
 		{
 			RECT clientRect;
 			GetClientRect(m_hwnd, &clientRect);
-			PaintContextWin context(*this, NULL, NULL, clientRect);
+			VisualizerBackendGDI::PaintContextGDI context((VisualizerBackendGDI&)m_backend, NULL, NULL, clientRect);
 			HitTest(res, context, pos);
 		}
 		m_activeSection = res.section;
@@ -1328,7 +1062,7 @@ namespace uba
 					HBITMAP nullBmp = CreateCompatibleBitmap(memDC, 1, 1);
 					HBITMAP oldBmp = (HBITMAP)SelectObject(textDC, nullBmp);
 
-					PaintContextWin context(*this, m_hwnd, memDC, rect, textDC);
+					VisualizerBackendGDI::PaintContextGDI context((VisualizerBackendGDI&)m_backend, m_hwnd, memDC, rect, textDC);
 					PaintAll(context, rect);
 
 					SelectObject(textDC, oldBmp);
@@ -1612,7 +1346,7 @@ namespace uba
 
 				RECT clientRect;
 				GetClientRect(m_hwnd, &clientRect);
-				PaintContextWin context(*this, NULL, NULL, clientRect);
+				VisualizerBackendGDI::PaintContextGDI context((VisualizerBackendGDI&)m_backend, NULL, NULL, clientRect);
 				HitTestResult res;
 				HitTest(res, context, { -1, -1 });
 				UpdateScrollbars(true);
@@ -1665,7 +1399,7 @@ namespace uba
 					GetClientRect(hWnd, &r);
 
 					HitTestResult res;
-					PaintContextWin context(*this, NULL, NULL, r);
+					VisualizerBackendGDI::PaintContextGDI context((VisualizerBackendGDI&)m_backend, NULL, NULL, r);
 					HitTest(res, context, { -1, -1 });
 
 					m_scrollPosX = Min(Max(m_scrollPosX, float(r.right - m_contentWidth)), 0.0f);
diff --git a/Engine/Source/Programs/UnrealBuildAccelerator/Visualizer/Private/Windows/UbaVisualizerWin.h b/Engine/Source/Programs/UnrealBuildAccelerator/Visualizer/Private/Windows/UbaVisualizerWin.h
--- a/Engine/Source/Programs/UnrealBuildAccelerator/Visualizer/Private/Windows/UbaVisualizerWin.h
+++ b/Engine/Source/Programs/UnrealBuildAccelerator/Visualizer/Private/Windows/UbaVisualizerWin.h
@@ -8,6 +8,8 @@
 
 namespace uba
 {
+	////////////////////////////////////////////////////////////////////////////////////////////////////
+
 	class VisualizerWin : public Visualizer
 	{
 	public:
@@ -30,10 +32,7 @@ namespace uba
 		void PaintClient(const Function<void(HDC hdc, HDC memDC, RECT& clientRect)>& paintFunc);
 		virtual void PaintEmptyScreen(PaintContext& context) override;
 		virtual void FinishIncompatible() override;
-		virtual HBRUSH CreateSolidBrush(COLORREF color) override;
 
-		void WriteProcessStats(Logger& out, const TraceView::Process& process);
-		void WriteWorkStats(Logger& out, const TraceView::WorkRecord& record);
 		void CopyTextToClipboard(const TString& str);
 		void UnselectAndRedraw();
 		bool UpdateAutoscroll();
@@ -76,10 +75,6 @@ namespace uba
 
 		bool m_isInPaint = false;
 
-		static constexpr int BitmapCacheHeight = 1024*1024;
-		HBITMAP m_lastBitmap = 0;
-		int m_lastBitmapOffset = BitmapCacheHeight;
-
 		bool m_usingNamed = false;
 		bool m_showPopup = false;
 		bool m_locked = false;
@@ -88,8 +83,6 @@ namespace uba
 		HBITMAP m_cachedBitmap = 0;
 		RECT m_cachedBitmapRect = { INT_MIN, INT_MIN, INT_MIN, INT_MIN };
 
-		Vector<HBITMAP> m_textBitmaps;
-
 		POINT m_mouseAnchor = {};
 		float m_scrollAtAnchorX = 0;
 		float m_scrollAtAnchorY = 0;
@@ -107,5 +100,8 @@ namespace uba
 		static LRESULT CALLBACK StaticWinProc(HWND hWnd, UINT Msg, WPARAM wParam, LPARAM lParam);
 
 		struct PaintContextWin;
+		friend class VisualizerBackendGDI;
 	};
+
+	////////////////////////////////////////////////////////////////////////////////////////////////////
 }
