[UBA]
* Started implementing native Visualizer for mac

diff --git a/Engine/Source/Programs/UnrealBuildAccelerator/UbaAgent.Target.cs b/Engine/Source/Programs/UnrealBuildAccelerator/UbaAgent.Target.cs
--- a/Engine/Source/Programs/UnrealBuildAccelerator/UbaAgent.Target.cs
+++ b/Engine/Source/Programs/UnrealBuildAccelerator/UbaAgent.Target.cs
@@ -100,10 +100,13 @@ public class UbaAgentTarget : TargetRules
 			{
 				Rules.Architectures.Architectures.Add(UnrealArch.Arm64);
 			}
-			string BinaryPrefix = Rules.bShouldCompileAsDLL ? "lib" : string.Empty;
-			string BinaryExt = Rules.bShouldCompileAsDLL ? ".dylib" : string.Empty;
-			BinaryName = $"{BinaryPrefix}{Rules.LaunchModuleName}{BinaryExt}";
-			Rules.OutputFile = Path.Combine(BinariesFolder, BinaryName);
+			if (Rules.LaunchModuleName != "UbaVisualizer") // Revisit this. UbaVisualizer fails packaging if name is changed
+			{
+				string BinaryPrefix = Rules.bShouldCompileAsDLL ? "lib" : string.Empty;
+				string BinaryExt = Rules.bShouldCompileAsDLL ? ".dylib" : string.Empty;
+				BinaryName = $"{BinaryPrefix}{Rules.LaunchModuleName}{BinaryExt}";
+				Rules.OutputFile = Path.Combine(BinariesFolder, BinaryName);
+			}
 			Rules.GlobalDefinitions.AddRange(new string[] {
 				"UBA_DETOURS_LIBRARY=\"libUbaDetours.dylib\"",
 				"UBA_AGENT_EXECUTABLE=\"UbaAgent\"",
diff --git a/Engine/Source/Programs/UnrealBuildAccelerator/Visualizer/Private/Linux/UbaVisualizerApp.cpp b/Engine/Source/Programs/UnrealBuildAccelerator/Visualizer/Private/Linux/UbaVisualizerApp.cpp
new file mode 100644
--- /dev/null
+++ b/Engine/Source/Programs/UnrealBuildAccelerator/Visualizer/Private/Linux/UbaVisualizerApp.cpp
@@ -0,0 +1,6 @@
+// Copyright Epic Games, Inc. All Rights Reserved.
+
+int main(int argc, char* argv[])
+{
+	return 0;
+}
\ No newline at end of file
diff --git a/Engine/Source/Programs/UnrealBuildAccelerator/Visualizer/Private/Mac/UbaVisualizerApp.mm b/Engine/Source/Programs/UnrealBuildAccelerator/Visualizer/Private/Mac/UbaVisualizerApp.mm
new file mode 100644
--- /dev/null
+++ b/Engine/Source/Programs/UnrealBuildAccelerator/Visualizer/Private/Mac/UbaVisualizerApp.mm
@@ -0,0 +1,293 @@
+#import <Cocoa/Cocoa.h>
+
+#include "UbaFile.h"
+#include "UbaVisualizer.h"
+
+// If you want to call into C++ code, declare your C++ functions here:
+#ifdef __cplusplus
+extern "C" {
+#endif
+void cpp_render_example(float width, float height);
+#ifdef __cplusplus
+}
+#endif
+
+namespace uba
+{
+    class VisualizerBackendMac : public VisualizerBackend
+    {
+    public:
+        struct PaintContextMac : PaintContext
+        {
+			virtual RECT GetClientRect()
+            {
+                return clientRect;
+            }
+			virtual void SetBkMode(int mode) {}
+			virtual void SetPen(HPEN pen) {}
+			virtual void SetBrush(HBRUSH b) {}
+			virtual void SetFont(const Font& font, bool underlined = false) { activeFont = font; }
+			virtual void SetTextColor(COLORREF color) {}
+			virtual u32 SetTextAlign(u32 align) { return 0; }
+
+            int textOffsetX = 0;
+            int textOffsetY = 0;
+            NSRect textRect;
+
+			virtual void Text(int x, int y, const RECT* lprect, StringView str)
+            {         
+                NSString *text = [NSString stringWithUTF8String:str.data];
+
+                if (!lprect)
+                {
+                    NSDictionary *attrs = @{
+                        NSFontAttributeName: [NSFont systemFontOfSize:10.0],
+                        NSForegroundColorAttributeName: [NSColor blackColor],
+                    };
+                    NSPoint pos = NSMakePoint(x + textOffsetX, y + textOffsetY);
+                    [text drawAtPoint:pos withAttributes:attrs];
+                }
+                else
+                {
+                    NSMutableParagraphStyle *style = [[NSMutableParagraphStyle alloc] init];
+                    style.lineBreakMode = NSLineBreakByTruncatingTail; // ... at the end
+                    style.alignment = NSTextAlignmentLeft;
+
+                    NSDictionary *attrs = @{
+                        NSFontAttributeName: [NSFont systemFontOfSize:10.0],
+                        NSForegroundColorAttributeName: [NSColor blackColor],
+                        NSParagraphStyleAttributeName: style
+                    };
+                    NSRect textRect = NSMakeRect(x + textOffsetX, y + textOffsetY, lprect->right - lprect->left, lprect->bottom - lprect->right); // narrow rect
+
+                    NSAttributedString *attrText = [[NSAttributedString alloc] initWithString:text attributes:attrs];
+                    [attrText drawWithRect:textRect
+                                options:NSStringDrawingUsesLineFragmentOrigin
+                                context:nil];
+                }
+            }
+
+			virtual COLORREF GetTextColor() { return {}; }
+			virtual Font& GetFont() { return activeFont; }
+			virtual void FillRect(const RECT& r, HBRUSH b)
+            {
+                [[NSColor greenColor] setFill];
+                NSRect rect = NSMakeRect(r.left, r.top, r.right - r.left, r.bottom - r.top);
+                NSRectFill(rect);
+            }
+			virtual void Rectangle(int left, int top, int right, int bottom) {}
+			virtual void Line(int fromX, int fromY, int toX, int toY) {}
+			virtual void Polyline(const Vector<POINT>& polyLine) {}
+			virtual void PolyPolyline(const POINT* points, const u32* lines, u32 lineCount) {}
+			virtual void DrawCenteredText(const StringBuffer<>* rows, u32 rowCount, HBRUSH backgroundBrush) {}
+			virtual SIZE GetTextSize(StringView str)
+            {
+                return { int(str.count)*8, 10 };
+            }
+			virtual void DrawCached(TraceView::BitmapCache& process, RECT rect, int offsetY, int textHeight, Font& processFont, RECT clippingRect, const Function<void(PaintContext& textCacheContext)>& func)
+            {
+                textOffsetX = rect.left;
+                textOffsetY = rect.top;
+                func(*this);
+                textOffsetX = 0;
+                textOffsetY = 0;
+            }
+
+			virtual POINT GetCursorClientPos() { return {}; }
+
+            RECT clientRect;
+			Font activeFont;
+        };
+
+		virtual void ResetCaches() override {}
+        virtual HBRUSH CreateSolidBrush(COLORREF color) override { return 0; }
+    };
+
+    class VisualizerMac : public Visualizer
+    {
+    public:
+        VisualizerMac(VisualizerConfig& config, Logger& logger, VisualizerBackend& backend) : Visualizer(config, logger, backend)
+        {
+            char path[512];
+            GetFullPathNameW("~/Trace.uba", 512, path, nullptr);
+            m_trace.ReadFile(m_traceView, path, false);
+		    m_defaultFont.height = 14;
+		    m_processFont.height = 14;
+		    m_timelineFont.height = 14;
+    		m_popupFont.height = 14;
+        }
+		virtual void PaintEmptyScreen(PaintContext& context) override
+        {
+            [[NSColor blackColor] setFill];
+            NSRect rect = NSMakeRect(50, 50, 200, 100);
+            NSRectFill(rect);
+        }
+		virtual void FinishIncompatible()
+        {
+        }
+
+        using Visualizer::PaintAll;
+    };
+
+	void Visualizer::UpdateFont(Font& font, int height, bool createUnderline)
+    {
+        font.height = height;
+    }
+}
+
+@interface MyCanvasView : NSView
+{
+    NSTimeInterval _time;
+    uba::VisualizerMac* _visualizer;
+}
+- (void)tick:(NSTimer *)timer;
+@end
+
+@implementation MyCanvasView
+
+
+- (instancetype)initWithFrame:(NSRect)frame {
+    self = [super initWithFrame:frame];
+    if (self) {
+        _time = 0.0;
+
+        uba::VisualizerConfig config("~/UbaVisualizerConfig.toml");
+        config.boxHeight = 14;
+        config.horizontalScaleValue = 1.0f;
+        auto logger = new uba::LoggerWithWriter(uba::g_consoleLogWriter);
+        auto backend = new uba::VisualizerBackendMac;
+
+        _visualizer = new uba::VisualizerMac(config, *logger, *backend);
+
+        // 60 FPS timer
+        [NSTimer scheduledTimerWithTimeInterval:(1.0/60.0)
+                                         target:self
+                                       selector:@selector(tick:)
+                                       userInfo:nil
+                                        repeats:YES];
+    }
+    return self;
+}
+
+- (void)tick:(NSTimer *)timer {
+    // advance your “time” or simulation here
+    _time += 1.0 / 60.0;
+
+    // Call into C++ per-frame if you like:
+    // cpp_update_example(_time);
+
+    // Request a redraw
+    [self setNeedsDisplay:YES];
+}
+
+
+
+- (BOOL)isFlipped {
+    // Use top-left origin like many 2D APIs (optional)
+    return YES;
+}
+
+- (void)drawRect:(NSRect)dirtyRect {
+   [super drawRect:dirtyRect];
+
+    NSRect bounds = self.bounds;
+
+    [[NSColor whiteColor] setFill];
+    NSRectFill(bounds);
+/*
+    // Example: animate the line using _time
+    [[NSColor blackColor] setStroke];
+    NSBezierPath *path = [NSBezierPath bezierPath];
+
+    CGFloat x = 10 + 50 * sin(_time * 2.0); // wiggle a bit
+
+    [path moveToPoint:NSMakePoint(x, 10)];
+    [path lineToPoint:NSMakePoint(bounds.size.width - x,
+                                  bounds.size.height - 10)];
+    [path stroke];
+*/
+    uba::VisualizerBackendMac::PaintContextMac context;
+    uba::RECT r { 0, 0, (int)bounds.size.width, (int)bounds.size.height };
+    context.clientRect = r;
+    _visualizer->PaintAll(context, r);
+
+    // Call into C++ renderer if you want:
+    cpp_render_example(float(bounds.size.width), float(bounds.size.height));
+}
+
+@end
+
+// Dummy C++ function implementation (replace with your own):
+void cpp_render_example(float width, float height) {
+    // This is normal C++ code, do whatever you want here.
+    // For example, update some state, log, etc.
+    // For actual drawing, pass data back to Objective-C++ or keep state
+    // in a shared object that drawRect uses.
+    (void)width;
+    (void)height;
+}
+
+@interface AppDelegate : NSObject <NSWindowDelegate>
+@end
+
+@implementation AppDelegate
+- (void)windowWillClose:(NSNotification *)notification {
+    [NSApp terminate:nil];
+}
+@end
+
+int main(int argc, const char * argv[]) {
+    @autoreleasepool {
+        [NSApplication sharedApplication];
+        [NSApp setActivationPolicy:NSApplicationActivationPolicyRegular];
+
+        // --- Minimal menu bar with "Quit" (Cmd+Q) ---
+        NSMenu *menubar = [[NSMenu alloc] init];
+        NSMenuItem *appMenuItem = [[NSMenuItem alloc] init];
+        [menubar addItem:appMenuItem];
+        [NSApp setMainMenu:menubar];
+
+        NSMenu *appMenu = [[NSMenu alloc] initWithTitle:@""];
+        NSString *appName = [[NSProcessInfo processInfo] processName];
+        NSMenuItem *quitItem =
+            [[NSMenuItem alloc] initWithTitle:[NSString stringWithFormat:@"Quit %@", appName]
+                                       action:@selector(terminate:)
+                                keyEquivalent:@"q"];
+        // Make sure it's Cmd+Q, not just 'q'
+        [quitItem setKeyEquivalentModifierMask:NSEventModifierFlagCommand];
+
+        [appMenu addItem:quitItem];
+        [appMenuItem setSubmenu:appMenu];
+        // --- end menu setup ---
+
+        // Create the window
+        NSRect frame = NSMakeRect(200, 200, 800, 600);
+        NSWindowStyleMask style =
+            NSWindowStyleMaskTitled |
+            NSWindowStyleMaskClosable |
+            NSWindowStyleMaskResizable |
+            NSWindowStyleMaskMiniaturizable;
+
+        NSWindow *window = [[NSWindow alloc]
+                            initWithContentRect:frame
+                            styleMask:style
+                            backing:NSBackingStoreBuffered
+                            defer:NO];
+
+        AppDelegate *delegate = [[AppDelegate alloc] init];
+        [window setDelegate:delegate];
+        [window setTitle:@"UbaVisualizer"];
+
+        // Attach our canvas view
+        MyCanvasView *view = [[MyCanvasView alloc] initWithFrame:frame];
+        [window setContentView:view];
+        [window makeKeyAndOrderFront:nil];
+
+        // Make sure app shows in Dock and gets focus
+        [NSApp activateIgnoringOtherApps:YES];
+
+        // Enter the Cocoa event loop
+        [NSApp run];
+    }
+    return 0;
+}
\ No newline at end of file
diff --git a/Engine/Source/Programs/UnrealBuildAccelerator/Visualizer/Private/PosixOS/UbaVisualizerApp.cpp b/Engine/Source/Programs/UnrealBuildAccelerator/Visualizer/Private/PosixOS/UbaVisualizerApp.cpp
deleted file mode 100644
diff --git a/Engine/Source/Programs/UnrealBuildAccelerator/Visualizer/Private/UbaTraceReader.cpp b/Engine/Source/Programs/UnrealBuildAccelerator/Visualizer/Private/UbaTraceReader.cpp
--- a/Engine/Source/Programs/UnrealBuildAccelerator/Visualizer/Private/UbaTraceReader.cpp
+++ b/Engine/Source/Programs/UnrealBuildAccelerator/Visualizer/Private/UbaTraceReader.cpp
@@ -115,7 +115,7 @@ namespace uba
 			return false;
 		auto closeFileMapping = MakeGuard([&]() { CloseFileMapping(m_logger, m_memoryHandle, fileName); m_memoryHandle = {}; });
 
-		m_memoryBegin = MapViewOfFile(m_logger, m_memoryHandle, FILE_MAP_READ, 0, 0);
+		m_memoryBegin = MapViewOfFile(m_logger, m_memoryHandle, FILE_MAP_READ, 0, fileSize);
 		if (!m_memoryBegin)
 			return false;
 		auto closeView = MakeGuard([&]() { UnmapViewOfFile(m_logger, m_memoryBegin, 0, TC("TraceReader")); m_memoryBegin = nullptr; m_memoryPos = nullptr; m_memoryEnd = nullptr; });
diff --git a/Engine/Source/Programs/UnrealBuildAccelerator/Visualizer/UbaVisualizer.Build.cs b/Engine/Source/Programs/UnrealBuildAccelerator/Visualizer/UbaVisualizer.Build.cs
--- a/Engine/Source/Programs/UnrealBuildAccelerator/Visualizer/UbaVisualizer.Build.cs
+++ b/Engine/Source/Programs/UnrealBuildAccelerator/Visualizer/UbaVisualizer.Build.cs
@@ -15,5 +15,13 @@ public class UbaVisualizer : ModuleRules
 		PrivateDependencyModuleNames.AddRange(new string[] {
 			"UbaCommon",
 		});
+
+		if (Target.Platform == UnrealTargetPlatform.Mac)
+		{
+			PublicFrameworks.AddRange([
+				"Cocoa",
+				"AppKit"
+			]);
+		}
 	}
 }
