[UBA]
* Added copyright to file

diff --git a/Engine/Source/Programs/UnrealBuildAccelerator/Visualizer/Private/Mac/UbaVisualizerApp.mm b/Engine/Source/Programs/UnrealBuildAccelerator/Visualizer/Private/Mac/UbaVisualizerApp.mm
--- a/Engine/Source/Programs/UnrealBuildAccelerator/Visualizer/Private/Mac/UbaVisualizerApp.mm
+++ b/Engine/Source/Programs/UnrealBuildAccelerator/Visualizer/Private/Mac/UbaVisualizerApp.mm
@@ -1,3 +1,5 @@
+// Copyright Epic Games, Inc. All Rights Reserved.
+
 #import <Cocoa/Cocoa.h>
 
 #include "UbaFile.h"
