[UBA]
* More visualizer refactorings to make it easier to implement other platforms
* Fixed some visual artifacts when using mouse scrollwheel while trace is playing

diff --git a/Engine/Source/Programs/UnrealBuildAccelerator/Visualizer/Private/UbaVisualizer.cpp b/Engine/Source/Programs/UnrealBuildAccelerator/Visualizer/Private/UbaVisualizer.cpp
--- a/Engine/Source/Programs/UnrealBuildAccelerator/Visualizer/Private/UbaVisualizer.cpp
+++ b/Engine/Source/Programs/UnrealBuildAccelerator/Visualizer/Private/UbaVisualizer.cpp
@@ -60,13 +60,13 @@ namespace uba
 			r.right = r.left + width;
 			r.bottom = r.top + height;
 
-			RECT clientRect = context.GetClientRect();
+			SIZE clientSize = context.GetClientSize();
 
-			if (r.right > clientRect.right)
+			if (r.right > clientSize.cx)
 				OffsetRect(&r, -width - 15, 0);
-			if (r.bottom > clientRect.bottom)
+			if (r.bottom > clientSize.cy)
 			{
-				OffsetRect(&r, 0, clientRect.bottom - r.bottom);
+				OffsetRect(&r, 0, clientSize.cy - r.bottom);
 				if (r.top < 0)
 					OffsetRect(&r, 0, -r.top);
 			}
@@ -112,6 +112,28 @@ namespace uba
 
 	////////////////////////////////////////////////////////////////////////////////////////////////////
 
+	bool Visualizer::Stats::operator==(const Stats& o) const
+	{
+		if (memcmp(this, &o, offsetof(Stats, connectionCount)) != 0)
+			return false;
+		if (connectionCount != o.connectionCount)
+			return false;
+		if (drives.size() != o.drives.size())
+			return false;
+		for (auto i=drives.begin(), i2=o.drives.begin(); i!=drives.end(); ++i, ++i2)
+		{
+			if (i->first != i2->first)
+				return false;
+			const Drive& a = i->second;
+			const Drive& b = i2->second;
+			if (a.busyPercent != b.busyPercent || a.readPerSecond != b.readPerSecond || a.writePerSecond != b.writePerSecond)
+				return false;
+		}
+		return true;
+	}
+
+	////////////////////////////////////////////////////////////////////////////////////////////////////
+
 	Visualizer::Visualizer(VisualizerConfig& config, Logger& logger, VisualizerBackend& backend)
 	:	m_logger(logger)
 	,	m_config(config)
@@ -160,11 +182,11 @@ namespace uba
 		return m_timelineFont.height + 8;
 	}
 
-	int Visualizer::GetTimelineTop(const RECT& clientRect)
+	int Visualizer::GetTimelineTop()
 	{
 		int timelineHeight = GetTimelineHeight();
 		int posY = m_contentHeight - timelineHeight;
-		int maxY = int(clientRect.bottom - timelineHeight);
+		int maxY = int(m_clientHeight - timelineHeight);
 		return m_config.LockTimelineToBottom ? maxY : Min(posY, maxY);
 	}
 
@@ -173,6 +195,17 @@ namespace uba
 		return !m_config.showDetailedData && !m_config.showTitleBars && !m_config.showCpuMemStats && !m_config.showNetworkStats && !m_config.showDriveStats && !m_config.showProcessBars;
 	}
 
+	StringBufferBase& Visualizer::GetTitlePrefix(StringBufferBase& out)
+	{
+		out.Clear();
+		out.Append(L"UbaVisualizer");
+		#if UBA_DEBUG
+		out.Append(L" (DEBUG)");
+		#endif
+		out.Append(L" - ");
+		return out;
+	}
+
 	StringBuffer<128> Visualizer::GetWorldTime(u64 time)
 	{
 		return GetWorldTime(TimeToS(time));
@@ -187,7 +220,6 @@ namespace uba
 	{
 		u64 playTime = GetPlayTime();
 
-		RECT clientRect = context.GetClientRect();
 		context.SetFont(m_defaultFont);
 
 		int posY = int(m_scrollPosY);
@@ -195,7 +227,7 @@ namespace uba
 		int processStepY = boxHeight + 1;
 		float scaleX = 50.0f*m_zoomValue*m_config.horizontalScaleValue;
 
-		RECT progressRect = clientRect;
+		RECT progressRect = { 0, 0, m_clientWidth, m_clientHeight };
 		progressRect.left += m_progressRectLeft;
 		progressRect.bottom -= 30;
 
@@ -222,7 +254,7 @@ namespace uba
 
 		if (m_config.showTimeline && !m_traceView.sessions.empty())
 		{
-			int timelineTop = GetTimelineTop(clientRect);
+			int timelineTop = GetTimelineTop();
 			if (pos.y >= timelineTop)// && pos.y < timelineTop + 40)
 			{
 				outResult.section = 3;
@@ -303,7 +335,7 @@ namespace uba
 
 		if (m_config.showActiveProcesses && !m_trace.m_activeProcesses.empty())
 		{
-			PaintActiveProcesses(context, posY, clientRect, [&](TraceView::ProcessLocation& processLocation, u32 boxHeight, bool firstWithHeight)
+			PaintActiveProcesses(context, posY, [&](TraceView::ProcessLocation& processLocation, u32 boxHeight, bool firstWithHeight)
 				{
 					if (pos.y < posY || pos.y > posY + int(boxHeight))
 						return;
@@ -437,8 +469,8 @@ namespace uba
 									u64 updateWrite = pair.second.writeBytes[i];
 									auto& statsDrive = outResult.stats.drives[pair.first];
 									statsDrive.busyPercent = updateBusy;
-									statsDrive.readPerSecond = u64(float(updateRead) / duration);
-									statsDrive.writePerSecond = u64(float(updateWrite) / duration);
+									statsDrive.readPerSecond = u32(float(updateRead) / duration);
+									statsDrive.writePerSecond = u32(float(updateWrite) / duration);
 								}
 								return;
 							}
@@ -853,7 +885,7 @@ namespace uba
 
 			m_activeProcessCountHistory[m_activeProcessCountHistoryIterator++ % sizeof_array(m_activeProcessCountHistory)] = u32(m_trace.m_activeProcesses.size());
 
-			PaintActiveProcesses(context, posY, clientRect, [&](TraceView::ProcessLocation& processLocation, u32 boxHeight, bool firstWithHeight)
+			PaintActiveProcesses(context, posY, [&](TraceView::ProcessLocation& processLocation, u32 boxHeight, bool firstWithHeight)
 				{
 					auto& session = m_trace.GetSession(m_traceView, processLocation.sessionIndex);
 					TraceView::Process& process = session.processors[processLocation.processorIndex].processes[processLocation.processIndex];
@@ -1454,7 +1486,7 @@ namespace uba
 		float timelineSelected = m_timelineSelected;
 
 		if (m_config.showTimeline && !m_traceView.sessions.empty())
-			PaintTimeline(context, clientRect);
+			PaintTimeline(context);
 
 		if (!m_filterString.empty())
 		{
@@ -1478,7 +1510,7 @@ namespace uba
 		{
 			int posX = int(m_scrollPosX) + progressRect.left;
 			int left = posX + int(timelineSelected * scaleX);
-			int timelineTop = GetTimelineTop(clientRect);
+			int timelineTop = GetTimelineTop();
 
 			// TODO: Draw line up
 			context.Line(left, 2, left, timelineTop);
@@ -1995,7 +2027,7 @@ namespace uba
 			context.FillRect(main, brush);
 	}
 
-	void Visualizer::PaintActiveProcesses(PaintContext& context, int& posY, const RECT& clientRect, const Function<void(TraceView::ProcessLocation&, u32, bool)>& drawProcess)
+	void Visualizer::PaintActiveProcesses(PaintContext& context, int& posY, const Function<void(TraceView::ProcessLocation&, u32, bool)>& drawProcess)
 	{
 		context.SetFont(m_processFont);
 		int startPosY = posY;
@@ -2014,12 +2046,12 @@ namespace uba
 
 		Font& activeFont = context.GetFont();
 
-		u32 maxHeight = clientRect.bottom;
+		u32 maxHeight = m_clientHeight;
 
 		bool fillHeight = ActiveProcessesShouldFillHeight();
 		if (fillHeight)
 		{
-			maxHeight = clientRect.bottom - posY;
+			maxHeight = m_clientHeight - posY;
 			if (m_config.showTimeline)
 				maxHeight -= GetTimelineHeight();
 		}
@@ -2097,10 +2129,10 @@ namespace uba
 		context.SetFont(m_defaultFont);
 	}
 
-	void Visualizer::PaintTimeline(PaintContext& context, const RECT& clientRect)
+	void Visualizer::PaintTimeline(PaintContext& context)
 	{
 		context.SetFont(m_timelineFont);
-		int top = GetTimelineTop(clientRect);
+		int top = GetTimelineTop();
 		float timeScale = (m_config.horizontalScaleValue*m_zoomValue)*50.0f;
 		float startOffset = ((m_scrollPosX/timeScale) - float(int(m_scrollPosX/timeScale))) * timeScale;
 		int index = -int(startOffset/timeScale);
@@ -2123,7 +2155,7 @@ namespace uba
 
 		int lineStepSize = textStepSize / 5;
 
-		RECT progressRect = clientRect;
+		RECT progressRect = { 0, 0, m_clientWidth, m_clientHeight };
 		progressRect.left += m_progressRectLeft;
 
 		context.SetPen(m_textPen);
@@ -2133,7 +2165,7 @@ namespace uba
 		while (true)
 		{
 			int pos = progressRect.left + int(startOffset + float(index)*timeScale);
-			if (pos >= clientRect.right)
+			if (pos >= m_clientWidth)
 				break;
 
 			int lineBottom = top + 5;
@@ -2465,5 +2497,171 @@ namespace uba
 	{
 	}
 
+	void Visualizer::Redraw(bool now)
+	{
+	}
+
+	void Visualizer::UpdateScrollbars(bool redraw)
+	{
+	}
+
+	void Visualizer::PostNewTitle(const StringView& title)
+	{
+	}
+
+	void Visualizer::StopUpdate()
+	{
+	}
+
+	bool Visualizer::NewTrace(u32 replay, bool paused)
+	{
+		m_replay = replay;
+		m_paused = paused;
+		m_autoScroll = true;
+		m_scrollPosX = 0;
+		m_scrollPosY = 0;
+		Reset();
+		StringBuffer<> title;
+		GetTitlePrefix(title);
+
+		auto g = MakeGuard([&]()
+			{
+				Redraw(true);
+				UpdateScrollbars(true);
+			});
+
+		if (m_client)
+		{
+			if (!m_trace.StartReadClient(m_traceView, *m_client))
+			{
+				m_clientDisconnect.Set();
+				return false;
+			}
+			m_namedTrace.Clear().Append(m_newTraceName);
+			m_traceView.finished = false;
+		}
+		else if (!m_fileName.IsEmpty())
+		{
+			m_trace.ReadFile(m_traceView, m_fileName.data, m_replay != 0);
+			m_traceView.finished = m_replay == 0;
+			PostNewTitle(GetTitlePrefix(title).Appendf(L"%s (v%u) - %s", m_fileName.data, m_traceView.version, GetWorldTime((u64)0).data));
+		}
+		else if (m_usingNamed)
+		{
+			if (!m_trace.StartReadNamed(m_traceView, m_newTraceName.data, true, m_replay != 0))
+				return false;
+			m_namedTrace.Clear().Append(m_newTraceName);
+			m_traceView.finished = false;
+			PostNewTitle(GetTitlePrefix(title).Appendf(L"%s (Listening for new sessions on channel '%s')", m_namedTrace.data, m_listenChannel.data));
+		}
+
+		return true;
+	}
+
+	bool Visualizer::UpdateTrace()
+	{
+		bool changed = false;
+		if (!m_paused)
+		{
+			u64 timeOffset = (GetTime() - m_startTime - m_pauseTime) * m_replay;
+			if (!m_fileName.IsEmpty())
+			{
+				if (m_replay)
+					m_trace.UpdateReadFile(m_traceView, timeOffset, changed);
+			}
+			else if (m_client)
+			{
+				if (!m_trace.UpdateReadClient(m_traceView, *m_client, changed))
+					m_clientDisconnect.Set();
+			}
+			else if (m_usingNamed)
+			{
+				if (!m_trace.UpdateReadNamed(m_traceView, m_replay ? timeOffset : ~u64(0), changed))
+					if (m_listenTimeout.IsCreated())
+						m_listenTimeout.Set();
+			}
+		}
+
+		if (m_traceView.finished)
+		{
+			m_autoScroll = false;
+			StopUpdate();
+			changed = true;
+		}
+
+		changed = UpdateAutoscroll() || changed;
+		changed = UpdateSelection() || changed;
+
+		return changed;
+	}
+
+	bool Visualizer::UpdateAutoscroll()
+	{
+		if (!m_autoScroll)
+			return false;
+
+		u64 playTime = GetPlayTime();
+
+		int right = m_clientWidth;
+		if (right == 0)
+			return false;
+
+		float timeS = TimeToS(playTime);
+
+		if (m_config.AutoScaleHorizontal)
+		{
+			m_scrollPosX = 0;
+			timeS = Max(timeS, 20.0f/m_zoomValue);
+			m_config.horizontalScaleValue = Max(float(right - m_progressRectLeft - 2)/(m_zoomValue*timeS*50.0f), 0.001f);
+			return true;
+		}
+		else
+		{
+			float oldScrollPosX = m_scrollPosX;
+			m_scrollPosX = Min(0.0f, (float)right - timeS*50.0f*m_config.horizontalScaleValue*m_zoomValue - float(m_progressRectLeft));
+			return oldScrollPosX != m_scrollPosX;
+		}
+	}
+
+	bool Visualizer::UpdateSelection()
+	{
+		return false;
+	}
+
+	void Visualizer::Reset()
+	{
+		m_backend.ResetCaches();
+		m_contentWidth = 0;
+		m_contentHeight = 0;
+		//m_autoScroll = true;
+		//m_scrollPosX = 0;
+		//m_scrollPosY = 0;
+		//m_zoomValue = 0.75f;
+		//m_horizontalScaleValue = 1.0f;
+
+		m_startTime = GetTime();
+		m_pauseTime = 0;
+
+		Unselect();
+	}
+
+	bool Visualizer::Unselect()
+	{
+		if (m_processSelected || m_sessionSelectedIndex != ~0u || m_statsSelected || m_timelineSelected != 0 || m_fetchedFilesSelected != ~0u || m_workSelected || !m_hyperLinkSelected.empty())
+		{
+			m_processSelected = false;
+			m_sessionSelectedIndex = ~0u;
+			m_statsSelected = false;
+			m_activeProcessGraphSelected = false;
+			m_buttonSelected = ~0u;
+			m_timelineSelected = 0;
+			m_fetchedFilesSelected = ~0u;
+			m_workSelected = false;
+			m_hyperLinkSelected.clear();
+			return true;
+		}
+		return false;
+	}
+
 	////////////////////////////////////////////////////////////////////////////////////////////////////
 }
diff --git a/Engine/Source/Programs/UnrealBuildAccelerator/Visualizer/Private/UbaVisualizer.h b/Engine/Source/Programs/UnrealBuildAccelerator/Visualizer/Private/UbaVisualizer.h
--- a/Engine/Source/Programs/UnrealBuildAccelerator/Visualizer/Private/UbaVisualizer.h
+++ b/Engine/Source/Programs/UnrealBuildAccelerator/Visualizer/Private/UbaVisualizer.h
@@ -18,7 +18,8 @@ namespace uba
 	protected:
 		u64 GetPlayTime();
 		int GetTimelineHeight();
-		int GetTimelineTop(const RECT& clientRect);
+		int GetTimelineTop();
+		StringBufferBase& GetTitlePrefix(StringBufferBase& out);
 		StringBuffer<128> GetWorldTime(u64 time);
 		StringBuffer<128> GetWorldTime(float seconds);
 		bool ActiveProcessesShouldFillHeight();
@@ -30,8 +31,8 @@ namespace uba
 		void HitTest(HitTestResult& outResult, PaintContext& context, const POINT& pos);
 		void PaintAll(PaintContext& context, const RECT& clientRect);
 		void PaintProcessRect(PaintContext& context, TraceView::Process& process, RECT rect, const RECT& progressRect, bool selected, bool writingBitmap);
-		void PaintActiveProcesses(PaintContext& context, int& posY, const RECT& clientRect, const Function<void(TraceView::ProcessLocation&, u32, bool)>& drawProcess);
-		void PaintTimeline(PaintContext& context, const RECT& clientRect);
+		void PaintActiveProcesses(PaintContext& context, int& posY, const Function<void(TraceView::ProcessLocation&, u32, bool)>& drawProcess);
+		void PaintTimeline(PaintContext& context);
 		using DrawTextFunc = Function<void(const StringBufferBase& text, RECT& rect, u32* outWidth)>;
 		void PaintDetailedStats(PaintContext& context, int& posY, const RECT& progressRect, TraceView::Session& session, bool isRemote, u64 playTime, const DrawTextFunc& drawTextFunc);
 		void PrintCacheWriteStats(Logger& logger, u32 processId);
@@ -40,7 +41,17 @@ namespace uba
 
 		virtual void PaintEmptyScreen(PaintContext& context);
 		virtual void FinishIncompatible();
-
+		virtual void Redraw(bool now);
+		virtual void UpdateScrollbars(bool redraw);
+		virtual void PostNewTitle(const StringView& title);
+		virtual bool UpdateSelection();
+		virtual void StopUpdate();
+
+		bool NewTrace(u32 replay, bool paused);
+		bool UpdateTrace();
+		bool UpdateAutoscroll();
+		void Reset();
+		bool Unselect();
 		void UpdateFont(Font& font, int height, bool createUnderline);
 
 		enum VisualizerFlag
@@ -61,15 +72,17 @@ namespace uba
 			u64 recvBytes = 0;
 			u64 sendBytes = 0;
 			u64 procActive = 0;
-			u16 connectionCount = 0;
 			float cpuLoad = 0;
+			u16 connectionCount = 0;
+
 			struct Drive
 			{
 				u8 busyPercent;
-				u64 readPerSecond = 0;
-				u64 writePerSecond = 0;
+				u32 readPerSecond = 0;
+				u32 writePerSecond = 0;
 			};
 			Map<char, Drive> drives;
+			bool operator==(const Stats& o) const;
 		};
 
 		struct HitTestResult
@@ -111,6 +124,19 @@ namespace uba
 		TraceView m_traceView;
 		VisualizerBackend& m_backend;
 
+		StringBuffer<256> m_namedTrace;
+		StringBuffer<256> m_fileName;
+		NetworkClient* m_client = nullptr;
+		Event m_clientDisconnect;
+		Event m_listenTimeout;
+		StringBuffer<256>m_listenChannel;
+		StringBuffer<256> m_newTraceName;
+		bool m_usingNamed = false;
+
+
+		int m_clientWidth;
+		int m_clientHeight;
+
 		ProcessBrushes m_processBrushes[2]; // Non-selected and selected
 
 		static constexpr int GraphHeight = 30;
@@ -149,6 +175,7 @@ namespace uba
 		u32 m_activeProcessCountHistory[5];
 		u32 m_activeProcessCountHistoryIterator = 0;
 
+		bool m_autoScroll = true;
 		bool m_paused = false;
 		u32 m_replay = 0;
 		u64 m_startTime = 0;
diff --git a/Engine/Source/Programs/UnrealBuildAccelerator/Visualizer/Private/UbaVisualizerBackend.h b/Engine/Source/Programs/UnrealBuildAccelerator/Visualizer/Private/UbaVisualizerBackend.h
--- a/Engine/Source/Programs/UnrealBuildAccelerator/Visualizer/Private/UbaVisualizerBackend.h
+++ b/Engine/Source/Programs/UnrealBuildAccelerator/Visualizer/Private/UbaVisualizerBackend.h
@@ -37,8 +37,7 @@ namespace uba
 
 		struct PaintContext
 		{
-			virtual RECT GetClientRect() = 0;
-
+			virtual SIZE GetClientSize() = 0;
 			virtual void SetBkMode(int mode) = 0;
 			virtual void SetPen(HPEN pen) = 0;
 			virtual void SetBrush(HBRUSH b) = 0;
diff --git a/Engine/Source/Programs/UnrealBuildAccelerator/Visualizer/Private/Windows/UbaVisualizerBackendGDI.h b/Engine/Source/Programs/UnrealBuildAccelerator/Visualizer/Private/Windows/UbaVisualizerBackendGDI.h
--- a/Engine/Source/Programs/UnrealBuildAccelerator/Visualizer/Private/Windows/UbaVisualizerBackendGDI.h
+++ b/Engine/Source/Programs/UnrealBuildAccelerator/Visualizer/Private/Windows/UbaVisualizerBackendGDI.h
@@ -17,7 +17,7 @@ namespace uba
 
 		struct PaintContextGDI : public PaintContext
 		{
-			virtual RECT GetClientRect() override { return clientRect; }
+			virtual SIZE GetClientSize() override { return clientSize; }
 			virtual void SetBkMode(int mode) override { ::SetBkMode(hdc, mode); }
 			virtual void SetTextColor(COLORREF color) override { ::SetTextColor(hdc, color); }
 			virtual COLORREF GetTextColor() override { return ::GetTextColor(hdc); }
@@ -91,9 +91,9 @@ namespace uba
 					SIZE textSize = GetTextSize(rows[i]);
 
 					RECT rect = { 0, 0, 0, 0 };
-					int top = (clientRect.bottom - textSize.cy)/2 + i*20;
-					rect.left = (clientRect.right - textSize.cx)/2;
-					rect.right = (clientRect.right + textSize.cx)/2;
+					int top = (clientSize.cy - textSize.cy)/2 + i*20;
+					rect.left = (clientSize.cx - textSize.cx)/2;
+					rect.right = (clientSize.cx + textSize.cx)/2;
 					rect.bottom = top + textSize.cy;
 					rect.top = top;
 
@@ -137,7 +137,7 @@ namespace uba
 						lastSelectedBitmap = process.bitmap;
 					}
 
-					PaintContextGDI textCacheContext{owner, NULL, textDC, clientRect };
+					PaintContextGDI textCacheContext{owner, NULL, textDC, clientSize };
 					func(textCacheContext);
 				}
 
@@ -184,13 +184,13 @@ namespace uba
 				return p;
 			}
 
-			PaintContextGDI(VisualizerBackendGDI& o, HWND hw, HDC dc, const RECT& cr, HDC tdc = 0) : owner(o), hwnd(hw), hdc(dc), clientRect(cr), textDC(tdc) {}
+			PaintContextGDI(VisualizerBackendGDI& o, HWND hw, HDC dc, SIZE cs, HDC tdc = 0) : owner(o), hwnd(hw), hdc(dc), clientSize(cs), textDC(tdc) {}
 
 
 			VisualizerBackendGDI& owner;
 			HWND hwnd;
 			HDC hdc;
-			const RECT& clientRect;
+			SIZE clientSize;
 
 			Font activeFont;
 
diff --git a/Engine/Source/Programs/UnrealBuildAccelerator/Visualizer/Private/Windows/UbaVisualizerWin.cpp b/Engine/Source/Programs/UnrealBuildAccelerator/Visualizer/Private/Windows/UbaVisualizerWin.cpp
--- a/Engine/Source/Programs/UnrealBuildAccelerator/Visualizer/Private/Windows/UbaVisualizerWin.cpp
+++ b/Engine/Source/Programs/UnrealBuildAccelerator/Visualizer/Private/Windows/UbaVisualizerWin.cpp
@@ -253,52 +253,6 @@ namespace uba
 		m_locked = lock;
 	}
 
-	StringBufferBase& VisualizerWin::GetTitlePrefix(StringBufferBase& out)
-	{
-		out.Clear();
-		out.Append(L"UbaVisualizer");
-#if UBA_DEBUG
-		out.Append(L" (DEBUG)");
-#endif
-		out.Append(L" - ");
-		return out;
-	}
-
-	bool VisualizerWin::Unselect()
-	{
-		if (m_processSelected || m_sessionSelectedIndex != ~0u || m_statsSelected || m_timelineSelected != 0 || m_fetchedFilesSelected != ~0u || m_workSelected || !m_hyperLinkSelected.empty())
-		{
-			m_processSelected = false;
-			m_sessionSelectedIndex = ~0u;
-			m_statsSelected = false;
-			m_activeProcessGraphSelected = false;
-			m_buttonSelected = ~0u;
-			m_timelineSelected = 0;
-			m_fetchedFilesSelected = ~0u;
-			m_workSelected = false;
-			m_hyperLinkSelected.clear();
-			return true;
-		}
-		return false;
-	}
-
-	void VisualizerWin::Reset()
-	{
-		m_backend.ResetCaches();
-		m_contentWidth = 0;
-		m_contentHeight = 0;
-		//m_autoScroll = true;
-		//m_scrollPosX = 0;
-		//m_scrollPosY = 0;
-		//m_zoomValue = 0.75f;
-		//m_horizontalScaleValue = 1.0f;
-
-		m_startTime = GetTime();
-		m_pauseTime = 0;
-
-		Unselect();
-	}
-
 	void VisualizerWin::InitBrushes()
 	{
 		if (m_config.DarkMode)
@@ -491,10 +445,13 @@ namespace uba
 		if (m_config.DarkMode)
 			UpdateTheme();
 
+		RECT r;
+		GetClientRect(m_hwnd, &r);
+		m_clientWidth = r.right;
+		m_clientHeight = r.bottom;
+
 		{
-			RECT clientRect;
-			GetClientRect(m_hwnd, &clientRect);
-			VisualizerBackendGDI::PaintContextGDI context((VisualizerBackendGDI&)m_backend, NULL, NULL, clientRect);
+			VisualizerBackendGDI::PaintContextGDI context((VisualizerBackendGDI&)m_backend, NULL, NULL, { m_clientWidth, m_clientHeight });
 			HitTestResult res;
 			HitTest(res, context, { -1, -1 });
 		}
@@ -721,28 +678,30 @@ namespace uba
 			m_activeProcessCountHistory[i] = activeProcessCount;
 	}
 
-	void VisualizerWin::PaintClient(const Function<void(HDC hdc, HDC memDC, RECT& clientRect)>& paintFunc)
+	void VisualizerWin::StopUpdate()
 	{
-		HDC hdc = GetDC(m_hwnd);
+		KillTimer(m_hwnd, 0);
+	}
 
-		RECT rect;
-		GetClientRect(m_hwnd, &rect);
+	void VisualizerWin::PaintClient(const Function<void(HDC hdc, HDC memDC)>& paintFunc)
+	{
+		HDC hdc = GetDC(m_hwnd);
 
 		HDC memDC = CreateCompatibleDC(hdc);
 
-		if (!EqualRect(&m_cachedBitmapRect, &rect))
+		if (m_cachedBitmapSize.cx != m_clientWidth || m_cachedBitmapSize.cy != m_clientHeight)
 		{
 			if (m_cachedBitmap)
 				DeleteObject(m_cachedBitmap);
-			m_cachedBitmap = CreateCompatibleBitmap(hdc, rect.right - rect.left, rect.bottom - rect.top);
-			m_cachedBitmapRect = rect;
+			m_cachedBitmap = CreateCompatibleBitmap(hdc, m_clientWidth, m_clientHeight);
+			m_cachedBitmapSize = { m_clientWidth, m_clientHeight };
 		}
 		HGDIOBJ oldBmp = SelectObject(memDC, m_cachedBitmap);
 
 		if (!m_isInPaint)
 		{
 			m_isInPaint = true;
-			paintFunc(hdc, memDC, rect);
+			paintFunc(hdc, memDC);
 			m_isInPaint = false;
 		}
 
@@ -806,36 +765,6 @@ namespace uba
 			Redraw(false);
 	}
 
-	bool VisualizerWin::UpdateAutoscroll()
-	{
-		if (!m_autoScroll)
-			return false;
-
-		u64 playTime = GetPlayTime();
-
-
-		RECT rect;
-		GetClientRect(m_hwnd, &rect);
-		if (rect.right == 0)
-			return false;
-
-		float timeS = TimeToS(playTime);
-
-		if (m_config.AutoScaleHorizontal)
-		{
-			m_scrollPosX = 0;
-			timeS = Max(timeS, 20.0f/m_zoomValue);
-			m_config.horizontalScaleValue = Max(float(rect.right - m_progressRectLeft - 2)/(m_zoomValue*timeS*50.0f), 0.001f);
-			return true;
-		}
-		else
-		{
-			float oldScrollPosX = m_scrollPosX;
-			m_scrollPosX = Min(0.0f, (float)rect.right - timeS*50.0f*m_config.horizontalScaleValue*m_zoomValue - float(m_progressRectLeft));
-			return oldScrollPosX != m_scrollPosX;
-		}
-	}
-
 	bool VisualizerWin::UpdateSelection()
 	{
 		if (!m_mouseOverWindow || m_dragToScrollCounter > 0)
@@ -846,16 +775,14 @@ namespace uba
 
 		HitTestResult res;
 		{
-			RECT clientRect;
-			GetClientRect(m_hwnd, &clientRect);
-			VisualizerBackendGDI::PaintContextGDI context((VisualizerBackendGDI&)m_backend, NULL, NULL, clientRect);
+			VisualizerBackendGDI::PaintContextGDI context((VisualizerBackendGDI&)m_backend, NULL, NULL, { m_clientWidth, m_clientHeight });
 			HitTest(res, context, pos);
 		}
 		m_activeSection = res.section;
 
 		if (res.processSelected == m_processSelected && res.processLocation == m_processSelectedLocation &&
 			res.sessionSelectedIndex == m_sessionSelectedIndex &&
-			res.statsSelected == m_statsSelected && memcmp(&res.stats, &m_stats, sizeof(Stats)) == 0 &&
+			res.statsSelected == m_statsSelected && res.stats == m_stats &&
 			res.buttonSelected == m_buttonSelected && res.timelineSelected == m_timelineSelected &&
 			res.activeProcessGraphSelected == m_activeProcessGraphSelected && 
 			res.activeProcessCount == m_activeProcessCount && 
@@ -882,15 +809,12 @@ namespace uba
 
 	void VisualizerWin::UpdateScrollbars(bool redraw)
 	{
-		RECT rect;
-		GetClientRect(m_hwnd, &rect);
-
 		SCROLLINFO si;
 		si.cbSize = sizeof(SCROLLINFO);
 		si.fMask = SIF_ALL | SIF_DISABLENOSCROLL;
 		si.nMin = 0;
 		si.nMax = m_contentHeight;
-		si.nPage = int(rect.bottom);
+		si.nPage = int(m_clientHeight);
 		si.nPos = -int(m_scrollPosY);
 		si.nTrackPos = 0;
 
@@ -934,7 +858,7 @@ namespace uba
 			}
 
 			si.nMax = m_contentWidth;
-			si.nPage = rect.right;
+			si.nPage = m_clientWidth;
 			si.nPos = -int(m_scrollPosX);
 			SetScrollInfo(m_hwnd, SB_HORZ, &si, redraw);
 		}
@@ -984,47 +908,8 @@ namespace uba
 		}
 		case WM_NEWTRACE:
 		{
-			m_replay = u32(wParam);
-			m_paused = lParam;
-			m_autoScroll = true;
-			m_scrollPosX = 0;
-			m_scrollPosY = 0;
-			Reset();
-			StringBuffer<> title;
-			GetTitlePrefix(title);
-
-			auto g = MakeGuard([&]()
-				{
-					Redraw(true);
-					UpdateScrollbars(true);
-				});
-
-			if (m_client)
-			{
-				if (!m_trace.StartReadClient(m_traceView, *m_client))
-				{
-					m_clientDisconnect.Set();
-					return false;
-				}
-				m_namedTrace.Clear().Append(m_newTraceName);
-				m_traceView.finished = false;
-			}
-			else if (!m_fileName.IsEmpty())
-			{
-				m_trace.ReadFile(m_traceView, m_fileName.data, m_replay != 0);
-				m_traceView.finished = m_replay == 0;
-				PostNewTitle(GetTitlePrefix(title).Appendf(L"%s (v%u) - %s", m_fileName.data, m_traceView.version, GetWorldTime((u64)0).data));
-			}
-			else if (m_usingNamed)
-			{
-				if (!m_trace.StartReadNamed(m_traceView, m_newTraceName.data, true, m_replay != 0))
-					return false;
-				m_namedTrace.Clear().Append(m_newTraceName);
-				m_traceView.finished = false;
-				PostNewTitle(GetTitlePrefix(title).Appendf(L"%s (Listening for new sessions on channel '%s')", m_namedTrace.data, m_listenChannel.data));
-			}
-
-			SetTimer(m_hwnd, 0, 200, NULL);
+			if (NewTrace(u32(wParam), lParam != 0))
+				SetTimer(m_hwnd, 0, 200, NULL);
 			return 0;
 		}
 
@@ -1048,8 +933,9 @@ namespace uba
 		case WM_PAINT:
 		{
 			u64 start = GetTime();
-			PaintClient([&](HDC hdc, HDC memDC, RECT& rect)
+			PaintClient([&](HDC hdc, HDC memDC)
 				{
+					RECT rect = { 0, 0, m_clientWidth, m_clientHeight };
 					FillRect(memDC, &rect, m_backgroundBrush);
 					SetBkColor(memDC, m_config.DarkMode ? RGB(70, 70, 70) : RGB(180, 180, 180));
 
@@ -1062,7 +948,7 @@ namespace uba
 					HBITMAP nullBmp = CreateCompatibleBitmap(memDC, 1, 1);
 					HBITMAP oldBmp = (HBITMAP)SelectObject(textDC, nullBmp);
 
-					VisualizerBackendGDI::PaintContextGDI context((VisualizerBackendGDI&)m_backend, m_hwnd, memDC, rect, textDC);
+					VisualizerBackendGDI::PaintContextGDI context((VisualizerBackendGDI&)m_backend, m_hwnd, memDC, { m_clientWidth, m_clientHeight }, textDC);
 					PaintAll(context, rect);
 
 					SelectObject(textDC, oldBmp);
@@ -1082,6 +968,10 @@ namespace uba
 			int width = LOWORD(lParam);
 			if (m_contentWidth && m_contentWidth + int(m_scrollPosX) < width)
 				m_scrollPosX = float(Min(0, width - m_contentWidth));
+			RECT r;
+			GetClientRect(m_hwnd, &r);
+			m_clientWidth = r.right;
+			m_clientHeight = r.bottom;
 			UpdateScrollbars(true);
 			break;
 		}
@@ -1091,37 +981,7 @@ namespace uba
 				break;
 			m_handlingTimer = true;
 
-			bool changed = false;
-			if (!m_paused)
-			{
-				u64 timeOffset = (GetTime() - m_startTime - m_pauseTime) * m_replay;
-				if (!m_fileName.IsEmpty())
-				{
-					if (m_replay)
-						m_trace.UpdateReadFile(m_traceView, timeOffset, changed);
-				}
-				else if (m_client)
-				{
-					if (!m_trace.UpdateReadClient(m_traceView, *m_client, changed))
-						m_clientDisconnect.Set();
-				}
-				else if (m_usingNamed)
-				{
-					if (!m_trace.UpdateReadNamed(m_traceView, m_replay ? timeOffset : ~u64(0), changed))
-						if (m_listenTimeout.IsCreated())
-							m_listenTimeout.Set();
-				}
-			}
-
-			if (m_traceView.finished)
-			{
-				m_autoScroll = false;
-				KillTimer(m_hwnd, 0);
-				changed = true;
-			}
-
-			changed = UpdateAutoscroll() || changed;
-			changed = UpdateSelection() || changed;
+			bool changed = UpdateTrace();
 			if (changed && !IsIconic(m_hwnd))
 			{
 				UpdateScrollbars(true);
@@ -1148,9 +1008,6 @@ namespace uba
 			{
 				if (m_activeSection == 2 || !controlDown) // process bars
 				{
-					RECT r;
-					GetClientRect(hWnd, &r);
-
 					// Use mouse cursor as scroll anchor point
 					POINT cursorPos = {};
 					GetCursorPos(&cursorPos);
@@ -1194,10 +1051,11 @@ namespace uba
 					UpdateAutoscroll();
 					UpdateSelection();
 
-					int minScroll = r.right - m_contentWidth;
+					int minScroll = m_clientWidth - m_contentWidth;
 					m_scrollPosX = Min(0.0f, Max(m_scrollPosX, float(minScroll)));
-					m_scrollPosY = Min(0.0f, Max(m_scrollPosY, float(r.bottom - m_contentHeight)));
+					m_scrollPosY = Min(0.0f, Max(m_scrollPosY, float(m_clientHeight - m_contentHeight)));
 
+					UpdateAutoscroll();
 					//if (!m_traceView.finished && m_scrollPosX <= minScroll)
 					//	m_autoScroll = true;
 
@@ -1229,11 +1087,9 @@ namespace uba
 			}
 			else
 			{
-				RECT r;
-				GetClientRect(hWnd, &r);
 				float oldScrollY = m_scrollPosY;
 				m_scrollPosY = m_scrollPosY + float(delta);
-				m_scrollPosY = Min(Max(m_scrollPosY, float(r.bottom - m_contentHeight)), 0.0f);
+				m_scrollPosY = Min(Max(m_scrollPosY, float(m_clientHeight - m_contentHeight)), 0.0f);
 				if (oldScrollY != m_scrollPosY)
 				{
 					UpdateScrollbars(true);
@@ -1254,19 +1110,16 @@ namespace uba
 			POINT pos{ p.x, p.y };
 			if (m_dragToScrollCounter > 0)
 			{
-				RECT r;
-				GetClientRect(hWnd, &r);
-
-				if (m_contentHeight <= r.bottom)
+				if (m_contentHeight <= m_clientHeight)
 					m_scrollPosY = 0;
 				else
-					m_scrollPosY = Max(Min(m_scrollAtAnchorY + float(pos.y - m_mouseAnchor.y), 0.0f), float(r.bottom - m_contentHeight));
+					m_scrollPosY = Max(Min(m_scrollAtAnchorY + float(pos.y - m_mouseAnchor.y), 0.0f), float(m_clientHeight - m_contentHeight));
 
-				if (m_contentWidth <= r.right)
+				if (m_contentWidth <= m_clientWidth)
 					m_scrollPosX = 0;
 				else
 				{
-					int minScroll = r.right - m_contentWidth;
+					int minScroll = m_clientWidth - m_contentWidth;
 					m_scrollPosX = Max(Min(m_scrollAtAnchorX + float(pos.x - m_mouseAnchor.x), 0.0f), float(minScroll));
 					if (!m_traceView.finished && m_scrollPosX <= float(minScroll))
 						m_autoScroll = true;
@@ -1311,11 +1164,9 @@ namespace uba
 		{
 			if (!m_parentHwnd && m_traceView.sessions.empty())
 			{
-				RECT r;
-				GetClientRect(hWnd, &r);
-				int centerHrz = r.right/2;
-				int centerVrt = r.bottom/2;
-				r = { centerHrz, centerVrt, centerHrz, centerVrt };
+				int centerHrz = m_clientWidth/2;
+				int centerVrt = m_clientHeight/2;
+				RECT r = { centerHrz, centerVrt, centerHrz, centerVrt };
 				InflateRect(&r, 180, 40);
 				POINTS p = MAKEPOINTS(lParam);
 				if (PtInRect(&r, { p.x, p.y }))
@@ -1344,9 +1195,7 @@ namespace uba
 				bool* values = &m_config.showProgress;
 				values[m_buttonSelected] = !values[m_buttonSelected];
 
-				RECT clientRect;
-				GetClientRect(m_hwnd, &clientRect);
-				VisualizerBackendGDI::PaintContextGDI context((VisualizerBackendGDI&)m_backend, NULL, NULL, clientRect);
+				VisualizerBackendGDI::PaintContextGDI context((VisualizerBackendGDI&)m_backend, NULL, NULL, { m_clientWidth, m_clientHeight });
 				HitTestResult res;
 				HitTest(res, context, { -1, -1 });
 				UpdateScrollbars(true);
@@ -1395,15 +1244,13 @@ namespace uba
 						m_pauseTime = 0;
 					}
 
-					RECT r;
-					GetClientRect(hWnd, &r);
-
 					HitTestResult res;
-					VisualizerBackendGDI::PaintContextGDI context((VisualizerBackendGDI&)m_backend, NULL, NULL, r);
+					VisualizerBackendGDI::PaintContextGDI context((VisualizerBackendGDI&)m_backend, NULL, NULL, { m_clientWidth, m_clientHeight });
 					HitTest(res, context, { -1, -1 });
 
-					m_scrollPosX = Min(Max(m_scrollPosX, float(r.right - m_contentWidth)), 0.0f);
-					m_scrollPosY = Min(0.0f, Max(m_scrollPosY, float(r.bottom - m_contentHeight)));
+					m_scrollPosX = Min(Max(m_scrollPosX, float(m_clientWidth - m_contentWidth)), 0.0f);
+					m_scrollPosY = Min(0.0f, Max(m_scrollPosY, float(m_clientHeight - m_contentHeight)));
+					UpdateAutoscroll();
 
 					UpdateScrollbars(true);
 					Redraw(true);
@@ -1439,11 +1286,9 @@ namespace uba
 			bool useHand = false;
 			if (!m_parentHwnd && m_traceView.sessions.empty())
 			{
-				RECT r;
-				GetClientRect(hWnd, &r);
-				int centerHrz = r.right/2;
-				int centerVrt = r.bottom/2;
-				r = { centerHrz, centerVrt, centerHrz, centerVrt };
+				int centerHrz = m_clientWidth/2;
+				int centerVrt = m_clientHeight/2;
+				RECT r = { centerHrz, centerVrt, centerHrz, centerVrt };
 				InflateRect(&r, 180, 40);
 				POINT pt;
 				GetCursorPos(&pt);
@@ -1731,8 +1576,6 @@ namespace uba
 		}
 		case WM_VSCROLL:
 		{
-			RECT r;
-			GetClientRect(hWnd, &r);
 			float oldScrollY = m_scrollPosY;
 
 			// HIWORD(wParam) only carries 16-bits, so use GetScrollInfo for larger scroll areas
@@ -1748,10 +1591,10 @@ namespace uba
 				m_scrollPosY = -float(scrollInfo.nTrackPos);
 				break;
 			case SB_PAGEDOWN:
-				m_scrollPosY = m_scrollPosY - float(r.bottom);
+				m_scrollPosY = m_scrollPosY - float(m_clientHeight);
 				break;
 			case SB_PAGEUP:
-				m_scrollPosY = m_scrollPosY + float(r.bottom);
+				m_scrollPosY = m_scrollPosY + float(m_clientHeight);
 				break;
 			case SB_LINEDOWN:
 				m_scrollPosY -= 30;
@@ -1760,7 +1603,7 @@ namespace uba
 				m_scrollPosY += 30;
 				break;
 			}
-			m_scrollPosY = Min(Max(m_scrollPosY, float(r.bottom - m_contentHeight)), 0.0f);
+			m_scrollPosY = Min(Max(m_scrollPosY, float(m_clientHeight - m_contentHeight)), 0.0f);
 
 			if (oldScrollY != m_scrollPosY)
 			{
@@ -1771,8 +1614,6 @@ namespace uba
 		}
 		case WM_HSCROLL:
 		{
-			RECT r;
-			GetClientRect(hWnd, &r);
 			float oldScrollX = m_scrollPosX;
 			bool autoScroll = false;
 
@@ -1782,6 +1623,8 @@ namespace uba
 			scrollInfo.fMask = SIF_TRACKPOS;
 			GetScrollInfo(m_hwnd, SB_HORZ, &scrollInfo);
 
+			int right = m_clientWidth;
+
 			switch (LOWORD(wParam))
 			{
 			case SB_THUMBTRACK:
@@ -1790,15 +1633,15 @@ namespace uba
 					m_contentWidthWhenThumbTrack = m_contentWidth;
 				break;
 			case SB_THUMBPOSITION:
-				autoScroll = m_contentWidthWhenThumbTrack - r.right <= HIWORD(wParam) + 10;
+				autoScroll = m_contentWidthWhenThumbTrack - right <= HIWORD(wParam) + 10;
 				m_contentWidthWhenThumbTrack = 0;
 				m_scrollPosX = -float(scrollInfo.nTrackPos);
 				break;
 			case SB_PAGEDOWN:
-				m_scrollPosX = m_scrollPosX - float(r.right);
+				m_scrollPosX = m_scrollPosX - float(right);
 				break;
 			case SB_PAGEUP:
-				m_scrollPosX = m_scrollPosX + float(r.right);
+				m_scrollPosX = m_scrollPosX + float(right);
 				break;
 			case SB_LINEDOWN:
 				m_scrollPosX -= 30;
@@ -1810,9 +1653,9 @@ namespace uba
 				return 0;
 			}
 
-			int minScroll = r.right - m_contentWidth;
+			int minScroll = right - m_contentWidth;
 			m_autoScroll = !m_traceView.finished && (m_scrollPosX <= float(minScroll) || autoScroll);
-			m_scrollPosX = Min(Max(m_scrollPosX, float(r.right - m_contentWidth)), 0.0f);
+			m_scrollPosX = Min(Max(m_scrollPosX, float(right - m_contentWidth)), 0.0f);
 
 			if (oldScrollX != m_scrollPosX)
 			{
@@ -1827,13 +1670,16 @@ namespace uba
 
 	LRESULT CALLBACK VisualizerWin::StaticWinProc(HWND hWnd, UINT Msg, WPARAM wParam, LPARAM lParam)
 	{
-		auto thisPtr = (VisualizerWin*)GetWindowLongPtr(hWnd, GWLP_USERDATA);
-		if (!thisPtr && Msg == WM_CREATE)
+		VisualizerWin* thisPtr;
+		if (Msg == WM_CREATE)
 		{
 			thisPtr = (VisualizerWin*)lParam;
 			SetWindowLongPtr(hWnd, GWLP_USERDATA, lParam);
 			//EnableNonClientDpiScaling(hWnd);
 		}
+		else
+			thisPtr = (VisualizerWin*)GetWindowLongPtr(hWnd, GWLP_USERDATA);
+
 		if (thisPtr && hWnd == thisPtr->m_hwnd)
 			return thisPtr->WinProc(hWnd, Msg, wParam, lParam);
 		else
diff --git a/Engine/Source/Programs/UnrealBuildAccelerator/Visualizer/Private/Windows/UbaVisualizerWin.h b/Engine/Source/Programs/UnrealBuildAccelerator/Visualizer/Private/Windows/UbaVisualizerWin.h
--- a/Engine/Source/Programs/UnrealBuildAccelerator/Visualizer/Private/Windows/UbaVisualizerWin.h
+++ b/Engine/Source/Programs/UnrealBuildAccelerator/Visualizer/Private/Windows/UbaVisualizerWin.h
@@ -27,18 +27,14 @@ namespace uba
 
 	private:
 		bool StartHwndThread();
-		bool Unselect();
-		void Reset();
-		void PaintClient(const Function<void(HDC hdc, HDC memDC, RECT& clientRect)>& paintFunc);
+		void PaintClient(const Function<void(HDC hdc, HDC memDC)>& paintFunc);
 		virtual void PaintEmptyScreen(PaintContext& context) override;
 		virtual void FinishIncompatible() override;
 
 		void CopyTextToClipboard(const TString& str);
 		void UnselectAndRedraw();
-		bool UpdateAutoscroll();
-		bool UpdateSelection();
-		void UpdateScrollbars(bool redraw);
-		StringBufferBase& GetTitlePrefix(StringBufferBase& out);
+		virtual bool UpdateSelection() override;
+		virtual void UpdateScrollbars(bool redraw) override;
 		void InitBrushes();
 		void ThreadLoop();
 		void Pause(bool pause);
@@ -50,38 +46,26 @@ namespace uba
 		void UpdateDefaultFont();
 		void UpdateProcessFont();
 		void ChangeFontSize(int offset);
-		void Redraw(bool now);
+		virtual void Redraw(bool now) override;
+		virtual void StopUpdate() override;
 		void UpdateTheme();
 
-		StringBuffer<256> m_namedTrace;
-		StringBuffer<256> m_fileName;
-
 		u64 m_lastPaintTimeMs;
 
 		Atomic<bool> m_looping = false;
 		HWND m_hwnd = 0;
 		HWND m_parentHwnd = 0;
 
-		NetworkClient* m_client = nullptr;
-		Event m_clientDisconnect;
-
-		StringBuffer<256>m_listenChannel;
-		StringBuffer<256> m_newTraceName;
-		Event m_listenTimeout;
-
 		int m_contentWidthWhenThumbTrack = 0;
 
-		bool m_autoScroll = true;
-
 		bool m_isInPaint = false;
 
-		bool m_usingNamed = false;
 		bool m_showPopup = false;
 		bool m_locked = false;
 		bool m_handlingTimer = false;
 
 		HBITMAP m_cachedBitmap = 0;
-		RECT m_cachedBitmapRect = { INT_MIN, INT_MIN, INT_MIN, INT_MIN };
+		SIZE m_cachedBitmapSize = { INT_MIN, INT_MIN };
 
 		POINT m_mouseAnchor = {};
 		float m_scrollAtAnchorX = 0;
@@ -94,7 +78,7 @@ namespace uba
 		Thread m_thread;
 
 		void PostNewTrace(u32 replay, bool paused);
-		void PostNewTitle(const StringView& title);
+		virtual void PostNewTitle(const StringView& title);
 		void PostQuit();
 		LRESULT WinProc(HWND hWnd, UINT Msg, WPARAM wParam, LPARAM lParam);
 		static LRESULT CALLBACK StaticWinProc(HWND hWnd, UINT Msg, WPARAM wParam, LPARAM lParam);
