[UBA]
* Fixed visualizer errors on mac

diff --git a/Engine/Source/Programs/UnrealBuildAccelerator/Visualizer/Private/Mac/UbaVisualizerApp.mm b/Engine/Source/Programs/UnrealBuildAccelerator/Visualizer/Private/Mac/UbaVisualizerApp.mm
--- a/Engine/Source/Programs/UnrealBuildAccelerator/Visualizer/Private/Mac/UbaVisualizerApp.mm
+++ b/Engine/Source/Programs/UnrealBuildAccelerator/Visualizer/Private/Mac/UbaVisualizerApp.mm
@@ -21,22 +21,22 @@
     public:
         struct PaintContextMac : PaintContext
         {
-			virtual RECT GetClientRect()
+			virtual SIZE GetClientSize() override
             {
-                return clientRect;
+                return clientSize;
             }
-			virtual void SetBkMode(int mode) {}
-			virtual void SetPen(HPEN pen) {}
-			virtual void SetBrush(HBRUSH b) {}
-			virtual void SetFont(const Font& font, bool underlined = false) { activeFont = font; }
-			virtual void SetTextColor(COLORREF color) {}
-			virtual u32 SetTextAlign(u32 align) { return 0; }
+			virtual void SetBkMode(int mode) override {}
+			virtual void SetPen(HPEN pen) override {}
+			virtual void SetBrush(HBRUSH b) override {}
+			virtual void SetFont(const Font& font, bool underlined = false) override { activeFont = font; }
+			virtual void SetTextColor(COLORREF color) override {}
+			virtual u32 SetTextAlign(u32 align) override { return 0; }
 
             int textOffsetX = 0;
             int textOffsetY = 0;
             NSRect textRect;
 
-			virtual void Text(int x, int y, const RECT* lprect, StringView str)
+			virtual void Text(int x, int y, const RECT* lprect, StringView str) override
             {         
                 NSString *text = [NSString stringWithUTF8String:str.data];
 
@@ -97,7 +97,7 @@ virtual void DrawCached(TraceView::BitmapCache& process, RECT rect, int offsetY,
 
 			virtual POINT GetCursorClientPos() { return {}; }
 
-            RECT clientRect;
+            SIZE clientSize;
 			Font activeFont;
         };
 
@@ -112,7 +112,26 @@ virtual void ResetCaches() override {}
         {
             char path[512];
             GetFullPathNameW("~/Trace.uba", 512, path, nullptr);
-            m_trace.ReadFile(m_traceView, path, false);
+
+            m_fileName.Append(path);
+            if (true)
+            {
+                m_trace.ReadFile(m_traceView, path, true);
+
+                m_startTime = GetTime();
+                
+                bool changed = true;
+                u64 time = 0;
+                m_traceView.finished = false;
+                m_trace.UpdateReadFile(m_traceView, time, changed);
+                m_pauseStart = m_startTime + time;
+                m_pauseTime = m_startTime - m_pauseStart;
+                m_autoScroll = true;
+                m_replay = 1;
+            }
+            else
+                m_trace.ReadFile(m_traceView, path, false);
+
 		    m_defaultFont.height = 14;
 		    m_processFont.height = 14;
 		    m_timelineFont.height = 14;
@@ -128,7 +147,21 @@ virtual void FinishIncompatible()
         {
         }
 
-        using Visualizer::PaintAll;
+        void Paint(SIZE size)
+        {
+            m_clientWidth = size.cx;
+            m_clientHeight = size.cy;
+
+            uba::VisualizerBackendMac::PaintContextMac context;
+            context.clientSize = size;
+            RECT r { 0, 0, size.cx, size.cy };
+            PaintAll(context, r);
+        }
+
+        void Update()
+        {
+            UpdateTrace();
+        }
     };
 
 	void Visualizer::UpdateFont(Font& font, int height, bool createUnderline)
@@ -172,11 +205,9 @@ - (instancetype)initWithFrame:(NSRect)frame {
 }
 
 - (void)tick:(NSTimer *)timer {
-    // advance your “time” or simulation here
     _time += 1.0 / 60.0;
 
-    // Call into C++ per-frame if you like:
-    // cpp_update_example(_time);
+    _visualizer->Update();
 
     // Request a redraw
     [self setNeedsDisplay:YES];
@@ -208,27 +239,12 @@ - (void)drawRect:(NSRect)dirtyRect {
                                   bounds.size.height - 10)];
     [path stroke];
 */
-    uba::VisualizerBackendMac::PaintContextMac context;
-    uba::RECT r { 0, 0, (int)bounds.size.width, (int)bounds.size.height };
-    context.clientRect = r;
-    _visualizer->PaintAll(context, r);
-
-    // Call into C++ renderer if you want:
-    cpp_render_example(float(bounds.size.width), float(bounds.size.height));
+    uba::SIZE size = { (int)bounds.size.width, (int)bounds.size.height };
+    _visualizer->Paint(size);
 }
 
 @end
 
-// Dummy C++ function implementation (replace with your own):
-void cpp_render_example(float width, float height) {
-    // This is normal C++ code, do whatever you want here.
-    // For example, update some state, log, etc.
-    // For actual drawing, pass data back to Objective-C++ or keep state
-    // in a shared object that drawRect uses.
-    (void)width;
-    (void)height;
-}
-
 @interface AppDelegate : NSObject <NSWindowDelegate>
 @end
 
diff --git a/Engine/Source/Programs/UnrealBuildAccelerator/Visualizer/Private/UbaVisualizer.cpp b/Engine/Source/Programs/UnrealBuildAccelerator/Visualizer/Private/UbaVisualizer.cpp
--- a/Engine/Source/Programs/UnrealBuildAccelerator/Visualizer/Private/UbaVisualizer.cpp
+++ b/Engine/Source/Programs/UnrealBuildAccelerator/Visualizer/Private/UbaVisualizer.cpp
@@ -198,11 +198,11 @@ namespace uba
 	StringBufferBase& Visualizer::GetTitlePrefix(StringBufferBase& out)
 	{
 		out.Clear();
-		out.Append(L"UbaVisualizer");
+		out.Append(TC("UbaVisualizer"));
 		#if UBA_DEBUG
-		out.Append(L" (DEBUG)");
+		out.Append(TC(L" (DEBUG)"));
 		#endif
-		out.Append(L" - ");
+		out.Append(TC(" - "));
 		return out;
 	}
 
@@ -2544,7 +2544,7 @@ namespace uba
 		{
 			m_trace.ReadFile(m_traceView, m_fileName.data, m_replay != 0);
 			m_traceView.finished = m_replay == 0;
-			PostNewTitle(GetTitlePrefix(title).Appendf(L"%s (v%u) - %s", m_fileName.data, m_traceView.version, GetWorldTime((u64)0).data));
+			PostNewTitle(GetTitlePrefix(title).Appendf(TC("%s (v%u) - %s"), m_fileName.data, m_traceView.version, GetWorldTime((u64)0).data));
 		}
 		else if (m_usingNamed)
 		{
@@ -2552,7 +2552,7 @@ namespace uba
 				return false;
 			m_namedTrace.Clear().Append(m_newTraceName);
 			m_traceView.finished = false;
-			PostNewTitle(GetTitlePrefix(title).Appendf(L"%s (Listening for new sessions on channel '%s')", m_namedTrace.data, m_listenChannel.data));
+			PostNewTitle(GetTitlePrefix(title).Appendf(TC("%s (Listening for new sessions on channel '%s')"), m_namedTrace.data, m_listenChannel.data));
 		}
 
 		return true;
