[UBA]
* Added support for CreateFile with CREATE_NEW on posix
* Added WriteFile implementation
* Implemented SetEndOfFile for posix

diff --git a/Engine/Source/Programs/UnrealBuildAccelerator/Common/Private/UbaFile.cpp b/Engine/Source/Programs/UnrealBuildAccelerator/Common/Private/UbaFile.cpp
--- a/Engine/Source/Programs/UnrealBuildAccelerator/Common/Private/UbaFile.cpp
+++ b/Engine/Source/Programs/UnrealBuildAccelerator/Common/Private/UbaFile.cpp
@@ -141,6 +141,28 @@ namespace uba
 		return true;
 	}
 
+	bool WriteFile(Logger& logger, const tchar* fileName, FileHandle fileHandle, const void* b, u64 bufferLen)
+	{
+		auto& stats = KernelStats::GetCurrent();
+		ExtendedTimerScope ts(stats.writeFile);
+		const u8* buffer = (const u8*)b;
+		u64 writeLeft = bufferLen;
+		while (writeLeft)
+		{
+#if PLATFORM_WINDOWS
+			DWORD written = 0;
+			if (!::WriteFile(asHANDLE(fileHandle), buffer, (DWORD)writeLeft, &written, NULL))
+				return logger.Error(TC("ERROR WriteFile failed writing to file %s (%s)"), fileName, LastErrorToText().data);
+#else
+			auto written = write(asFileDescriptor(fileHandle), buffer, writeLeft);
+			if (written == -1)
+				return false;
+#endif
+			writeLeft -= written;
+		}
+		return true;
+	}
+
 	UBA_NOINLINE bool ReportOpenFileError(Logger& logger, const tchar* fileName, u32 lastError) // To prevent stack usage
 	{
 		StringBuffer<4096> additionalInfo;
@@ -373,7 +395,8 @@ namespace uba
 			return logger.Error(TC("SetFileInformationByHandle failed on %s (%s)"), fileName, LastErrorToText().data);
 		return true;
 #else
-		UBA_ASSERTF(false, TC("SetEndOfFile not implemented"));
+		if (ftruncate(asFileDescriptor(handle), size) == -1)
+			return logger.Error(TC("ftruncate failed on %s (%s)"), fileName, LastErrorToText().data);
 		return false;
 #endif
 	}
@@ -553,6 +576,8 @@ namespace uba
 			flags |= O_CREAT | O_TRUNC;
 		else if (createDisp == OPEN_EXISTING)
 			flags = O_NONBLOCK;
+		else if (createDisp == CREATE_NEW)
+			flags |= O_CREAT | O_EXCL;
 		else
 			UBA_ASSERTF(false, TC("CreateFileW create disposition %u not supported"), createDisp);
 
@@ -583,6 +608,12 @@ namespace uba
 			return InvalidFileHandle;
 		}
 
+		if (errno == EEXIST)
+		{
+			SetLastError(ERROR_FILE_EXISTS);
+			return InvalidFileHandle;
+		}
+
 		if (errno == ENOENT)
 		{
 			SetLastError(ERROR_FILE_NOT_FOUND);
diff --git a/Engine/Source/Programs/UnrealBuildAccelerator/Common/Public/UbaFile.h b/Engine/Source/Programs/UnrealBuildAccelerator/Common/Public/UbaFile.h
--- a/Engine/Source/Programs/UnrealBuildAccelerator/Common/Public/UbaFile.h
+++ b/Engine/Source/Programs/UnrealBuildAccelerator/Common/Public/UbaFile.h
@@ -12,6 +12,7 @@ namespace uba
 	inline constexpr u32 ERROR_FILE_NOT_FOUND = ENOENT;
 	inline constexpr u32 ERROR_PATH_NOT_FOUND = ENOENT;
 	inline constexpr u32 ERROR_ALREADY_EXISTS = EEXIST;
+	inline constexpr u32 ERROR_FILE_EXISTS = EEXIST;
 	inline constexpr u32 ERROR_ACCESS_DENIED = EACCES;
 	inline constexpr u32 ERROR_DIRECTORY = ENOTDIR;
 	inline constexpr u32 MOVEFILE_REPLACE_EXISTING = 0x00000001;
@@ -31,6 +32,7 @@ namespace uba
 	inline constexpr u32 FILE_FLAG_BACKUP_SEMANTICS = 0x02000000;
 	inline constexpr u32 FILE_FLAG_SEQUENTIAL_SCAN = 0;
 
+	inline constexpr u32 CREATE_NEW = 1;
 	inline constexpr u32 CREATE_ALWAYS = 2;
 	inline constexpr u32 OPEN_EXISTING = 3;
 	#endif
@@ -63,6 +65,7 @@ namespace uba
 	bool GetFileInformation(FileInformation& out, Logger& logger, const tchar* fileName);
 
 	bool ReadFile(Logger& logger, const tchar* fileName, FileHandle fileHandle, void* b, u64 bufferLen);
+	bool WriteFile(Logger& logger, const tchar* fileName, FileHandle fileHandle, const void* b, u64 bufferLen);
 	bool OpenFileSequentialRead(Logger& logger, const tchar* fileName, FileHandle& outHandle, bool fileNotFoundIsError = true, bool overlapped = false, bool showErrors = true);
 	bool FileExists(Logger& logger, const tchar* fileName, u64* outSize = nullptr, u32* outAttributes = nullptr, u64* lastWriteTime = nullptr);
 	bool SetFilePointer(Logger& logger, const tchar* fileName, FileHandle handle, u64 position);
