[UBA]
* Implement file based trace channel support. Only used by non-windows
* Enabled memory traces for non-windows (WIP)

diff --git a/Engine/Source/Programs/UnrealBuildAccelerator/Common/Private/UbaSession.cpp b/Engine/Source/Programs/UnrealBuildAccelerator/Common/Private/UbaSession.cpp
--- a/Engine/Source/Programs/UnrealBuildAccelerator/Common/Private/UbaSession.cpp
+++ b/Engine/Source/Programs/UnrealBuildAccelerator/Common/Private/UbaSession.cpp
@@ -1987,7 +1987,7 @@ namespace uba
 				u64 traceReserveSize = info.traceReserveSizeMb * 1024 * 1024;
 				if (m_detailedTrace)
 					traceReserveSize *= 2;
-				m_trace.StartWriteAndThread(IsWindows ? traceName.data : nullptr, traceReserveSize, true); // non-windows named shared memory not implemented (only needed for UbaVisualizer which you can't run on linux either way)
+				m_trace.StartWriteAndThread(traceName.data, traceReserveSize, true);
 			}
 		}
 
diff --git a/Engine/Source/Programs/UnrealBuildAccelerator/Common/Private/UbaTrace.cpp b/Engine/Source/Programs/UnrealBuildAccelerator/Common/Private/UbaTrace.cpp
--- a/Engine/Source/Programs/UnrealBuildAccelerator/Common/Private/UbaTrace.cpp
+++ b/Engine/Source/Programs/UnrealBuildAccelerator/Common/Private/UbaTrace.cpp
@@ -5,15 +5,12 @@
 #include "UbaProcessHandle.h"
 #include "UbaStats.h"
 
-#if PLATFORM_WINDOWS
-#include <tlhelp32.h>
-#include <psapi.h>
-#endif
-
 namespace uba
 {
 	constexpr u64 TraceMessageMaxSize = 256 * 1024;
 
+	////////////////////////////////////////////////////////////////////////////////////////////////////
+
 	Trace::Trace(LogWriter& logWriter)
 	:	m_logger(logWriter)
 	,	m_channel(m_logger)
@@ -145,6 +142,12 @@ namespace uba
 
 		if (namedTrace)
 		{
+			#if !PLATFORM_WINDOWS
+			StringBuffer<64> uidName;
+			GetMappingHandleName(uidName, m_memoryHandle.uid);
+			namedTrace = uidName.data;
+			#endif
+
 			m_namedTrace = namedTrace;
 			if (channelName && *channelName && m_channel.Init(channelName))
 				m_channel.Write(namedTrace);
@@ -716,161 +719,5 @@ namespace uba
 			writer.Write7BitEncoded(items[i]);
 	}
 
-	TraceChannel::TraceChannel(Logger& logger) : m_logger(logger)
-	{
-	}
-
-	bool TraceChannel::Init(const tchar* channelName)
-	{
-		#if PLATFORM_WINDOWS
-		StringBuffer<245> channelMutex;
-		channelMutex.Append(TCV("Uba")).Append(channelName).Append(TCV("Channel"));
-		m_memHandle = uba::CreateMemoryMappingW(m_logger, PAGE_READWRITE, 256, channelName, TC("TraceChannel"));
-		if (!m_memHandle.IsValid())
-		{
-			m_logger.Error(TC("Failed to create file mapping %s for trace channel (%s)"), channelName, LastErrorToText().data);
-			return false;
-		}
-		bool isCreator = GetLastError() != ERROR_ALREADY_EXISTS;
-
-		auto mhg = MakeGuard([&]() { CloseFileMapping(m_logger, m_memHandle, TC("TraceChannel")); m_memHandle = {}; });
-
-		m_mem = MapViewOfFile(m_logger, m_memHandle, FILE_MAP_WRITE, 0, 256);
-		if (!m_mem)
-		{
-			m_logger.Error(TC("Failed to map file mapping for uba trace channel"));
-			return false;
-		}
-
-		if (isCreator)
-			*(tchar*)m_mem = 0;
-
-		auto mg = MakeGuard([&]() { UnmapViewOfFile(m_logger, m_mem, 256, channelMutex.data); m_mem = nullptr; });
-
-		channelMutex.Append(channelName).Append(TCV("Mutex"));
-		m_mutex = CreateMutexW(false, channelMutex.data);
-		if (m_mutex == InvalidMutexHandle)
-			return false;
-
-		mg.Cancel();
-		mhg.Cancel();
-		#endif
-		return true;
-	}
-
-	void TraceChannel::Deinit()
-	{
-		#if PLATFORM_WINDOWS
-		if (m_mem)
-			UnmapViewOfFile(m_logger, m_mem, 256, TC("TraceChannel"));
-		if (m_memHandle.IsValid())
-			CloseFileMapping(m_logger, m_memHandle, TC("TraceChannel"));
-		if (m_mutex != InvalidMutexHandle)
-			CloseMutex(m_mutex);
-		m_mem = nullptr;
-		m_memHandle = {};
-		m_mutex = InvalidMutexHandle;
-		#endif
-	}
-
-	TraceChannel::~TraceChannel()
-	{
-		Deinit();
-	}
-
-	bool TraceChannel::Write(const tchar* traceName, const tchar* ifMatching)
-	{
-		#if PLATFORM_WINDOWS
-		WaitForSingleObject((HANDLE)m_mutex, INFINITE);
-		auto g = MakeGuard([this]() { ReleaseMutex(m_mutex); });
-		if (ifMatching)
-			if (!Equals((tchar*)m_mem, ifMatching))
-				return true;
-		TStrcpy_s((tchar*)m_mem, 256, traceName);
-		#endif
-		return true;
-	}
-
-	bool TraceChannel::Read(StringBufferBase& outTraceName)
-	{
-		#if PLATFORM_WINDOWS
-		WaitForSingleObject((HANDLE)m_mutex, INFINITE);
-		outTraceName.Append((tchar*)m_mem);
-		ReleaseMutex(m_mutex);
-		#endif
-		return true;
-	}
-
-	bool TraceChannel::IsInitialized()
-	{
-		return m_mutex != InvalidMutexHandle;
-	}
-
-	static OwnerInfo InternalGetOwnerInfo()
-	{
-		static tchar buffer[260];
-		*buffer = 0;
-
-		OwnerInfo info { buffer, 0 };
-
-		StringBuffer<32> ownerPidStr;
-		ownerPidStr.count = GetEnvironmentVariableW(TC("UBA_OWNER_PID"), ownerPidStr.data, ownerPidStr.capacity);
-		if (ownerPidStr.count)
-		{
-			GetEnvironmentVariableW(TC("UBA_OWNER_ID"), buffer, sizeof_array(buffer));
-			ownerPidStr.Parse(info.pid);
-			return info;
-		}
-
-		#if PLATFORM_WINDOWS
-		HANDLE snapshotHandle = CreateToolhelp32Snapshot(TH32CS_SNAPPROCESS, 0);
-		if (snapshotHandle == INVALID_HANDLE_VALUE)
-			return info;
-
-		PROCESSENTRY32 pe = { 0 };
-		pe.dwSize = sizeof(PROCESSENTRY32);
-		UnorderedMap<u32, u32> pidToParent;
-		if (Process32First(snapshotHandle, &pe))
-		{
-			do
-			{
-				pidToParent[pe.th32ProcessID] = pe.th32ParentProcessID;
-			}
-			while (Process32Next(snapshotHandle, &pe));
-		}
-		CloseHandle(snapshotHandle);
-
-		u32 pid = ::GetCurrentProcessId();
-		while (true)
-		{
-			auto findIt = pidToParent.find(pid);
-			if (findIt == pidToParent.end())
-				break;
-			pid = findIt->second;
-			pidToParent.erase(findIt);
-
-			HANDLE parentHandle = OpenProcess(PROCESS_QUERY_INFORMATION | PROCESS_VM_READ, FALSE, pid);
-			if (parentHandle == NULL)
-				break;
-			tchar moduleName[260];
-			DWORD len = GetModuleFileNameExW(parentHandle, 0, moduleName, MAX_PATH);
-			CloseHandle(parentHandle);
-			if (!len)
-				break;
-			if (!Contains(moduleName, L"devenv.exe"))
-				continue;
-			TStrcpy_s(buffer, MAX_PATH, L"vs");
-			info.pid = pid;
-			break;
-		}
-		#endif
-
-		return info;
-	}
-
-	const OwnerInfo& GetOwnerInfo()
-	{
-		static OwnerInfo info = InternalGetOwnerInfo();
-		return info;
-	}
+	////////////////////////////////////////////////////////////////////////////////////////////////////
 }
diff --git a/Engine/Source/Programs/UnrealBuildAccelerator/Common/Private/UbaTraceChannel.cpp b/Engine/Source/Programs/UnrealBuildAccelerator/Common/Private/UbaTraceChannel.cpp
new file mode 100644
--- /dev/null
+++ b/Engine/Source/Programs/UnrealBuildAccelerator/Common/Private/UbaTraceChannel.cpp
@@ -0,0 +1,339 @@
+// Copyright Epic Games, Inc. All Rights Reserved.
+
+#include "UbaTraceChannel.h"
+#include "UbaFile.h"
+#include "UbaTrace.h"
+
+#if PLATFORM_WINDOWS
+#include <tlhelp32.h>
+#include <psapi.h>
+#endif
+
+namespace uba
+{
+	////////////////////////////////////////////////////////////////////////////////////////////////////
+
+	TraceChannel::TraceChannel(Logger& logger) : m_logger(logger)
+	{
+	}
+
+	bool TraceChannel::Init(const tchar* channelName)
+	{
+	#if UBA_USE_TRACE_CHANNEL_FILE
+		StringBuffer<> fileName;
+
+		#if PLATFORM_WINDOWS
+		fileName.Append(TCV("d:\\temp\\"));
+		#else
+		const char* tmpdir = getenv("TMPDIR");
+		if (!tmpdir || tmpdir[0] == '\0')
+			tmpdir = "/tmp";
+		fileName.Append(tmpdir);
+		#endif
+
+		m_channelFile = fileName.EnsureEndsWithSlash().Append(TCV("UbaTrace_")).Append(channelName).ToString();
+
+	#else
+		StringBuffer<245> channelMutex;
+		channelMutex.Append(TCV("Uba")).Append(channelName).Append(TCV("Channel"));
+		m_memHandle = uba::CreateMemoryMappingW(m_logger, PAGE_READWRITE, 256, channelName, TC("TraceChannel"));
+		if (!m_memHandle.IsValid())
+		{
+			m_logger.Error(TC("Failed to create file mapping %s for trace channel (%s)"), channelName, LastErrorToText().data);
+			return false;
+		}
+		bool isCreator = GetLastError() != ERROR_ALREADY_EXISTS;
+
+		auto mhg = MakeGuard([&]() { CloseFileMapping(m_logger, m_memHandle, TC("TraceChannel")); m_memHandle = {}; });
+
+		m_mem = MapViewOfFile(m_logger, m_memHandle, FILE_MAP_WRITE, 0, 256);
+		if (!m_mem)
+		{
+			m_logger.Error(TC("Failed to map file mapping for uba trace channel"));
+			return false;
+		}
+
+		if (isCreator)
+			*(tchar*)m_mem = 0;
+
+		auto mg = MakeGuard([&]() { UnmapViewOfFile(m_logger, m_mem, 256, channelMutex.data); m_mem = nullptr; });
+
+		channelMutex.Append(channelName).Append(TCV("Mutex"));
+		m_mutex = CreateMutexW(false, channelMutex.data);
+		if (m_mutex == InvalidMutexHandle)
+			return false;
+
+		mg.Cancel();
+		mhg.Cancel();
+	#endif
+		return true;
+	}
+
+	void TraceChannel::Deinit()
+	{
+		#if UBA_USE_TRACE_CHANNEL_FILE
+		#else
+		if (m_mem)
+			UnmapViewOfFile(m_logger, m_mem, 256, TC("TraceChannel"));
+		if (m_memHandle.IsValid())
+			CloseFileMapping(m_logger, m_memHandle, TC("TraceChannel"));
+		if (m_mutex != InvalidMutexHandle)
+			CloseMutex(m_mutex);
+		m_mem = nullptr;
+		m_memHandle = {};
+		m_mutex = InvalidMutexHandle;
+		#endif
+	}
+
+	TraceChannel::~TraceChannel()
+	{
+		Deinit();
+	}
+
+#if UBA_USE_TRACE_CHANNEL_FILE
+	template<typename Function>
+	bool UpdateTraceFile(Logger& logger, StringView channelFile, const tchar* traceName, const tchar* ifMatching, const Function& onPotentialValidTrace)
+	{
+		u8 writeBuffer[4096];
+		BinaryWriter writer(writeBuffer, 0, sizeof(writeBuffer));
+		bool modified = false;
+
+#if PLATFORM_WINDOWS
+		DWORD currentProcessId = GetCurrentProcessId();
+		FILETIME createTime, exitTime, kernelTime, userTime;
+		if (!GetProcessTimes(GetCurrentProcess(), &createTime, &exitTime, &kernelTime, &userTime))
+			return logger.Error(TC("Failed to get current process process times"));
+		u64 currentProcessCreateTime = ToLargeInteger(createTime.dwHighDateTime, createTime.dwLowDateTime).QuadPart
+		if (*traceName)
+		{
+			writer.Write7BitEncoded(currentProcessId);
+			writer.Write7BitEncoded(ToLargeInteger(createTime.dwHighDateTime, createTime.dwLowDateTime).QuadPart);
+			writer.WriteString(traceName);
+			modified = true;
+		}
+#else
+		pid_t currentProcessId = getpid();
+		u64 currentProcessCreateTime = 0;
+#endif
+
+		if (*traceName)
+		{
+			writer.Write7BitEncoded(currentProcessId);
+			writer.Write7BitEncoded(currentProcessCreateTime);
+			writer.WriteString(traceName);
+			modified = true;
+		}
+
+		u32 createDisp = CREATE_NEW;
+		u32 retryCount = 0;
+		while (true)
+		{
+			FileHandle file = CreateFileW(channelFile.data, GENERIC_READ|GENERIC_WRITE, 0, createDisp, DefaultAttributes());
+			if (file == InvalidFileHandle)
+			{
+				if (createDisp == CREATE_NEW && GetLastError() == ERROR_FILE_EXISTS)
+				{
+					createDisp = OPEN_EXISTING;
+					continue;
+				}
+				if (retryCount > 10)
+					return logger.Warning(TC("Failed to create/open trace channel file %s. No trace will be available (%s)"), channelFile.data, LastErrorToText().data);
+				++retryCount;
+				Sleep(100);
+				continue;
+			}
+			auto g = MakeGuard([&]() { CloseFile(channelFile.data, file); });
+
+			u64 fileSize = 0;
+			if (!GetFileSizeEx(fileSize, file))
+				return logger.Error(TC("Failed to get file size for %s"), channelFile.data);
+
+			u8 readBuffer[4096];
+			if (fileSize > sizeof(readBuffer))
+				return logger.Error(TC("TraceChannel file %s too big"), channelFile.data);
+
+			if (!ReadFile(logger, channelFile.data, file, readBuffer, fileSize))
+				return false;
+
+			BinaryReader reader(readBuffer, 0, fileSize);
+			while (reader.GetLeft())
+			{
+				auto mg = MakeGuard([&]() { modified = true; });
+				u64 processId;
+				u64 processCreateTime;
+				StringBuffer<> trace;
+
+				if (!reader.TryRead7BitEncoded(processId))
+					break;
+				if (!reader.TryRead7BitEncoded(processCreateTime))
+					break;
+				if (!reader.TryReadString(trace))
+					break;
+
+				if (processId == currentProcessId)
+				{
+					if (ifMatching && trace.Equals(ifMatching))
+					{
+						if (!*traceName)
+							continue;
+						if (!trace.Equals(traceName))
+							mg.Execute();
+						trace.Clear().Append(traceName);
+					}
+				}
+				else
+				{
+#if PLATFORM_WINDOWS
+					HANDLE p = OpenProcess(PROCESS_QUERY_LIMITED_INFORMATION, TRUE, u32(processId));
+					if (p == NULL)
+						continue;
+					auto pg = MakeGuard([p](){ CloseHandle(p); });
+					if (!GetProcessTimes(p, &createTime, &exitTime, &kernelTime, &userTime))
+						continue;
+					u64 processCreateTime2 = ToLargeInteger(createTime.dwHighDateTime, createTime.dwLowDateTime).QuadPart;
+					if (processCreateTime != processCreateTime2)
+						continue;
+#else
+					if (kill(processId, 0) != 0)
+						continue;
+#endif
+				}
+
+				mg.Cancel();
+
+				writer.Write7BitEncoded(processId);
+				writer.Write7BitEncoded(processCreateTime);
+				writer.WriteString(trace);
+
+				onPotentialValidTrace(trace);
+			}
+
+			if (!modified)
+				return true;
+
+			u32 newFileSize = u32(writer.GetPosition());
+
+			if (!SetFilePointer(logger, channelFile.data, file, 0))
+				return false;
+
+			if (!WriteFile(logger, channelFile.data, file, writeBuffer, newFileSize))
+				return false;
+
+			if (!SetEndOfFile(logger, channelFile.data, file, newFileSize))
+				return false;
+
+			return true;
+		}
+	}
+#endif
+
+	bool TraceChannel::Write(const tchar* traceName, const tchar* ifMatching)
+	{
+		#if UBA_USE_TRACE_CHANNEL_FILE
+		return UpdateTraceFile(m_logger, m_channelFile, traceName, ifMatching, [](StringView trace) {});
+		#else
+		WaitForSingleObject((HANDLE)m_mutex, INFINITE);
+		auto g = MakeGuard([this]() { ReleaseMutex(m_mutex); });
+		if (ifMatching)
+			if (!Equals((tchar*)m_mem, ifMatching))
+				return true;
+		TStrcpy_s((tchar*)m_mem, 256, traceName);
+		return true;
+		#endif
+	}
+
+	bool TraceChannel::Read(StringBufferBase& outTraceName)
+	{
+		#if UBA_USE_TRACE_CHANNEL_FILE
+		return UpdateTraceFile(m_logger, m_channelFile, TC(""), TC(""), [&](StringView trace)
+			{
+				if (!outTraceName.count)
+					outTraceName.Append(trace);
+			});
+		#else
+		WaitForSingleObject((HANDLE)m_mutex, INFINITE);
+		outTraceName.Append((tchar*)m_mem);
+		ReleaseMutex(m_mutex);
+		return true;
+#endif
+	}
+
+	bool TraceChannel::IsInitialized()
+	{
+#if UBA_USE_TRACE_CHANNEL_FILE
+		return !m_channelFile.empty();
+#else
+		return m_mutex != InvalidMutexHandle;
+#endif
+	}
+
+	static OwnerInfo InternalGetOwnerInfo()
+	{
+		static tchar buffer[260];
+		*buffer = 0;
+
+		OwnerInfo info { buffer, 0 };
+
+		StringBuffer<32> ownerPidStr;
+		ownerPidStr.count = GetEnvironmentVariableW(TC("UBA_OWNER_PID"), ownerPidStr.data, ownerPidStr.capacity);
+		if (ownerPidStr.count)
+		{
+			GetEnvironmentVariableW(TC("UBA_OWNER_ID"), buffer, sizeof_array(buffer));
+			ownerPidStr.Parse(info.pid);
+			return info;
+		}
+
+		#if PLATFORM_WINDOWS
+		HANDLE snapshotHandle = CreateToolhelp32Snapshot(TH32CS_SNAPPROCESS, 0);
+		if (snapshotHandle == INVALID_HANDLE_VALUE)
+			return info;
+
+		PROCESSENTRY32 pe = { 0 };
+		pe.dwSize = sizeof(PROCESSENTRY32);
+		UnorderedMap<u32, u32> pidToParent;
+		if (Process32First(snapshotHandle, &pe))
+		{
+			do
+			{
+				pidToParent[pe.th32ProcessID] = pe.th32ParentProcessID;
+			}
+			while (Process32Next(snapshotHandle, &pe));
+		}
+		CloseHandle(snapshotHandle);
+
+		u32 pid = ::GetCurrentProcessId();
+		while (true)
+		{
+			auto findIt = pidToParent.find(pid);
+			if (findIt == pidToParent.end())
+				break;
+			pid = findIt->second;
+			pidToParent.erase(findIt);
+
+			HANDLE parentHandle = OpenProcess(PROCESS_QUERY_INFORMATION | PROCESS_VM_READ, FALSE, pid);
+			if (parentHandle == NULL)
+				break;
+			tchar moduleName[260];
+			DWORD len = GetModuleFileNameExW(parentHandle, 0, moduleName, MAX_PATH);
+			CloseHandle(parentHandle);
+			if (!len)
+				break;
+			if (!Contains(moduleName, L"devenv.exe"))
+				continue;
+			TStrcpy_s(buffer, MAX_PATH, L"vs");
+			info.pid = pid;
+			break;
+		}
+		#endif
+
+		return info;
+	}
+
+	const OwnerInfo& GetOwnerInfo()
+	{
+		static OwnerInfo info = InternalGetOwnerInfo();
+		return info;
+	}
+
+	////////////////////////////////////////////////////////////////////////////////////////////////////
+}
\ No newline at end of file
diff --git a/Engine/Source/Programs/UnrealBuildAccelerator/Common/Public/UbaTrace.h b/Engine/Source/Programs/UnrealBuildAccelerator/Common/Public/UbaTrace.h
--- a/Engine/Source/Programs/UnrealBuildAccelerator/Common/Public/UbaTrace.h
+++ b/Engine/Source/Programs/UnrealBuildAccelerator/Common/Public/UbaTrace.h
@@ -2,10 +2,9 @@
 
 #pragma once
 
-#include "UbaFileMapping.h"
-#include "UbaLogger.h"
 #include "UbaProcessStats.h"
 #include "UbaThread.h"
+#include "UbaTraceChannel.h"
 #include "UbaWorkManager.h"
 
 namespace uba
@@ -15,24 +14,6 @@ namespace uba
 	struct ProcessLogLine;
 	struct ProcessStats;
 
-	class TraceChannel
-	{
-	public:
-		TraceChannel(Logger& logger);
-		~TraceChannel();
-
-		bool Init(const tchar* channelName = TC("Default"));
-		void Deinit();
-		bool Write(const tchar* traceName, const tchar* ifMatching = nullptr);
-		bool Read(StringBufferBase& outTraceName);
-		bool IsInitialized();
-
-		Logger& m_logger;
-		MutexHandle m_mutex = InvalidMutexHandle;
-		FileMappingHandle m_memHandle;
-		void* m_mem = nullptr;
-	};
-
 	#define UBA_TRACE_TYPES \
 		UBA_TRACE_TYPE(SessionAdded) \
 		UBA_TRACE_TYPE(SessionUpdate) \
diff --git a/Engine/Source/Programs/UnrealBuildAccelerator/Common/Public/UbaTraceChannel.h b/Engine/Source/Programs/UnrealBuildAccelerator/Common/Public/UbaTraceChannel.h
new file mode 100644
--- /dev/null
+++ b/Engine/Source/Programs/UnrealBuildAccelerator/Common/Public/UbaTraceChannel.h
@@ -0,0 +1,40 @@
+// Copyright Epic Games, Inc. All Rights Reserved.
+
+#pragma once
+
+#include "UbaFileMapping.h"
+#include "UbaLogger.h"
+
+#define UBA_USE_TRACE_CHANNEL_FILE !PLATFORM_WINDOWS
+
+namespace uba
+{
+	////////////////////////////////////////////////////////////////////////////////////////////////////
+
+	class TraceChannel
+	{
+	public:
+		TraceChannel(Logger& logger);
+		~TraceChannel();
+
+		bool Init(const tchar* channelName = TC("Default"));
+		void Deinit();
+		bool Write(const tchar* traceName, const tchar* ifMatching = nullptr);
+		bool Read(StringBufferBase& outTraceName);
+		bool IsInitialized();
+
+		Logger& m_logger;
+
+		#if UBA_USE_TRACE_CHANNEL_FILE
+		TString m_channelFile;
+
+		#else
+		MutexHandle m_mutex = InvalidMutexHandle;
+		FileMappingHandle m_memHandle;
+		void* m_mem = nullptr;
+		#endif
+	};
+
+	////////////////////////////////////////////////////////////////////////////////////////////////////
+
+}
diff --git a/Engine/Source/Programs/UnrealBuildAccelerator/Visualizer/Private/Linux/UbaVisualizerApp.cpp b/Engine/Source/Programs/UnrealBuildAccelerator/Visualizer/Private/Linux/UbaVisualizerApp.cpp
--- a/Engine/Source/Programs/UnrealBuildAccelerator/Visualizer/Private/Linux/UbaVisualizerApp.cpp
+++ b/Engine/Source/Programs/UnrealBuildAccelerator/Visualizer/Private/Linux/UbaVisualizerApp.cpp
@@ -1,6 +1,58 @@
 // Copyright Epic Games, Inc. All Rights Reserved.
 
+#include "UbaTraceChannel.h"
+#include "UbaTraceReader.h"
+
 int main(int argc, char* argv[])
 {
+	using namespace uba;
+
+	LoggerWithWriter logger(g_consoleLogWriter);
+
+	logger.Info(TC("UbaVisualizer (test implementation)"));
+
+	TraceChannel channel(logger);
+	if (!channel.Init())
+		return -1;
+
+	StringBuffer<> oldTrace;
+
+	TraceView view;
+	TraceReader reader(logger);
+
+	while (true)
+	{
+		StringBuffer<> newTrace;
+		if (!channel.Read(newTrace))
+			return -1;
+		if (oldTrace.Equals(newTrace))
+		{
+			Sleep(100);
+			continue;
+		}
+		logger.Info(TC("New trace %s"), newTrace.data);
+		oldTrace = newTrace;
+
+		if (!newTrace.count)
+			continue;
+		
+		if (!reader.StartReadNamed(view, newTrace.data))
+			logger.Info(TC("StartReadNamed failed"));
+		
+		while (true)
+		{
+			//if (!view.sessions.empty())
+			//	printf("%s\n", view.sessions[0].name.c_str());
+			bool outChanged = false;
+			if (!reader.UpdateReadNamed(view, ~0ull, outChanged))
+				break;
+			//if (outChanged)
+			//	logger.Info(TC("Changed"));
+			if (view.finished)
+				break;
+		}
+
+	}
+
 	return 0;
 }
\ No newline at end of file
diff --git a/Engine/Source/Programs/UnrealBuildAccelerator/Visualizer/Private/UbaTraceReader.cpp b/Engine/Source/Programs/UnrealBuildAccelerator/Visualizer/Private/UbaTraceReader.cpp
--- a/Engine/Source/Programs/UnrealBuildAccelerator/Visualizer/Private/UbaTraceReader.cpp
+++ b/Engine/Source/Programs/UnrealBuildAccelerator/Visualizer/Private/UbaTraceReader.cpp
@@ -272,7 +272,17 @@ namespace uba
 		if (*namedTrace && m_namedTrace != namedTrace)
 		{
 #if PLATFORM_WINDOWS
+			u64 traceCapacity = 0;
 			m_memoryHandle.mh = ::OpenFileMappingW(PAGE_READWRITE, false, namedTrace);
+#else
+			u64 traceCapacity = 128 * 1024 * 1024;
+			int oflags = O_RDONLY | O_NOFOLLOW;
+			m_memoryHandle.shmFd = shm_open(namedTrace, oflags, S_IRUSR | S_IWUSR);
+
+			// We don't want to unlink
+			//if (m_memoryHandle.shmFd != -1)
+			//	m_memoryHandle.uid = atoi(strchr(namedTrace, '_') + 1);
+			//printf("%s: %i\n", namedTrace, m_memoryHandle.shmFd);
 #endif
 
 			if (!m_memoryHandle.IsValid())
@@ -281,9 +291,13 @@ namespace uba
 					m_logger.Error(TC("OpenFileMappingW - Failed to open file mapping %s (%s)"), namedTrace, LastErrorToText().data);
 				return false;
 			}
-			m_memoryBegin = MapViewOfFile(m_logger, m_memoryHandle, FILE_MAP_READ, 0, 0);
+			m_memoryBegin = MapViewOfFile(m_logger, m_memoryHandle, FILE_MAP_READ, 0, traceCapacity);
 			if (!m_memoryBegin)
+			{
+				if (!silentFail)
+					m_logger.Error(TC("MapViewOfFile - Failed to map file mapping %s (%s)"), namedTrace, LastErrorToText().data);
 				return false;
+			}
 		}
 
 		m_namedTrace = namedTrace;
@@ -1531,12 +1545,15 @@ namespace uba
 	void TraceReader::Unmap()
 	{
 #if PLATFORM_WINDOWS
+		u64 unmapSize = 0;
 		if (m_hostProcess)
 			CloseHandle(m_hostProcess);
 		m_hostProcess = 0;
+#else
+		u64 unmapSize = 128 * 1024 * 1024;
 #endif
 		if (m_memoryBegin)
-			UnmapViewOfFile(m_logger, m_memoryBegin, 0, TC("TraceReader"));
+			UnmapViewOfFile(m_logger, m_memoryBegin, unmapSize, TC("TraceReader"));
 		m_memoryBegin = nullptr;
 		if (m_memoryHandle.IsValid())
 			CloseFileMapping(m_logger, m_memoryHandle, TC("TraceReader"));
diff --git a/Engine/Source/Programs/UnrealBuildAccelerator/Visualizer/Private/UbaVisualizer.cpp b/Engine/Source/Programs/UnrealBuildAccelerator/Visualizer/Private/UbaVisualizer.cpp
--- a/Engine/Source/Programs/UnrealBuildAccelerator/Visualizer/Private/UbaVisualizer.cpp
+++ b/Engine/Source/Programs/UnrealBuildAccelerator/Visualizer/Private/UbaVisualizer.cpp
@@ -200,7 +200,7 @@ namespace uba
 		out.Clear();
 		out.Append(TC("UbaVisualizer"));
 		#if UBA_DEBUG
-		out.Append(TC(L" (DEBUG)"));
+		out.Append(TC(" (DEBUG)"));
 		#endif
 		out.Append(TC(" - "));
 		return out;
