[UBA]
* Added track input support to posix
* Fixed so tracked inputs include child processes tracked inputs
* Added "[Scheduler] WriteTrackedInputs = true" to enable writing tracked inputs

diff --git a/Engine/Source/Programs/UnrealBuildAccelerator/Common/Private/UbaProcess.cpp b/Engine/Source/Programs/UnrealBuildAccelerator/Common/Private/UbaProcess.cpp
--- a/Engine/Source/Programs/UnrealBuildAccelerator/Common/Private/UbaProcess.cpp
+++ b/Engine/Source/Programs/UnrealBuildAccelerator/Common/Private/UbaProcess.cpp
@@ -396,13 +396,26 @@ namespace uba
 
 		#else
 
-		if (m_parentProcess && m_parentProcess->m_nativeProcessId != 0) // Can't do wait on grandchildren on Linux.. but since we use PR_SET_CHILD_SUBREAPER we should once parent is gone and child is orphaned
-			return true;
-
-		#if PLATFORM_MAC
-		if (m_parentProcess && m_gotExitMessage) // TODO: We need a timeout here... if child crashes we will never get exit message
-			return false;
-		#endif
+		if (m_parentProcess)
+		{
+#if PLATFORM_LINUX // Can't do wait on grandchildren on Linux.. but since we use PR_SET_CHILD_SUBREAPER we should once parent is gone and child is orphaned
+			if (m_parentProcess->m_nativeProcessId != 0)
+				return true;
+#elif PLATFORM_MAC
+			if (m_gotExitMessage)
+				return false;
+			// No child should outlive the parent but there could be a bit of race here between threads
+			if (m_parentProcess->m_hasExited)
+			{
+				u64 timeSinceParentExitedMs = TimeToMs(GetTime() - m_parentProcess->m_processStats.exitTime);
+				if (timeSinceParentExitedMs > 5*1000)
+				{
+					m_session.m_logger.Info(TC("Child process %s outlived parent for 5 seconds. Will treat it as not active"), m_realApplication.c_str());
+					return false;
+				}
+			}
+#endif
+		}
 
 		siginfo_t signalInfo;
 		while (m_nativeProcessId != 0)
@@ -665,6 +678,10 @@ namespace uba
 				// We assume children are running one at the time.. this is an so-so estimation
 				childMaxPeakMemory = Max(childMaxPeakMemory, child.m_processStats.peakMemory.load());
 				childMaxDetoursMemory = Max(childMaxDetoursMemory, child.m_processStats.detoursMemory.load());
+
+				u64 size = m_trackedInputs.size();
+				m_trackedInputs.resize(size + child.m_trackedInputs.size());
+				memcpy(m_trackedInputs.data() + size, child.m_trackedInputs.data(), child.m_trackedInputs.size());
 			}
 			AtomicMax(m_processStats.peakMemory, m_processStats.peakMemory + childMaxPeakMemory);
 			AtomicMax(m_processStats.detoursMemory, m_processStats.detoursMemory + childMaxDetoursMemory);
@@ -937,6 +954,7 @@ namespace uba
 		CallSession(MessageType_Init, m_session.GetInitResponse(response, msg));
 		writer.WriteU32(m_id);
 		writer.WriteBool(m_parentProcess != nullptr);
+		writer.WriteBool(m_startInfo.trackInputs);
 		writer.WriteString(m_startInfo.application);
 		writer.WriteString(m_startInfo.workingDir);
 		writer.WriteU64(response.sharedMemoryAllocatorHandle.ToU64());
@@ -1997,7 +2015,6 @@ namespace uba
 			payload.allowDiscardVirtualMemory = m_session.m_allowDiscardVirtualMemory;
 			payload.suppressLogging = m_session.m_suppressLogging;
 			payload.isChild = m_parentProcess != nullptr;
-			payload.trackInputs = m_startInfo.trackInputs;
 			payload.useCustomAllocator = allowCustomAllocator && m_startInfo.useCustomAllocator && m_startInfo.rules->AllowMiMalloc();
 			payload.reportAllExceptions = m_startInfo.reportAllExceptions || m_startInfo.rules->ReportAllExceptions();
 			payload.isRunningWine = IsRunningWine();
diff --git a/Engine/Source/Programs/UnrealBuildAccelerator/Common/Private/UbaScheduler.cpp b/Engine/Source/Programs/UnrealBuildAccelerator/Common/Private/UbaScheduler.cpp
--- a/Engine/Source/Programs/UnrealBuildAccelerator/Common/Private/UbaScheduler.cpp
+++ b/Engine/Source/Programs/UnrealBuildAccelerator/Common/Private/UbaScheduler.cpp
@@ -1,5 +1,9 @@
 // Copyright Epic Games, Inc. All Rights Reserved.
 
+#if PLATFORM_WINDOWS
+#include <winternl.h>
+#endif
+
 #include "UbaScheduler.h"
 #include "UbaApplicationRules.h"
 #include "UbaCacheClient.h"
@@ -13,10 +17,6 @@
 #include "UbaSessionServer.h"
 #include "UbaStringBuffer.h"
 
-#if PLATFORM_WINDOWS
-#include <winternl.h>
-#endif
-
 // TODO. This is how we want scheduler to work.
 // 
 // 1. Cache should be queried from front of queue.
@@ -163,6 +163,8 @@ namespace uba
 
 			table->GetValueAsBool(memWatchdog, TC("MemWatchdog"));
 
+			table->GetValueAsBool(writeTrackedInputs, TC("WriteTrackedInputs"));
+
 			u32 temp = 0;
 			if (table->GetValueAsU32(temp, TC("MemTracingLevel")))
 				memTracingLevel = u8(temp);
@@ -186,6 +188,7 @@ namespace uba
 	,	m_forceRemote(info.forceRemote)
 	,	m_forceNative(info.forceNative)
 	,	m_useThreadToExitRemoteProcess(info.useThreadToExitRemoteProcess)
+	,	m_writeTrackedInputs(info.writeTrackedInputs)
 	,	m_processConfigs(info.processConfigs)
 	,	m_writeToCache(info.writeToCache && info.cacheClientCount)
 	{
@@ -204,6 +207,9 @@ namespace uba
 			m_memPagefileActivityThreshold = info.memPagefileActivityThreshold;
 		}
 
+		if (m_writeTrackedInputs)
+			DirectoryCache().CreateDirectoryW(m_session.GetLogger(), StringBuffer<>(m_session.GetSessionDir()).Append(TCV("Inputs")).data);
+
 		m_session.SetOuterScheduler(this);
 	}
 
@@ -362,6 +368,8 @@ namespace uba
 		bool canDetour = info.canDetour;
 		bool canExecuteRemotely = info.canExecuteRemotely && info.canDetour;
 
+		info2->trackInputs |= m_writeTrackedInputs;
+
 		if (m_processConfigs)
 		{
 			auto name = info2->application;
@@ -716,6 +724,28 @@ namespace uba
 
 		if (!m_useThreadToExitRemoteProcess || info->isLocal)
 		{
+			/*
+			static int counter = 0;
+			StringBuffer<> buf;
+			buf.Appendf("[REDACTED_USER_HOME]/access/action%u.txt", counter++);
+			FileAccessor f(m_session.GetLogger(), buf.data);
+			if (f.CreateWrite())
+			{
+				Set<TString> files;
+				auto& inputs = handle.m_process->GetTrackedInputs();
+				BinaryReader reader(inputs.data(), 0, inputs.size());
+				while (reader.GetLeft())
+					files.insert(reader.ReadString());
+
+				for (auto& input : files)
+				{
+					f.Write(input.data(), input.size());
+					f.Write("\n", 1);
+				}
+				f.Close();
+			}
+			*/
+
 			ExitProcess(*info, *handle.m_process, handle.m_process->GetExitCode(), false);
 			return;
 		}
@@ -1254,6 +1284,34 @@ namespace uba
 			if (exitCode != 0)
 				++m_errorCount;
 
+			if (m_writeTrackedInputs)
+			{
+				static int counter = 0;
+				StringBuffer<> inputsFile;
+				inputsFile.Append(m_session.GetSessionDir()).Append(TCV("inputs")).EnsureEndsWithSlash();
+				static u32 processId = 1; // TODO: This should be done in a better way.. or not at all?
+				GenerateNameForProcess(inputsFile, process.GetStartInfo().arguments, ++processId);
+				inputsFile.Append(TCV(".log"));
+				FileAccessor f(m_session.GetLogger(), inputsFile.data);
+				if (f.CreateWrite())
+				{
+					Set<TString> files;
+					auto& inputs = process.GetTrackedInputs();
+					BinaryReader reader(inputs.data(), 0, inputs.size());
+					while (reader.GetLeft())
+						files.insert(reader.ReadString());
+					for (auto& input : files)
+					{
+						u8 buffer[1024];
+						BinaryWriter writer(buffer, 0, sizeof(buffer));
+						writer.WriteUtf8String(input.data(), input.size());
+						writer.WriteByte('\n');
+						f.Write(buffer, writer.GetPosition());
+					}
+					f.Close();
+				}
+			}
+
 			FinishProcess(ph);
 			delete[] dependencies;
 			delete si;
diff --git a/Engine/Source/Programs/UnrealBuildAccelerator/Common/Public/UbaScheduler.h b/Engine/Source/Programs/UnrealBuildAccelerator/Common/Public/UbaScheduler.h
--- a/Engine/Source/Programs/UnrealBuildAccelerator/Common/Public/UbaScheduler.h
+++ b/Engine/Source/Programs/UnrealBuildAccelerator/Common/Public/UbaScheduler.h
@@ -38,6 +38,7 @@ namespace uba
 		bool writeToCache = false; // Set to true in combination with setting cacheClient to populate cache
 		bool useThreadToExitRemoteProcess = false; // Set this to true to move exit process call from job to thread. 
 		bool memWatchdog = false; // Will watch memory and throttle/kill processes if needed
+		bool writeTrackedInputs = false; // Set to true to write out log files with tracked inputs
 		u32 maxRacingPercent = 0u; // If processor weights go below this percent, then we will start filling up with racing processes
 		u8 memTracingLevel = 0;
 		u8 memStartWaitPercent = 85; // When memory usage goes above this percent, no new processes will be spawned until back below
@@ -189,6 +190,7 @@ namespace uba
 		bool m_forceNative;
 		bool m_useThreadToExitRemoteProcess;
 		bool m_allowDisableRemoteExecution = false;
+		bool m_writeTrackedInputs = false;
 		bool m_cancelled = false;
 		ConfigTable* m_processConfigs = nullptr;
 
diff --git a/Engine/Source/Programs/UnrealBuildAccelerator/Core/Public/UbaDetoursPayload.h b/Engine/Source/Programs/UnrealBuildAccelerator/Core/Public/UbaDetoursPayload.h
--- a/Engine/Source/Programs/UnrealBuildAccelerator/Core/Public/UbaDetoursPayload.h
+++ b/Engine/Source/Programs/UnrealBuildAccelerator/Core/Public/UbaDetoursPayload.h
@@ -20,7 +20,6 @@ namespace uba
 		u32 version = 0;
 		u32 rulesIndex = 0;
 		bool runningRemote = false;
-		bool trackInputs = false;
 		bool useCustomAllocator = true;
 		bool isRunningWine = false;
 		bool isChild = false;
diff --git a/Engine/Source/Programs/UnrealBuildAccelerator/Core/Public/UbaProtocol.h b/Engine/Source/Programs/UnrealBuildAccelerator/Core/Public/UbaProtocol.h
--- a/Engine/Source/Programs/UnrealBuildAccelerator/Core/Public/UbaProtocol.h
+++ b/Engine/Source/Programs/UnrealBuildAccelerator/Core/Public/UbaProtocol.h
@@ -50,7 +50,7 @@ namespace uba
 		#undef UBA_PROCESS_MESSAGE
 	};
 
-	inline constexpr u32 ProcessMessageVersion = 1347;
+	inline constexpr u32 ProcessMessageVersion = 1348;
 
 	inline constexpr u32 CommunicationMemSize = IsWindows ? 64*1024 : 64*1024*2; // Macos expands some commandlines to be crazy long
 
diff --git a/Engine/Source/Programs/UnrealBuildAccelerator/Detours/Private/PosixOS/UbaDetoursFunctionsPosix.cpp b/Engine/Source/Programs/UnrealBuildAccelerator/Detours/Private/PosixOS/UbaDetoursFunctionsPosix.cpp
--- a/Engine/Source/Programs/UnrealBuildAccelerator/Detours/Private/PosixOS/UbaDetoursFunctionsPosix.cpp
+++ b/Engine/Source/Programs/UnrealBuildAccelerator/Detours/Private/PosixOS/UbaDetoursFunctionsPosix.cpp
@@ -353,6 +353,8 @@ int Shared_open(const char* funcName, const char* file, int flags, int mode, con
 	{
 		int res = trueOpen(file, flags, mode);
 		DEBUG_LOG_TRUE(funcName, "NODETOUR%s (%s) -> %i", isWriteStr, file, res);
+		if (res != -1)
+			TrackInput(fileName);
 		return res;
 	}
 
@@ -518,6 +520,8 @@ int Shared_open(const char* funcName, const char* file, int flags, int mode, con
 	if (fd == -1)
 		return fd;
 
+	TrackInput(fileName);
+
 	SCOPED_WRITE_LOCK(g_fileHandlesLock, lock);
 	auto insres2 = g_fileHandles.insert({ fd, DetouredHandle() });
 	UBA_ASSERTF(insres2.second, "File handle already added");
@@ -545,6 +549,8 @@ FILE* Shared_fopen(const char* funcName, const char* path, const char* mode, con
 		{
 			FILE* res = TRUE_WRAPPER(fopen)(path, mode);
 			DEBUG_LOG_TRUE(funcName, "(%s  %s) -> %p", path, mode, res);
+			if (res)
+				TrackInput(ToView(path));
 			return res;
 		}
 		//UBA_ASSERTF(false, "%s with append (%s) not implemented (%s)", funcName, mode, path);
@@ -1229,6 +1235,8 @@ UBA_EXPORT int UBA_WRAPPER(openat)(int dirfd, const char* pathname, int flags, .
 	
 	int res = TRUE_WRAPPER(openat)(dirfd, pathname, flags, mode);
 	DEBUG_LOG_TRUE("openat", "(NODETOUR) %s -> %s", pathname, res);
+	if (res)
+		TrackInput(ToView(pathname));
 	return res;
 }
 
@@ -2846,6 +2854,7 @@ namespace uba
 		UBA_ASSERTF(!g_isInitialized, "Already initialized");
 		g_isInitialized = true;
 
+		bool trackInputs = false;
 		u64 directoryTableHandle;
 		u32 directoryTableSize;
 		u32 directoryTableCount;
@@ -2865,6 +2874,7 @@ namespace uba
 
 			g_processId = reader.ReadU32();
 			g_isChild = reader.ReadBool();
+			trackInputs = reader.ReadBool();
 
 			reader.ReadString(g_virtualApplication);
 			reader.ReadString(g_virtualWorkingDir);
@@ -2929,6 +2939,11 @@ namespace uba
 			g_overlayFilesSize = overlayFilesSize;
 		}
 
+		if (trackInputs)
+			InitTrackInput();
+
+		TrackInput(g_virtualApplication);
+
 		if (g_isChild)
 			Rpc_GetWrittenFiles();
 
@@ -2977,7 +2992,9 @@ namespace uba
 				Rpc_UpdateCloseHandle(path, fo->closeId, fo->deleteOnClose, fo->newName.c_str(), memoryHandle, memoryWritten, true);
 			}
 		}
-		
+
+		SendInput();
+
 		BinaryWriter writer;
 		writer.WriteByte(MessageType_Exit);
 		writer.WriteU32(0); // Exit code
diff --git a/Engine/Source/Programs/UnrealBuildAccelerator/Detours/Private/UbaDetoursShared.cpp b/Engine/Source/Programs/UnrealBuildAccelerator/Detours/Private/UbaDetoursShared.cpp
--- a/Engine/Source/Programs/UnrealBuildAccelerator/Detours/Private/UbaDetoursShared.cpp
+++ b/Engine/Source/Programs/UnrealBuildAccelerator/Detours/Private/UbaDetoursShared.cpp
@@ -9,6 +9,8 @@
 
 namespace uba
 {
+	constexpr u32 TrackInputsMemCapacity = 512 * 1024;
+
 	VARIABLE_MEM(StringBuffer<512>, g_virtualApplication);
 	VARIABLE_MEM(StringBuffer<512>, g_virtualApplicationDir);
 	VARIABLE_MEM(ProcessStats, g_stats);
@@ -35,6 +37,8 @@ namespace uba
 	bool g_suppressLogging = false;
 	u32 g_vfsEntryCount;
 	u32 g_vfsMatchingLength;
+	u8* g_trackInputsMem;
+	u32 g_trackInputsBufPos;
 
 	void InitSharedVariables()
 	{
@@ -174,6 +178,57 @@ namespace uba
 	SuppressDetourScope::SuppressDetourScope() { ++t_disallowDetour; }
 	SuppressDetourScope::~SuppressDetourScope() { --t_disallowDetour; }
 
+	void InitTrackInput()
+	{
+		g_trackInputsMem = (u8*)g_memoryBlock.Allocate(TrackInputsMemCapacity, 1, TC("TrackInputs"));
+	}
+
+	void SendInput()
+	{
+		if (!g_trackInputsMem)
+			return;
+
+		u32 left = g_trackInputsBufPos;
+		u32 reserveSize = left;
+		u32 pos = 0;
+		while (left)
+		{
+			RPC_MESSAGE(InputDependencies, log)
+				writer.Write7BitEncoded(reserveSize);
+			reserveSize = 0;
+			u32 toWrite = Min(left, u32(writer.GetCapacityLeft() - sizeof(u32)));
+			writer.WriteU32(toWrite);
+			writer.WriteBytes(g_trackInputsMem + pos, toWrite);
+			writer.Flush();
+			left -= toWrite;
+			pos += toWrite;
+		}
+		g_trackInputsBufPos = 0;
+	}
+
+
+	void TrackInput(StringView file)
+	{
+		if (!g_trackInputsMem)
+			return;
+
+		if (g_trackInputsBufPos > TrackInputsMemCapacity - 2048)
+			SendInput();
+
+		BinaryWriter w(g_trackInputsMem, g_trackInputsBufPos, TrackInputsMemCapacity);
+		w.WriteString(file);
+		g_trackInputsBufPos = u32(w.GetPosition());
+	}
+
+	void SkipTrackInput(StringView file)
+	{
+		// Just here to easily log out what we are ignoring in terms of input
+	}
+
+	bool IsTrackingInput()
+	{
+		return g_trackInputsMem != 0;
+	}
 
 	bool FixPath(StringBufferBase& out, const tchar* path, u32 pathLen)
 	{
diff --git a/Engine/Source/Programs/UnrealBuildAccelerator/Detours/Private/UbaDetoursShared.h b/Engine/Source/Programs/UnrealBuildAccelerator/Detours/Private/UbaDetoursShared.h
--- a/Engine/Source/Programs/UnrealBuildAccelerator/Detours/Private/UbaDetoursShared.h
+++ b/Engine/Source/Programs/UnrealBuildAccelerator/Detours/Private/UbaDetoursShared.h
@@ -133,6 +133,11 @@ namespace uba
 	const tchar* Shared_GetFileAttributes(FileAttributes& outAttr, StringView fileName, bool checkIfDir = false);
 
 	void InitSharedVariables();
+	void InitTrackInput();
+	void TrackInput(StringView file);
+	void SkipTrackInput(StringView file);
+	bool IsTrackingInput();
+	void SendInput();
 
 	template<typename T> struct VariableMem { template<typename... Args> void Create(Args&&... args) { new (data) T(args...); }; u64 data[AlignUp(sizeof(T), sizeof(u64)) / sizeof(u64)]; };
 	#define VARIABLE_MEM(type, name) VariableMem<type> name##Mem; type& name = (type&)name##Mem.data;
diff --git a/Engine/Source/Programs/UnrealBuildAccelerator/Detours/Private/Windows/UbaDetoursFunctionsKernelBase.inl b/Engine/Source/Programs/UnrealBuildAccelerator/Detours/Private/Windows/UbaDetoursFunctionsKernelBase.inl
--- a/Engine/Source/Programs/UnrealBuildAccelerator/Detours/Private/Windows/UbaDetoursFunctionsKernelBase.inl
+++ b/Engine/Source/Programs/UnrealBuildAccelerator/Detours/Private/Windows/UbaDetoursFunctionsKernelBase.inl
@@ -2804,7 +2804,7 @@ HMODULE Recursive_LoadLibraryExW(const StringView& filePath, LPCWSTR originalNam
 		originalName = newName.data;
 	}
 
-	TrackInput(originalName);
+	TrackInput(ToView(originalName));
 
 	DEBUG_LOG_TRUE(L"INTERNAL LoadLibraryExW", L"%ls", originalName);
 
@@ -2854,7 +2854,7 @@ HMODULE GetModuleHandleNoSyscall(const tchar* moduleName)
 
 HMODULE Shared_LoadLibrary(LPCWSTR lpLibFileName, HANDLE hFile, DWORD dwFlags, bool& mightExist)
 {
-	if (!g_runningRemote && !g_trackInputsMem || t_disallowDetour)
+	if (!g_runningRemote && !IsTrackingInput() || t_disallowDetour)
 		return NULL;
 
 	StringBuffer<> path;
diff --git a/Engine/Source/Programs/UnrealBuildAccelerator/Detours/Private/Windows/UbaDetoursFunctionsNtDll.inl b/Engine/Source/Programs/UnrealBuildAccelerator/Detours/Private/Windows/UbaDetoursFunctionsNtDll.inl
--- a/Engine/Source/Programs/UnrealBuildAccelerator/Detours/Private/Windows/UbaDetoursFunctionsNtDll.inl
+++ b/Engine/Source/Programs/UnrealBuildAccelerator/Detours/Private/Windows/UbaDetoursFunctionsNtDll.inl
@@ -901,9 +901,9 @@ NTSTATUS NTAPI Shared_NtCreateFile(bool IsCreateFunc, PHANDLE hFileHandle, ACCES
 			return res;
 
 		if (!isSystemOrTempFile && !isWrite && !t_disallowDetour && fileName[fileName.count-1] != ':')
-			TrackInput(fileName.data);
+			TrackInput(fileName);
 		else
-			SkipTrackInput(fileName.data);
+			SkipTrackInput(fileName);
 		return res;
 	}
 
@@ -1159,12 +1159,12 @@ NTSTATUS NTAPI Shared_NtCreateFile(bool IsCreateFunc, PHANDLE hFileHandle, ACCES
 				if (!info.tracked)
 				{
 					info.tracked = true;
-					TrackInput(fileName.data);
+					TrackInput(fileName);
 				}
 			}
 			else
 			{
-				SkipTrackInput(fileName.data);
+				SkipTrackInput(fileName);
 			}
 		};
 
diff --git a/Engine/Source/Programs/UnrealBuildAccelerator/Detours/Private/Windows/UbaDetoursFunctionsWin.cpp b/Engine/Source/Programs/UnrealBuildAccelerator/Detours/Private/Windows/UbaDetoursFunctionsWin.cpp
--- a/Engine/Source/Programs/UnrealBuildAccelerator/Detours/Private/Windows/UbaDetoursFunctionsWin.cpp
+++ b/Engine/Source/Programs/UnrealBuildAccelerator/Detours/Private/Windows/UbaDetoursFunctionsWin.cpp
@@ -129,48 +129,6 @@ bool g_isRunningWine;
 int g_uiLanguage;
 u32 g_processId;
 
-constexpr u32 TrackInputsMemCapacity = 512 * 1024;
-u8* g_trackInputsMem;
-u32 g_trackInputsBufPos;
-
-void SendInput()
-{
-	u32 left = g_trackInputsBufPos;
-	u32 reserveSize = left;
-	u32 pos = 0;
-	while (left)
-	{
-		RPC_MESSAGE(InputDependencies, log)
-		writer.Write7BitEncoded(reserveSize);
-		reserveSize = 0;
-		u32 toWrite = Min(left, u32(writer.GetCapacityLeft() - sizeof(u32)));
-		writer.WriteU32(toWrite);
-		writer.WriteBytes(g_trackInputsMem + pos, toWrite);
-		writer.Flush();
-		left -= toWrite;
-		pos += toWrite;
-	}
-	g_trackInputsBufPos = 0;
-}
-
-
-void TrackInput(const wchar_t* file)
-{
-	if (!g_trackInputsMem)
-		return;
-
-	if (g_trackInputsBufPos > TrackInputsMemCapacity - 2048)
-		SendInput();
-
-	BinaryWriter w(g_trackInputsMem, g_trackInputsBufPos, TrackInputsMemCapacity);
-	w.WriteString(file);
-	g_trackInputsBufPos = u32(w.GetPosition());
-}
-void SkipTrackInput(const wchar_t* file)
-{
-	// Just here to easily log out what we are ignoring in terms of input
-}
-
 u8 g_emptyMemoryFileMem;
 MemoryFile& g_emptyMemoryFile = *new MemoryFile(&g_emptyMemoryFileMem, true);
 
@@ -602,8 +560,7 @@ void SendExitMessage(DWORD exitCode, u64 startTime)
 	#endif
 	#endif
 
-	if (g_trackInputsMem)
-		SendInput();
+	SendInput();
 
 	g_stats.detoursMemory = g_memoryBlock.writtenSize; // + g_directoryTable.m_memorySize + g_mappedFileTable.m_memPosition;
 
@@ -962,9 +919,6 @@ void Init(const DetoursPayload& payload, u64 startTime)
 		g_stdHandle[2] = makeDetouredHandle(new DetouredHandle(HandleType_StdIn)); // STD_IN
 	}
 
-	if (payload.trackInputs)
-		g_trackInputsMem = (u8*)g_memoryBlock.Allocate(TrackInputsMemCapacity, 1, L"TrackInputs");
-	
 	g_systemRoot.count = True_GetEnvironmentVariableW(L"SystemRoot", g_systemRoot.data, g_systemRoot.capacity);
 	g_systemRoot.MakeLower();
 	UBA_ASSERT(g_systemRoot.count);
@@ -979,6 +933,7 @@ void Init(const DetoursPayload& payload, u64 startTime)
 	StringBuffer<512> applicationBuffer;
 	StringBuffer<512> workingDirBuffer;
 
+	bool trackInputs = false;
 	HANDLE sharedMemoryAllocatorHandle;
 	SharedMemoryHandle permanentFilesHandle;
 	HANDLE directoryTableHandle;
@@ -997,6 +952,7 @@ void Init(const DetoursPayload& payload, u64 startTime)
 
 		g_processId = reader.ReadU32();
 		g_isChild = reader.ReadBool();
+		trackInputs = reader.ReadBool();
 
 		reader.ReadString(applicationBuffer);
 		reader.ReadString(workingDirBuffer);
@@ -1021,7 +977,10 @@ void Init(const DetoursPayload& payload, u64 startTime)
 		DEBUG_LOG_PIPE(L"Init", L"");
 	}
 
-	TrackInput(applicationBuffer.data);
+	if (trackInputs)
+		InitTrackInput();
+
+	TrackInput(applicationBuffer);
 
 	VirtualizePath(applicationBuffer);
 	VirtualizePath(workingDirBuffer);
