[UBA]
* A fix for the include ordering fix for linux (Gah to windows macros for CreateDirectory)

diff --git a/Engine/Source/Programs/UnrealBuildAccelerator/Common/Private/UbaScheduler.cpp b/Engine/Source/Programs/UnrealBuildAccelerator/Common/Private/UbaScheduler.cpp
--- a/Engine/Source/Programs/UnrealBuildAccelerator/Common/Private/UbaScheduler.cpp
+++ b/Engine/Source/Programs/UnrealBuildAccelerator/Common/Private/UbaScheduler.cpp
@@ -209,7 +209,7 @@ namespace uba
 		}
 
 		if (m_writeTrackedInputs)
-			DirectoryCache().CreateDirectoryW(m_session.GetLogger(), StringBuffer<>(m_session.GetSessionDir()).Append(TCV("Inputs")).data);
+			DirectoryCache().CreateDirectory(m_session.GetLogger(), StringBuffer<>(m_session.GetSessionDir()).Append(TCV("Inputs")).data);
 
 		m_session.SetOuterScheduler(this);
 	}
