[UBA]
* Implemented d2d/directwrite visualizer backend. For big traces frame time goes from 80ms to 10ms

diff --git a/Engine/Source/Programs/UnrealBuildAccelerator/Common/Private/UbaConfig.cpp b/Engine/Source/Programs/UnrealBuildAccelerator/Common/Private/UbaConfig.cpp
--- a/Engine/Source/Programs/UnrealBuildAccelerator/Common/Private/UbaConfig.cpp
+++ b/Engine/Source/Programs/UnrealBuildAccelerator/Common/Private/UbaConfig.cpp
@@ -205,12 +205,14 @@ namespace uba
 
 		auto consumeIdentifier = [&](StringBufferBase& out) -> char
 			{
+				const char* first = i;
 				while (i != e)
 				{
 					tchar c = *i;
 					bool validChar = (c >= 'a' && c <= 'z') || (c >= 'A' && c <= 'Z') || c == '_' || c == '-';
 					if (!validChar)
-						return *i;
+						if (first == i || c < '0' || c > '9')
+							return *i;
 					out.Append(*i);
 					++i;
 				}
@@ -272,7 +274,7 @@ namespace uba
 				consumeIdentifier(key);
 				token = consumeEmpty();
 				if (token != '=')
-					return logger.Error(TC("Unexpected equals sign after key name %s"), key);
+					return logger.Error(TC("Unexpected equals sign after key name %s"), key.data);
 				++i;
 				token = consumeEmpty();
 
diff --git a/Engine/Source/Programs/UnrealBuildAccelerator/Visualizer/Private/Mac/UbaVisualizerApp.mm b/Engine/Source/Programs/UnrealBuildAccelerator/Visualizer/Private/Mac/UbaVisualizerApp.mm
--- a/Engine/Source/Programs/UnrealBuildAccelerator/Visualizer/Private/Mac/UbaVisualizerApp.mm
+++ b/Engine/Source/Programs/UnrealBuildAccelerator/Visualizer/Private/Mac/UbaVisualizerApp.mm
@@ -25,12 +25,10 @@ virtual SIZE GetClientSize() override
             {
                 return clientSize;
             }
-			virtual void SetBkMode(int mode) override {}
 			virtual void SetPen(HPEN pen) override {}
 			virtual void SetBrush(HBRUSH b) override {}
 			virtual void SetFont(const Font& font, bool underlined = false) override { activeFont = font; }
 			virtual void SetTextColor(COLORREF color) override {}
-			virtual u32 SetTextAlign(u32 align) override { return 0; }
 
             int textOffsetX = 0;
             int textOffsetY = 0;
@@ -79,7 +77,7 @@ virtual void FillRect(const RECT& r, HBRUSH b)
             }
 			virtual void Rectangle(int left, int top, int right, int bottom) {}
 			virtual void Line(int fromX, int fromY, int toX, int toY) {}
-			virtual void Polyline(const Vector<POINT>& polyLine) {}
+			virtual void Polyline(u32 cacheId, const Vector<POINT>& polyLine) {}
 			virtual void PolyPolyline(const POINT* points, const u32* lines, u32 lineCount) {}
 			virtual void DrawCenteredText(const StringBuffer<>* rows, u32 rowCount, HBRUSH backgroundBrush) {}
 			virtual SIZE GetTextSize(StringView str)
diff --git a/Engine/Source/Programs/UnrealBuildAccelerator/Visualizer/Private/UbaVisualizer.cpp b/Engine/Source/Programs/UnrealBuildAccelerator/Visualizer/Private/UbaVisualizer.cpp
--- a/Engine/Source/Programs/UnrealBuildAccelerator/Visualizer/Private/UbaVisualizer.cpp
+++ b/Engine/Source/Programs/UnrealBuildAccelerator/Visualizer/Private/UbaVisualizer.cpp
@@ -81,7 +81,7 @@ namespace uba
 				tr.left += line.left;
 				tr.top += line.top;
 				context.SetTextColor(line.color);
-				context.Text(tr.left, tr.top, nullptr, line.str);
+				context.Text(tr.left, tr.top, line.str, nullptr);
 			}
 		}
 
@@ -112,6 +112,35 @@ namespace uba
 
 	////////////////////////////////////////////////////////////////////////////////////////////////////
 
+	void DrawCenteredText(VisualizerBackend::PaintContext& context, const StringBuffer<>* rows, u32 rowCount, HBRUSH backgroundBrush)
+	{
+		SIZE clientSize = context.GetClientSize();
+		for (u32 i=0; i!=rowCount; ++i)
+		{ 
+			SIZE textSize = context.GetTextSize(rows[i]);
+
+			RECT rect = { 0, 0, 0, 0 };
+			int top = (clientSize.cy - textSize.cy)/2 + i*20;
+			rect.left = (clientSize.cx - textSize.cx)/2;
+			rect.right = (clientSize.cx + textSize.cx)/2;
+			rect.bottom = top + textSize.cy;
+			rect.top = top;
+
+			if (backgroundBrush)
+			{
+				RECT r = rect;
+				r.left -= 4;
+				r.right += 4;
+				r.top -= 2;
+				r.bottom += 2;
+				context.FillRect(r, backgroundBrush);
+			}
+			context.Text(rect.left, rect.top, rows[i], nullptr);
+		}
+	}
+
+	////////////////////////////////////////////////////////////////////////////////////////////////////
+
 	bool Visualizer::Stats::operator==(const Stats& o) const
 	{
 		if (memcmp(this, &o, offsetof(Stats, connectionCount)) != 0)
@@ -142,6 +171,113 @@ namespace uba
 	{
 	}
 
+	void Visualizer::InitBrushes()
+	{
+		auto CreateBrush = [&](u8 r, u8 g, u8 b) { return m_backend.CreateSolidBrush(RGB(r, g, b)); };
+
+		if (m_config.DarkMode)
+		{
+			m_textColor = RGB(190, 190, 190);
+			m_textWarningColor = RGB(190, 190, 0);
+			m_textErrorColor = RGB(190, 0, 0);
+
+			m_processBrushes[0].inProgress = CreateBrush(70, 70, 70);
+			m_processBrushes[1].inProgress = CreateBrush(130, 130, 130);
+			m_processBrushes[0].inProgressRace = CreateBrush(50, 50, 50);
+			m_processBrushes[1].inProgressRace = CreateBrush(120, 120, 120);
+
+			m_processBrushes[0].error = CreateBrush(140, 0, 0);
+			m_processBrushes[1].error = CreateBrush(190, 0, 0);
+
+			m_processBrushes[0].returned = CreateBrush(50, 50, 120);
+			m_processBrushes[1].returned = CreateBrush(70, 70, 160);
+
+			m_processBrushes[0].recv = CreateBrush(10, 92, 10);
+			m_processBrushes[1].recv = CreateBrush(10, 130, 10);
+			m_processBrushes[0].success = CreateBrush(10, 100, 10);
+			m_processBrushes[1].success = CreateBrush(10, 140, 10);
+			m_processBrushes[0].send = CreateBrush(10, 115, 10);
+			m_processBrushes[1].send = CreateBrush(10, 145, 10);
+			m_processBrushes[0].cacheFetchTest = CreateBrush(24, 120, 110);
+			m_processBrushes[1].cacheFetchTest = CreateBrush(31, 150, 138);
+			m_processBrushes[0].cacheFetchDownload = CreateBrush(24, 112, 110);
+			m_processBrushes[1].cacheFetchDownload = CreateBrush(31, 143, 138);
+			m_processBrushes[0].warning = CreateBrush(100, 100, 0);
+			m_processBrushes[1].warning = CreateBrush(100, 100, 0);
+
+			m_backgroundBrush = CreateBrush(38, 37, 37);
+			m_separatorPen = m_backend.CreatePen(RGB(50, 50, 50));
+			m_tooltipBackgroundBrush = CreateBrush(64, 64, 64);
+			m_checkboxPen = m_backend.CreatePen(RGB(130, 130, 130));
+
+			m_sendColor = RGB(0, 170, 0);
+			m_recvColor = RGB(0, 170, 255);
+			m_cpuColor = RGB(170, 170, 0);
+			m_memColor = RGB(170, 0, 255);
+			m_driveColor = RGB(170, 65, 55);
+			m_activeProcColor = RGB(0, 170, 170);
+		}
+		else
+		{
+			#if PLATFORM_WINDOWS
+			m_textColor = GetSysColor(COLOR_INFOTEXT);
+			COLORREF backgroundColor = GetSysColor(0);
+			COLORREF tooltipBackgroundColor = GetSysColor(COLOR_INFOBK);
+			#else
+			m_textColor = RGB(0, 0, 0);
+			COLORREF backgroundColor = RGB(220,220,220);
+			COLORREF tooltipBackgroundColor = RGB(200,200,200);
+			#endif
+
+			m_textWarningColor = RGB(170, 130, 0);
+			m_textErrorColor = RGB(190, 0, 0);
+
+			m_processBrushes[0].inProgress = CreateBrush(150, 150, 150);
+			m_processBrushes[1].inProgress = CreateBrush(180, 180, 180);
+			m_processBrushes[0].inProgressRace = CreateBrush(130, 130, 130);
+			m_processBrushes[1].inProgressRace = CreateBrush(170, 170, 170);
+
+			m_processBrushes[0].error = CreateBrush(255, 70, 70);
+			m_processBrushes[1].error = CreateBrush(255, 100, 70);
+
+			m_processBrushes[0].returned = CreateBrush(150, 150, 200);
+			m_processBrushes[1].returned = CreateBrush(170, 170, 200);
+
+			m_processBrushes[0].recv = CreateBrush(10, 190, 10);
+			m_processBrushes[1].recv = CreateBrush(20, 210, 20);
+			m_processBrushes[0].success = CreateBrush(10, 200, 10);
+			m_processBrushes[1].success = CreateBrush(20, 220, 20);
+			m_processBrushes[0].send = CreateBrush(80, 210, 80);
+			m_processBrushes[1].send = CreateBrush(90, 250, 90);
+			m_processBrushes[0].cacheFetchTest = CreateBrush(150, 150, 200);
+			m_processBrushes[1].cacheFetchTest = CreateBrush(170, 170, 200);
+			m_processBrushes[0].cacheFetchDownload = CreateBrush(150, 150, 200);
+			m_processBrushes[1].cacheFetchDownload = CreateBrush(170, 170, 200);
+			m_processBrushes[0].warning = CreateBrush(170, 170, 0);
+			m_processBrushes[1].warning = CreateBrush(170, 170, 0);
+
+			m_separatorPen = m_backend.CreatePen(RGB(180, 180, 180));
+			m_backgroundBrush = m_backend.CreateSolidBrush(backgroundColor);
+			m_tooltipBackgroundBrush = m_backend.CreateSolidBrush(tooltipBackgroundColor);
+			m_checkboxPen = m_backend.CreatePen(RGB(130, 130, 130));
+
+			m_sendColor = RGB(0, 170, 0); // Green
+			m_recvColor = RGB(63, 72, 204); // Blue
+			m_cpuColor = RGB(200, 130, 0); // Orange
+			m_memColor = RGB(170, 0, 255); // Purple
+			m_driveColor = RGB(255, 115, 96);
+			m_activeProcColor = RGB(0, 170, 170);
+		}
+
+		m_textPen = m_backend.CreatePen(m_textColor);
+		m_sendPen = m_backend.CreatePen(m_sendColor);
+		m_recvPen = m_backend.CreatePen(m_recvColor);
+		m_cpuPen = m_backend.CreatePen(m_cpuColor);
+		m_memPen = m_backend.CreatePen(m_memColor);
+		m_drivePen = m_backend.CreatePen(m_driveColor);
+		m_activeProcPen = m_backend.CreatePen(m_activeProcColor);
+	}
+
 	u64 Visualizer::GetPlayTime()
 	{
 		u64 currentTime = m_paused ? m_pauseStart : GetTime();
@@ -200,7 +336,7 @@ namespace uba
 		out.Clear();
 		out.Append(TC("UbaVisualizer"));
 		#if UBA_DEBUG
-		out.Append(TC(" (DEBUG)"));
+		out.Appendf(TC(" (DEBUG, %s)"), m_backend.GetName());
 		#endif
 		out.Append(TC(" - "));
 		return out;
@@ -639,9 +775,8 @@ namespace uba
 	{
 		u64 playTime = GetPlayTime();
 
-		context.SetBkMode(TRANSPARENT);
 		context.SetTextColor(m_textColor);
-
+		context.SetPen(m_textPen);
 
 		if (m_traceView.sessions.empty())
 		{
@@ -656,7 +791,7 @@ namespace uba
 				StringBuffer<> str;
 				str.Appendf(TC("Unsupported trace version %u (Versions supported are %u to %u)"), m_traceView.version, TraceReadCompatibilityVersion, TraceVersion);
 				context.SetTextColor(m_textWarningColor);
-				context.DrawCenteredText(&str, 1, nullptr);
+				DrawCenteredText(context, &str, 1, nullptr);
 			}
 			else
 			{
@@ -691,19 +826,13 @@ namespace uba
 				rect.bottom = posY + activeFont.height + 2;
 				context.SetTextColor(type == LogEntryType_Info ? m_textColor : (type == LogEntryType_Error ? m_textErrorColor : m_textWarningColor));
 
-				u32 align = TA_LEFT;
 				if (centered)
-				{
-					posX += (endX - posX)/2;
-					align = context.SetTextAlign(TA_CENTER);
-				}
+					posX += ((endX - posX) - context.GetTextSize(text).cx)/2;
 				if (underlined)
 					context.SetFont(activeFont, true);
-				context.Text(posX, posY, &rect, text);
+				context.Text(posX, posY, text, &rect);
 				if (underlined)
 					context.SetFont(activeFont);
-				if (centered)
-					context.SetTextAlign(align);
 				if (moveY)
 					posY = rect.bottom;
 			};
@@ -856,7 +985,7 @@ namespace uba
 					time += 0.25f;
 				}
 				if (line.size() > 1)
-					context.Polyline(line);
+					context.Polyline(0, line);
 			}
 			posY += GraphHeight;
 		}
@@ -952,7 +1081,7 @@ namespace uba
 		u32 visibleBoxes = 0;
 
 		TraceView::ProcessLocation processLocation { 0, 0, 0 };
-		for (u64 sessionIt = 0, sessionEnd = m_traceView.sessions.size(); sessionIt != sessionEnd; ++sessionIt)
+		for (u32 sessionIt = 0, sessionEnd = u32(m_traceView.sessions.size()); sessionIt != sessionEnd; ++sessionIt)
 		{
 			bool isFirst = sessionIt == 0;
 			auto& session = *sortedSessions[sessionIt].session;
@@ -1011,15 +1140,18 @@ namespace uba
 					rect.top = posY;
 					rect.bottom = textBottom;
 
+					context.SetTextColor(m_textColor);
 					if (selected)
-						context.SetBkMode(OPAQUE);
+					{
+						RECT r2 = rect;
+						r2.right = r2.left + context.GetTextSize(text).cx;
+						context.FillRect(r2, m_tooltipBackgroundBrush);
+					}
 
 					SIZE textSize = context.GetTextSize(text);
 					if (textSize.cx)
 						session.fullNameWidth = u32(textSize.cx);
-					context.Text(5, posY+2, &rect, text);
-					if (selected)
-						context.SetBkMode(TRANSPARENT);
+					context.Text(5, posY+2, text, &rect);
 				}
 				posY += m_sessionStepY;
 			}
@@ -1035,7 +1167,7 @@ namespace uba
 
 					Vector<POINT> line;
 
-					auto drawGraph = [&](auto& values, auto maxValue, double scale, HPEN pen, bool accumulatingValue, int offsetY = 0)
+					auto drawGraph = [&](u32 id, auto& values, auto maxValue, double scale, HPEN pen, bool accumulatingValue, int offsetY = 0)
 						{
 							bool loop = true;
 							u32 reconnectIndex = 0;
@@ -1117,27 +1249,27 @@ namespace uba
 									prevTime = updateTime;
 								}
 								if (line.size() > 1)
-									context.Polyline(line);
+									context.Polyline(id + sessionIt*10, line);
 								++reconnectIndex;
 							}
 						};
 
 					if (m_config.showNetworkStats && session.highestSendPerS != 0 && session.highestRecvPerS != 0)
 					{
-						drawGraph(session.networkSend, 100000000000000ull, double(session.highestSendPerS) / graphHeight, m_sendPen, true);
-						drawGraph(session.networkRecv, 100000000000000ull, double(session.highestRecvPerS) / graphHeight, m_recvPen, true, 1);
+						drawGraph(1, session.networkSend, 100000000000000ull, double(session.highestSendPerS) / graphHeight, m_sendPen, true);
+						drawGraph(2, session.networkRecv, 100000000000000ull, double(session.highestRecvPerS) / graphHeight, m_recvPen, true, 1);
 					}
 
 					if (m_config.showCpuMemStats)
 					{
-						drawGraph(session.cpuLoad, 0.0f, -1.0, m_cpuPen, false);
-						drawGraph(session.memAvail, session.memTotal, 1.0, m_memPen, false);
+						drawGraph(3, session.cpuLoad, 0.0f, -1.0, m_cpuPen, false);
+						drawGraph(4, session.memAvail, session.memTotal, 1.0, m_memPen, false);
 					}
 
 					if (m_config.showDriveStats)
 						for (auto& pair : session.drives)
 							if (pair.second.busyHighest)
-								drawGraph(pair.second.busyPercent, 0, -0.01, m_drivePen, false);
+								drawGraph(5, pair.second.busyPercent, 0, -0.01, m_drivePen, false);
 
 				}
 				posY += GraphHeight;
@@ -1151,16 +1283,18 @@ namespace uba
 							return;
 						bool selected = m_fetchedFilesSelected == processLocation.sessionIndex && text.StartsWith(TC("Fetched Files"));
 						if (selected)
-							context.SetBkMode(OPAQUE);
+						{
+							RECT r2 = rect;
+							r2.right = r2.left + context.GetTextSize(text).cx;
+							context.FillRect(r2, m_tooltipBackgroundBrush);
+						}
 						if (rect.bottom > progressRect.bottom)
 						{
 							rect.bottom = progressRect.bottom;
-							context.Text(rect.left, rect.top, &rect, text);
+							context.Text(rect.left, rect.top, text, &rect);
 						}
 						else
-							context.Text(rect.left, rect.top, nullptr, text);
-						if (selected)
-							context.SetBkMode(TRANSPARENT);
+							context.Text(rect.left, rect.top, text, nullptr);
 						if (outWidth)
 							*outWidth = u32(context.GetTextSize(text).cx);
 					};
@@ -1257,12 +1391,12 @@ namespace uba
 											textCacheContext.SetTextColor(RGB(5, 60, 5));
 											++rect2.left;
 											++textY;
-											textCacheContext.Text(rect2.left, textY, &rect2, process.description);
+											textCacheContext.Text(rect2.left, textY, process.description, &rect2);
 											--rect2.left;
 											--textY;
 										}
 										textCacheContext.SetTextColor(m_textColor);
-										textCacheContext.Text(rect2.left, textY, &rect2, process.description);
+										textCacheContext.Text(rect2.left, textY, process.description, &rect2);
 
 										if (!selected)
 											process.bitmapDirty = false;
@@ -1280,7 +1414,7 @@ namespace uba
 
 							StringBuffer<> buf;
 							buf.AppendValue(u64(processLocation.processorIndex) + 1);
-							context.Text(5, posY + offsetY, &rect, buf);
+							context.Text(5, posY + offsetY, buf, &rect);
 						}
 					}
 
@@ -1329,7 +1463,7 @@ namespace uba
 
 							StringBuffer<> buf;
 							buf.AppendValue(u64(trackIndex) + 1);
-							context.Text(5, posY + offsetY, &rect, buf);
+							context.Text(5, posY + offsetY, buf, &rect);
 						}
 
 						int lastDrawnRight = 0;
@@ -1430,7 +1564,7 @@ namespace uba
 									{
 										RECT rect2{ 0,int(work.bitmapOffset), 256, int(work.bitmapOffset) + m_processFont.height };
 										textCacheContext.FillRect(rect2, brush);
-										textCacheContext.Text(rect2.left, rect2.top, &rect2, ToView(work.description));
+										textCacheContext.Text(rect2.left, rect2.top, ToView(work.description), &rect2);
 										if (!selected)
 											work.bitmapDirty = false;
 									});
@@ -1466,7 +1600,7 @@ namespace uba
 					if (int(x) >= progressRect.right)
 						break;
 
-					context.Text(x, y, nullptr, ToView(item));
+					context.Text(x, y, ToView(item), nullptr);
 
 					++rowIndex;
 					if (rowIndex > table.rowMaxCount || int(y) >= progressRect.bottom)
@@ -1495,7 +1629,7 @@ namespace uba
 			str[1].Appendf(TC("Box Count: %u"), visibleBoxes);
 
 			context.SetFont(m_popupFont);
-			context.DrawCenteredText(str, 2, m_tooltipBackgroundBrush);
+			DrawCenteredText(context, str, 2, m_tooltipBackgroundBrush);
 		}
 
 		if (m_config.showCursorLine && m_mouseOverWindow)
@@ -1583,7 +1717,7 @@ namespace uba
 				{
 					context.SetTextColor(color);
 					left -= (context.GetTextSize(text).cx) + 5;
-					context.Text(left, top, nullptr, text);
+					context.Text(left, top, text, nullptr);
 				};
 
 			if (m_config.showDriveStats)
@@ -2195,7 +2329,7 @@ namespace uba
 					textRect.bottom = textRect.top + m_timelineFont.height;
 					textRect.right = pos + s.cx/2;
 					textRect.left = pos - s.cx/2;
-					context.Text(textRect.left, textRect.top, nullptr, buffer);
+					context.Text(textRect.left, textRect.top, buffer, nullptr);
 				}
 			}
 			if (!(number % lineStepSize))
@@ -2491,6 +2625,24 @@ namespace uba
 
 	void Visualizer::PaintEmptyScreen(PaintContext& context)
 	{
+		StringBuffer<> str[2];
+		if (!m_fileName.IsEmpty())
+		{
+			str[0].Append(TCV("Loading trace file"));
+		}
+		else
+		{
+			str[1].Append(TCV("Click here to open trace file"));
+			if (m_listenChannel.count)
+			{
+				str[0].Appendf(TC("Listening for new sessions on channel '%s'"), m_listenChannel.data);
+				str[1].Append(TCV(" instead"));
+			}
+			else
+				str[0].Append(TCV("No trace active"));
+		}
+		context.SetFont(m_popupFont);
+		DrawCenteredText(context, str, 2, nullptr);
 	}
 
 	void Visualizer::FinishIncompatible()
@@ -2630,7 +2782,7 @@ namespace uba
 
 	void Visualizer::Reset()
 	{
-		m_backend.ResetCaches();
+		m_backend.Reset();
 		m_contentWidth = 0;
 		m_contentHeight = 0;
 		//m_autoScroll = true;
diff --git a/Engine/Source/Programs/UnrealBuildAccelerator/Visualizer/Private/UbaVisualizer.h b/Engine/Source/Programs/UnrealBuildAccelerator/Visualizer/Private/UbaVisualizer.h
--- a/Engine/Source/Programs/UnrealBuildAccelerator/Visualizer/Private/UbaVisualizer.h
+++ b/Engine/Source/Programs/UnrealBuildAccelerator/Visualizer/Private/UbaVisualizer.h
@@ -14,8 +14,8 @@ namespace uba
 	{
 	public:
 		Visualizer(VisualizerConfig& config, Logger& logger, VisualizerBackend& backend);
+		void InitBrushes();
 
-	protected:
 		u64 GetPlayTime();
 		int GetTimelineHeight();
 		int GetTimelineTop();
diff --git a/Engine/Source/Programs/UnrealBuildAccelerator/Visualizer/Private/UbaVisualizerBackend.h b/Engine/Source/Programs/UnrealBuildAccelerator/Visualizer/Private/UbaVisualizerBackend.h
--- a/Engine/Source/Programs/UnrealBuildAccelerator/Visualizer/Private/UbaVisualizerBackend.h
+++ b/Engine/Source/Programs/UnrealBuildAccelerator/Visualizer/Private/UbaVisualizerBackend.h
@@ -38,30 +38,39 @@ namespace uba
 		struct PaintContext
 		{
 			virtual SIZE GetClientSize() = 0;
-			virtual void SetBkMode(int mode) = 0;
 			virtual void SetPen(HPEN pen) = 0;
 			virtual void SetBrush(HBRUSH b) = 0;
 			virtual void SetFont(const Font& font, bool underlined = false) = 0;
 			virtual void SetTextColor(COLORREF color) = 0;
-			virtual u32 SetTextAlign(u32 align) = 0;
-			virtual void Text(int x, int y, const RECT* lprect, StringView str) = 0;
+			virtual void Text(int x, int y, StringView str, const RECT* clipRect) = 0;
 			virtual COLORREF GetTextColor() = 0;
 			virtual Font& GetFont() = 0;
 			virtual void FillRect(const RECT& r, HBRUSH b) = 0;
 			virtual void Rectangle(int left, int top, int right, int bottom) = 0;
 			virtual void Line(int fromX, int fromY, int toX, int toY) = 0;
-			virtual void Polyline(const Vector<POINT>& polyLine) = 0;
+			virtual void Polyline(u32 cacheId, const Vector<POINT>& polyLine) = 0;
 			virtual void PolyPolyline(const POINT* points, const u32* lines, u32 lineCount) = 0;
-			virtual void DrawCenteredText(const StringBuffer<>* rows, u32 rowCount, HBRUSH backgroundBrush) = 0;
 			virtual SIZE GetTextSize(StringView str) = 0;
 			virtual void DrawCached(TraceView::BitmapCache& process, RECT rect, int offsetY, int textHeight, Font& processFont, RECT clippingRect, const Function<void(PaintContext& textCacheContext)>& func) = 0;
 
 			virtual POINT GetCursorClientPos() = 0;
 		};
 
-		virtual void ResetCaches() = 0;
+		virtual const tchar* GetName() = 0;
+		virtual void Init(void* windowHandle) = 0;
+		virtual void Reset() = 0;
+
+		virtual void Resize(int width, int height) = 0;
+		virtual void Paint(const Function<void(PaintContext& context, RECT rect)>& func) = 0;
+
+		virtual PaintContext* CreateContext(int clientWidth, int clientHeight) = 0;
+		virtual void DeleteContext(PaintContext*) = 0;
 
 		virtual HBRUSH CreateSolidBrush(COLORREF color) = 0;
+		virtual HPEN CreatePen(COLORREF color) = 0;
+
+		virtual HFONT CreateFont(const tchar* font, int height, bool createUnderline) = 0;
+		virtual void DeleteFont(HFONT font) = 0;
 	};
 
 	////////////////////////////////////////////////////////////////////////////////////////////////////
diff --git a/Engine/Source/Programs/UnrealBuildAccelerator/Visualizer/Private/UbaVisualizerConfig.cpp b/Engine/Source/Programs/UnrealBuildAccelerator/Visualizer/Private/UbaVisualizerConfig.cpp
--- a/Engine/Source/Programs/UnrealBuildAccelerator/Visualizer/Private/UbaVisualizerConfig.cpp
+++ b/Engine/Source/Programs/UnrealBuildAccelerator/Visualizer/Private/UbaVisualizerConfig.cpp
@@ -45,6 +45,8 @@ namespace uba
 		UBA_VISUALIZER_FLAGS2
 		#undef UBA_VISUALIZER_FLAG
 
+		config.GetValueAsBool(useD2D, TC("UseD2D"));
+
 		fontSize = Min(fontSize, 30u);
 		return true;
 	}
@@ -70,6 +72,9 @@ namespace uba
 		#define UBA_VISUALIZER_FLAG(name, defaultValue, desc) config.AddValue(TC(#name), name);
 		UBA_VISUALIZER_FLAGS2
 		#undef UBA_VISUALIZER_FLAG
+
+		config.AddValue(TC("UseD2D"), useD2D);
+
 		return config.SaveToFile(logger, filename.c_str());
 	}
 }
\ No newline at end of file
diff --git a/Engine/Source/Programs/UnrealBuildAccelerator/Visualizer/Private/UbaVisualizerConfig.h b/Engine/Source/Programs/UnrealBuildAccelerator/Visualizer/Private/UbaVisualizerConfig.h
--- a/Engine/Source/Programs/UnrealBuildAccelerator/Visualizer/Private/UbaVisualizerConfig.h
+++ b/Engine/Source/Programs/UnrealBuildAccelerator/Visualizer/Private/UbaVisualizerConfig.h
@@ -55,6 +55,7 @@ namespace uba
 		u32 maxActiveProcessHeight = 16;
 		u32 boxHeight = 12;
 		float horizontalScaleValue = 0.5f;
+		bool useD2D = false;
 
 		#define UBA_VISUALIZER_FLAG(name, defaultValue, desc) bool show##name = defaultValue;
 		UBA_VISUALIZER_FLAGS1
diff --git a/Engine/Source/Programs/UnrealBuildAccelerator/Visualizer/Private/Windows/UbaVisualizerApp.cpp b/Engine/Source/Programs/UnrealBuildAccelerator/Visualizer/Private/Windows/UbaVisualizerApp.cpp
--- a/Engine/Source/Programs/UnrealBuildAccelerator/Visualizer/Private/Windows/UbaVisualizerApp.cpp
+++ b/Engine/Source/Programs/UnrealBuildAccelerator/Visualizer/Private/Windows/UbaVisualizerApp.cpp
@@ -90,6 +90,8 @@ int WINAPI wWinMain(_In_ HINSTANCE hInstance, _In_opt_ HINSTANCE hPrevInstance,
 	u32 replay = 0;
 	u64 parent = 0;
 	bool copyAndLaunch = true;
+	bool useGDI = false;
+	bool useD2D = false;
 
 	int argc;
 	auto argv = CommandLineToArgvW(GetCommandLine(), &argc);
@@ -188,6 +190,14 @@ int WINAPI wWinMain(_In_ HINSTANCE hInstance, _In_opt_ HINSTANCE hPrevInstance,
 		{
 			copyAndLaunch = false;
 		}
+		else if (name.Equals(TCV("-gdi")))
+		{
+			useGDI = true;
+		}
+		else if (name.Equals(TCV("-d2d")))
+		{
+			useD2D = true;
+		}
 		else if (name.Equals(TCV("-ownerPid")))
 		{
 			u32 ownerPid;
@@ -312,6 +322,10 @@ int WINAPI wWinMain(_In_ HINSTANCE hInstance, _In_opt_ HINSTANCE hPrevInstance,
 	visualizerConfig.parent = parent;
 	visualizerConfig.ShowAllTraces = showAllTraces;
 	visualizerConfig.Load(logger);
+	if (useGDI)
+		visualizerConfig.useD2D = false;
+	if (useD2D)
+		visualizerConfig.useD2D = true;
 
 	NetworkBackendTcp networkBackend(logWriter);
 	VisualizerWin visualizer(visualizerConfig, logger);
diff --git a/Engine/Source/Programs/UnrealBuildAccelerator/Visualizer/Private/Windows/UbaVisualizerBackendD2D.h b/Engine/Source/Programs/UnrealBuildAccelerator/Visualizer/Private/Windows/UbaVisualizerBackendD2D.h
new file mode 100644
--- /dev/null
+++ b/Engine/Source/Programs/UnrealBuildAccelerator/Visualizer/Private/Windows/UbaVisualizerBackendD2D.h
@@ -0,0 +1,406 @@
+// Copyright Epic Games, Inc. All Rights Reserved.
+
+#pragma once
+
+#include "UbaVisualizerBackend.h"
+#include "Windows/UbaVisualizerBackendGDI.h"
+
+#include <d2d1.h>
+#include <dwrite.h>
+
+#pragma comment(lib, "d2d1.lib")
+#pragma comment(lib, "dwrite.lib")
+
+namespace uba
+{
+	class VisualizerWin;
+
+	////////////////////////////////////////////////////////////////////////////////////////////////////
+
+	class VisualizerBackendD2D final : public VisualizerBackend
+	{
+	public:
+		static constexpr int BitmapCacheHeight = 16*1024;
+
+		static inline D2D1::ColorF ToRgbD2D(COLORREF c)
+		{
+			return  ((c & 0x000000FF) << 16) | (c & 0x0000FF00) | ((c & 0x00FF0000) >> 16);
+		}
+
+		struct PaintContextD2D : public PaintContext
+		{
+			PaintContextD2D(VisualizerBackendD2D& o, ID2D1RenderTarget& rt, SIZE cs)
+			:	owner(o)
+			,	renderTarget(rt)
+			,	clientSize(cs)
+			{
+			}
+
+			virtual SIZE GetClientSize() override
+			{
+				return clientSize;
+			}
+			virtual void SetPen(HPEN pen) override
+			{
+				activePen = pen;
+			}
+			virtual void SetBrush(HBRUSH b) override
+			{
+				activeBrush = b;
+			}
+			virtual void SetFont(const Font& font, bool underlined = false) override
+			{
+				activeFont = font;
+				activeFontHandle = (IDWriteTextFormat*)(underlined ? font.handleUnderlined : font.handle);
+			}
+			virtual void SetTextColor(COLORREF color) override
+			{
+				if (activeTextColor != color)
+				{
+					activeTextColor = color;
+					owner.m_textBrush->SetColor(ToRgbD2D(color));
+				}
+			}
+			virtual void Text(int x, int y, StringView str, const RECT* clipRect) override
+			{
+				x += clientOffsetX;
+				y += clientOffsetY;
+
+				int right = x + 10000;
+				int bottom = y + 100;
+				if (clipRect)
+				{
+					//UBA_ASSERT(clipRect->left <= x);
+					//UBA_ASSERT(clipRect->top <= y);
+					right = clipRect->right + clientOffsetX;
+					bottom = clipRect->bottom + clientOffsetY;
+				}
+				if (y > clientSize.cy)
+					return;
+				if (bottom < 0)
+					return;
+
+				D2D1_RECT_F layoutRect = D2D1::RectF(float(x), float(y), float(right), float(bottom));
+
+				if (clipRect)
+					renderTarget.PushAxisAlignedClip(layoutRect, D2D1_ANTIALIAS_MODE_ALIASED);
+
+				if ((HFONT)activeFontHandle == activeFont.handleUnderlined)
+				{
+					IDWriteTextLayout* layout = nullptr;
+					owner.m_dwFactory->CreateTextLayout(str.data, str.count, activeFontHandle, 10000.0f, 1000.0f, &layout);
+					DWRITE_TEXT_METRICS m{};
+					DWRITE_TEXT_RANGE range{ 0, str.count };
+					layout->SetUnderline(TRUE, range);
+					D2D1_POINT_2F origin = D2D1::Point2F(layoutRect.left, layoutRect.top);
+					renderTarget.DrawTextLayout(origin, layout, owner.m_textBrush);
+					layout->Release();
+				}
+				else
+				{
+					renderTarget.DrawText(str.data, str.count, activeFontHandle, &layoutRect, owner.m_textBrush);
+				}
+
+				if (clipRect)
+					renderTarget.PopAxisAlignedClip();
+			}
+			virtual COLORREF GetTextColor() override
+			{
+				return activeTextColor;
+			}
+			virtual Font& GetFont() override
+			{
+				return activeFont;
+			}
+			virtual void FillRect(const RECT& r, HBRUSH b) override
+			{
+				D2D1_RECT_F rect = D2D1::RectF(float(r.left), float(r.top), float(r.right), float(r.bottom));
+				renderTarget.FillRectangle(rect, (ID2D1SolidColorBrush*)b);
+			}
+			virtual void Rectangle(int left, int top, int right, int bottom) override
+			{
+				D2D1_RECT_F rect = D2D1::RectF(float(left)+0.5f, float(top)+0.5f, float(right)-0.5f, float(bottom)-0.5f);
+				float strokeWidth = 1.0f;
+				renderTarget.DrawRectangle(rect, (ID2D1SolidColorBrush*)activePen, strokeWidth);
+			}
+			virtual void Line(int fromX, int fromY, int toX, int toY) override
+			{
+				D2D1_POINT_2F p0 = D2D1::Point2F(float(fromX)+0.5f, float(fromY)+0.5f);
+				D2D1_POINT_2F p1 = D2D1::Point2F(float(toX)+0.5f, float(toY)+0.5f);
+				float strokeWidth = 1.0f;
+				renderTarget.DrawLine(p0, p1, (ID2D1SolidColorBrush*)activePen, strokeWidth, nullptr);
+			}
+			virtual void Polyline(u32 cacheId, const Vector<POINT>& polyLine) override
+			{
+				float strokeWidth = 1.0f;
+				D2D1_POINT_2F prev = {};
+				bool first = true;
+				for (auto& p : polyLine)
+				{
+					D2D1_POINT_2F pf = D2D1::Point2F(float(p.x), float(p.y));
+					if (!first)
+						renderTarget.DrawLine(prev, pf, (ID2D1SolidColorBrush*)activePen, strokeWidth, nullptr);
+					first = false;
+					prev = pf;
+				}
+			}
+			virtual void PolyPolyline(const POINT* points, const u32* lines, u32 lineCount) override
+			{
+				float strokeWidth = 1.0f;
+				const POINT* pointsIt = points;
+				for (u32 i=0; i!=lineCount; ++i)
+				{
+					D2D1_POINT_2F prev = {};
+					bool first = true;
+					for (u32 j=0; j!=lines[i]; ++j)
+					{
+						D2D1_POINT_2F pf = D2D1::Point2F(float(pointsIt->x), float(pointsIt->y));
+						if (!first)
+							renderTarget.DrawLine(prev, pf, (ID2D1SolidColorBrush*)activePen, strokeWidth, nullptr);
+						first = false;
+						prev = pf;
+						++pointsIt;
+					}
+				}
+			}
+			virtual SIZE GetTextSize(StringView str) override
+			{
+				IDWriteTextLayout* layout = nullptr;
+				HRESULT hr = owner.m_dwFactory->CreateTextLayout(str.data, str.count, activeFontHandle, 10000.0f, 1000.0f, &layout);
+				if (FAILED(hr) || !layout)
+					return { 0, 0 };
+				DWRITE_TEXT_METRICS m{};
+				layout->GetMetrics(&m);
+				layout->Release();
+				return { int(m.width), int(m.height) }; // m.widthIncludingTrailingWhitespace
+			}
+			virtual void DrawCached(TraceView::BitmapCache& process, RECT rect, int offsetY, int textHeight, Font& processFont, RECT clippingRect, const Function<void(PaintContext& textCacheContext)>& func) override
+			{
+				if (!process.bitmap || process.bitmapDirty)
+				{
+					if (!process.bitmap)
+					{
+						if (owner.m_lastBitmapOffset + processFont.height > BitmapCacheHeight)
+						{
+							if (owner.m_lastBitmap)
+								owner.m_textBitmaps.push_back(owner.m_lastBitmap);
+
+							ID2D1BitmapRenderTarget* textRT = nullptr;
+							renderTarget.CreateCompatibleRenderTarget(D2D1::SizeF(float(256), float(BitmapCacheHeight)), &textRT);
+							UBA_ASSERT(textRT);
+
+							owner.m_lastBitmapOffset = 0;
+							owner.m_lastBitmap = textRT;
+						}
+						process.bitmap = (HBITMAP)owner.m_lastBitmap;
+						process.bitmapOffset = owner.m_lastBitmapOffset;
+						owner.m_lastBitmapOffset += processFont.height;
+					}
+					if (lastSelectedBitmap != (ID2D1BitmapRenderTarget*)process.bitmap)
+						lastSelectedBitmap = (ID2D1BitmapRenderTarget*)process.bitmap;
+
+					lastSelectedBitmap->BeginDraw();
+					PaintContextD2D textCacheContext{ owner, *lastSelectedBitmap, { 256, BitmapCacheHeight } };
+					textCacheContext.SetFont(activeFont);
+					func(textCacheContext);
+					lastSelectedBitmap->EndDraw();
+				}
+
+				if (lastSelectedBitmap != (ID2D1BitmapRenderTarget*)process.bitmap)
+					lastSelectedBitmap = (ID2D1BitmapRenderTarget*)process.bitmap;
+
+				int processWidth = rect.right - rect.left;
+				int width = Min(processWidth, 256);
+				int bitmapOffsetY = process.bitmapOffset;
+				int bltOffsetY = offsetY;
+				if (bltOffsetY < 0)
+				{
+					bitmapOffsetY -= bltOffsetY;
+					bltOffsetY = 0;
+				}
+				int height = Min(textHeight, processFont.height);
+				if (bltOffsetY + height > textHeight)
+					height = textHeight - bltOffsetY;
+
+				int left = rect.left;
+				if (left <= -256 || height <= 0)
+					return;
+
+				int bitmapOffsetX = rect.left - left;
+
+				if (left < clippingRect.left)
+				{
+					int diff = clippingRect.left - left;
+					rect.left = clippingRect.left;
+					width -= diff;
+					bitmapOffsetX += diff;
+				}
+
+				D2D1_RECT_F destRect = D2D1::RectF(float(rect.left), float(rect.top + bltOffsetY), float(rect.left + width), float(rect.top + bltOffsetY + height));
+				D2D1_RECT_F sourceRect = D2D1::RectF(float(bitmapOffsetX), float(bitmapOffsetY), float(bitmapOffsetX + width), float(bitmapOffsetY + height));
+
+				ID2D1Bitmap* textBitmap = nullptr;
+				lastSelectedBitmap->GetBitmap(&textBitmap);
+				renderTarget.DrawBitmap(textBitmap, &destRect, 1.0f, D2D1_BITMAP_INTERPOLATION_MODE_NEAREST_NEIGHBOR, &sourceRect);
+				textBitmap->Release();
+			}
+			virtual POINT GetCursorClientPos() override
+			{
+				POINT p;
+				GetCursorPos(&p);
+				ScreenToClient(owner.m_owner.m_hwnd, &p);
+				return p;
+			}
+
+			int clientOffsetX = 0;
+			int clientOffsetY = 0;
+
+			VisualizerBackendD2D& owner;
+			ID2D1RenderTarget& renderTarget;
+			SIZE clientSize;
+			HPEN activePen = 0;
+			HBRUSH activeBrush = 0;
+			Font activeFont;
+			IDWriteTextFormat* activeFontHandle = nullptr;
+			COLORREF activeTextColor = 0;
+			ID2D1BitmapRenderTarget* lastSelectedBitmap = 0;
+		};
+
+		VisualizerBackendD2D(VisualizerWin& o) : m_owner(o)
+		{
+		}
+
+		virtual const tchar* GetName() override { return TC("D2D"); }
+
+		bool Create()
+		{
+			D2D1_FACTORY_OPTIONS options{};
+			// options.debugLevel = D2D1_DEBUG_LEVEL_INFORMATION;
+			if (D2D1CreateFactory(D2D1_FACTORY_TYPE_SINGLE_THREADED, options, &m_d2dFactory) != S_OK)
+				return false;
+			if (DWriteCreateFactory(DWRITE_FACTORY_TYPE_SHARED, __uuidof(IDWriteFactory), reinterpret_cast<IUnknown**>(&m_dwFactory)) != S_OK)
+				return false;
+			return true;
+		}
+
+		virtual void Init(void* windowHandle) override
+		{
+		}
+
+		virtual void Reset() override
+		{
+			m_dirty = true;
+			for (auto b : m_brushes)
+				b->Release();
+			m_brushes.clear();
+			for (auto bm : m_textBitmaps)
+				bm->Release();
+			m_textBitmaps.clear();
+
+		}
+
+		virtual void Resize(int width, int height) override
+		{
+			if (m_d2dRenderTarget)
+				m_d2dRenderTarget->Resize(D2D1::SizeU(width, height));
+		}
+
+		virtual void Paint(const Function<void(PaintContext& context, RECT rect)>& func) override
+		{
+			if (!m_d2dRenderTarget)
+			{
+				D2D1_SIZE_U size = D2D1::SizeU(m_owner.m_clientWidth, m_owner.m_clientHeight);
+				auto prop = D2D1::RenderTargetProperties();
+				auto hwndProp = D2D1::HwndRenderTargetProperties((HWND)m_owner.m_hwnd, size);
+				if (m_d2dFactory->CreateHwndRenderTarget(prop, hwndProp, &m_d2dRenderTarget) != S_OK)
+					return;
+				//m_d2dRenderTarget->SetAntialiasMode(D2D1_ANTIALIAS_MODE_ALIASED);
+			}
+
+			if (m_dirty)
+			{
+				m_owner.InitBrushes();
+				m_d2dRenderTarget->CreateSolidColorBrush(ToRgbD2D(m_owner.m_textColor), &m_textBrush);
+			}
+			PAINTSTRUCT ps;
+			BeginPaint(m_owner.m_hwnd, &ps);
+
+			m_d2dRenderTarget->BeginDraw();
+
+			COLORREF bgColor = m_owner.m_config.DarkMode ? RGB(37, 37, 38) : RGB(180, 180, 180);
+			m_d2dRenderTarget->Clear(ToRgbD2D(bgColor));
+
+			int clientWidth = m_owner.m_clientWidth;
+			int clientHeight = m_owner.m_clientHeight;
+			PaintContextD2D context(*this, *m_d2dRenderTarget, { clientWidth, clientHeight });
+			func(context, { 0, 0, clientWidth, clientHeight });
+
+			HRESULT hr = m_d2dRenderTarget->EndDraw();
+			if (hr == D2DERR_RECREATE_TARGET)
+			{
+				Reset();
+				m_d2dRenderTarget->Release();
+				m_d2dRenderTarget = nullptr;
+			}
+
+			EndPaint(m_owner.m_hwnd, &ps);
+		}
+
+		virtual PaintContext* CreateContext(int clientWidth, int clientHeight) override
+		{
+			return new PaintContextD2D(*this, *m_d2dRenderTarget, { m_owner.m_clientWidth, m_owner.m_clientHeight });
+		}
+
+		virtual void DeleteContext(PaintContext* context) override
+		{
+			delete (PaintContextD2D*)context;
+		}
+
+		virtual HBRUSH CreateSolidBrush(COLORREF color) override
+		{
+			ID2D1SolidColorBrush* brush;
+			m_d2dRenderTarget->CreateSolidColorBrush(ToRgbD2D(color), &brush);
+			m_brushes.push_back(brush);
+			return (HBRUSH)brush;
+		}
+
+		virtual HPEN CreatePen(COLORREF color) override
+		{
+			return (HPEN)CreateSolidBrush(color);
+		}
+
+		virtual HFONT CreateFont(const tchar* font, int height, bool createUnderline) override
+		{
+			auto weight = DWRITE_FONT_WEIGHT_NORMAL;
+			auto style = DWRITE_FONT_STYLE_NORMAL;
+			auto stretch = DWRITE_FONT_STRETCH_NORMAL;
+			IDWriteTextFormat* textFormat = 0;
+			if (m_dwFactory->CreateTextFormat(font, nullptr, weight, style, stretch, float(height - 4), L"en-us", &textFormat) != S_OK)
+				return 0;
+			textFormat->SetTextAlignment(DWRITE_TEXT_ALIGNMENT_LEADING);
+			textFormat->SetParagraphAlignment(DWRITE_PARAGRAPH_ALIGNMENT_NEAR);
+			return (HFONT)textFormat;
+		}
+
+		virtual void DeleteFont(HFONT font)
+		{
+			((IDWriteTextFormat*)font)->Release();
+		}
+
+		VisualizerWin& m_owner;
+		ID2D1Factory* m_d2dFactory = nullptr;
+		ID2D1HwndRenderTarget* m_d2dRenderTarget = nullptr;
+		List<ID2D1SolidColorBrush*> m_brushes;
+
+		IDWriteFactory* m_dwFactory = nullptr;
+		ID2D1SolidColorBrush* m_textBrush = nullptr;
+
+		ID2D1BitmapRenderTarget* m_lastBitmap = 0;
+		int m_lastBitmapOffset = BitmapCacheHeight;
+		Vector<ID2D1BitmapRenderTarget*> m_textBitmaps;
+
+		bool m_dirty = true;
+	};
+
+	////////////////////////////////////////////////////////////////////////////////////////////////////
+}
diff --git a/Engine/Source/Programs/UnrealBuildAccelerator/Visualizer/Private/Windows/UbaVisualizerBackendGDI.h b/Engine/Source/Programs/UnrealBuildAccelerator/Visualizer/Private/Windows/UbaVisualizerBackendGDI.h
--- a/Engine/Source/Programs/UnrealBuildAccelerator/Visualizer/Private/Windows/UbaVisualizerBackendGDI.h
+++ b/Engine/Source/Programs/UnrealBuildAccelerator/Visualizer/Private/Windows/UbaVisualizerBackendGDI.h
@@ -18,26 +18,27 @@ namespace uba
 		struct PaintContextGDI : public PaintContext
 		{
 			virtual SIZE GetClientSize() override { return clientSize; }
-			virtual void SetBkMode(int mode) override { ::SetBkMode(hdc, mode); }
 			virtual void SetTextColor(COLORREF color) override { ::SetTextColor(hdc, color); }
 			virtual COLORREF GetTextColor() override { return ::GetTextColor(hdc); }
 			virtual void SetPen(HPEN pen) override { ::SelectObject(hdc, pen); }
-			virtual u32 SetTextAlign(u32 align) override { return ::SetTextAlign(hdc, align); }
-			virtual void Text(int x, int y, const RECT* clipRect, StringView str) override { ::ExtTextOutW(hdc, x, y, clipRect ? ETO_CLIPPED : 0, clipRect, str.data, str.count, NULL); }
+			virtual void Text(int x, int y, StringView str, const RECT* clipRect) override
+			{
+				::ExtTextOutW(hdc, x, y, clipRect ? ETO_CLIPPED : 0, clipRect, str.data, str.count, NULL);
+			}
 
 			virtual void SetFont(const Font& font, bool underlined) override
 			{
 				if (underlined)
 				{
-					if (font.handleUnderlined == activeFont.handleUnderlined)
-						return;
-					::SelectObject(hdc, font.handleUnderlined);
+					if (font.handleUnderlined != activeFontHandle)
+						::SelectObject(hdc, font.handleUnderlined);
+					activeFontHandle = font.handleUnderlined;
 				}
 				else
 				{
-					if (font.handle == activeFont.handle)
-						return;
-					::SelectObject(hdc, font.handle);
+					if (font.handle != activeFontHandle)
+						::SelectObject(hdc, font.handle);
+					activeFontHandle = font.handle;
 				}
 				activeFont = font;
 			}
@@ -74,7 +75,7 @@ namespace uba
 				::LineTo(hdc, toX, toY);
 			}
 
-			virtual void Polyline(const Vector<POINT>& polyLine) override
+			virtual void Polyline(u32 cacheId, const Vector<POINT>& polyLine) override
 			{
 				::Polyline(hdc, polyLine.data(), int(polyLine.size()));
 			}
@@ -84,29 +85,6 @@ namespace uba
 				::PolyPolyline(hdc, points, (const DWORD*)lines, int(lineCount));
 			}
 
-			virtual void DrawCenteredText(const StringBuffer<>* rows, u32 rowCount, HBRUSH backgroundBrush) override
-			{
-				for (u32 i=0; i!=rowCount; ++i)
-				{ 
-					SIZE textSize = GetTextSize(rows[i]);
-
-					RECT rect = { 0, 0, 0, 0 };
-					int top = (clientSize.cy - textSize.cy)/2 + i*20;
-					rect.left = (clientSize.cx - textSize.cx)/2;
-					rect.right = (clientSize.cx + textSize.cx)/2;
-					rect.bottom = top + textSize.cy;
-					rect.top = top;
-
-					if (backgroundBrush)
-					{
-						RECT r = rect;
-						InflateRect(&r, 4, 2);
-						FillRect(r, backgroundBrush);
-					}
-					Text(rect.left, rect.top, nullptr, rows[i]);
-				}
-			}
-
 			virtual SIZE GetTextSize(StringView str) override
 			{
 				SIZE s;
@@ -193,13 +171,25 @@ namespace uba
 			SIZE clientSize;
 
 			Font activeFont;
+			HFONT activeFontHandle;
 
 			HDC textDC = NULL;
 			HBRUSH lastSelectedBrush = 0;
 			HBITMAP lastSelectedBitmap = 0;
 		};
 
-		virtual void ResetCaches() override
+		VisualizerBackendGDI(VisualizerWin& owner) : m_owner(owner)
+		{
+		}
+
+		virtual const tchar* GetName() override { return TC("GDI"); }
+
+		virtual void Init(void* windowHandle) override
+		{
+			m_owner.InitBrushes();
+		}
+
+		virtual void Reset() override
 		{
 			for (HBITMAP bm : m_textBitmaps)
 				DeleteObject(bm);
@@ -207,6 +197,70 @@ namespace uba
 			DeleteObject(m_lastBitmap);
 			m_lastBitmap = 0;
 			m_lastBitmapOffset = BitmapCacheHeight;
+			m_owner.InitBrushes();
+		}
+
+		virtual void Resize(int width, int height) override
+		{
+		}
+
+		virtual void Paint(const Function<void(PaintContext& context, RECT rect)>& func) override
+		{
+			HWND hwnd = m_owner.m_hwnd;
+			int clientWidth = m_owner.m_clientWidth;
+			int clientHeight = m_owner.m_clientHeight;
+
+			HDC hdc = GetDC(hwnd);
+
+			HDC memDC = CreateCompatibleDC(hdc);
+
+			if (m_cachedBitmapSize.cx != clientWidth || m_cachedBitmapSize.cy != clientHeight)
+			{
+				if (m_cachedBitmap)
+					DeleteObject(m_cachedBitmap);
+				m_cachedBitmap = CreateCompatibleBitmap(hdc, clientWidth, clientHeight);
+				m_cachedBitmapSize = { clientWidth, clientHeight };
+			}
+			HGDIOBJ oldBmp = SelectObject(memDC, m_cachedBitmap);
+
+			RECT rect = { 0, 0, clientWidth, clientHeight };
+			FillRect(memDC, &rect, m_owner.m_backgroundBrush);
+			SetBkColor(memDC, m_owner.m_config.DarkMode ? RGB(70, 70, 70) : RGB(180, 180, 180));
+			SetBkMode(memDC, TRANSPARENT);
+
+			HDC textDC = CreateCompatibleDC(memDC);
+			SetTextColor(textDC, m_owner.m_textColor);
+			SelectObject(textDC, m_owner.m_processFont.handle);
+			SelectObject(textDC, GetStockObject(NULL_BRUSH));
+			SetBkMode(textDC, TRANSPARENT);
+
+			HBITMAP nullBmp = CreateCompatibleBitmap(memDC, 1, 1);
+			HBITMAP oldBmp2 = (HBITMAP)SelectObject(textDC, nullBmp);
+
+			VisualizerBackendGDI::PaintContextGDI context(*this, hwnd, memDC, { clientWidth, clientHeight }, textDC);
+			
+			func(context, rect);
+
+			SelectObject(textDC, oldBmp2);
+			DeleteObject(nullBmp);
+			DeleteDC(textDC);
+
+			BitBlt(hdc, 0, 0, rect.right - rect.left, rect.bottom - rect.top, memDC, 0, 0, SRCCOPY);
+
+			SelectObject(memDC, oldBmp);
+			DeleteDC(memDC);
+
+			ReleaseDC(hwnd, hdc);
+		}
+
+		virtual PaintContext* CreateContext(int clientWidth, int clientHeight) override
+		{
+			return new PaintContextGDI(*this, NULL, NULL, { clientWidth, clientHeight });
+		}
+
+		virtual void DeleteContext(PaintContext* context) override
+		{
+			delete (PaintContextGDI*)context;
 		}
 
 		virtual HBRUSH CreateSolidBrush(COLORREF color) override
@@ -214,6 +268,27 @@ namespace uba
 			return ::CreateSolidBrush(color);
 		}
 
+		virtual HPEN CreatePen(COLORREF color) override
+		{
+			return ::CreatePen(PS_SOLID, 1, color);
+		}
+
+		HFONT CreateFont(const tchar* font, int height, bool createUnderline) override
+		{
+			DWORD underline = createUnderline ? 1 : 0;
+			return ::CreateFontW(4 - height, 0, 0, 0, FW_NORMAL, 0, underline, 0, ANSI_CHARSET, OUT_DEFAULT_PRECIS, CLIP_DEFAULT_PRECIS, CLEARTYPE_QUALITY, DEFAULT_PITCH, font);
+		}
+
+		void DeleteFont(HFONT font)
+		{
+			::DeleteObject(font);
+		}
+
+		VisualizerWin& m_owner;
+
+		HBITMAP m_cachedBitmap = 0;
+		SIZE m_cachedBitmapSize = { INT_MIN, INT_MIN };
+
 		HBITMAP m_lastBitmap = 0;
 		int m_lastBitmapOffset = BitmapCacheHeight;
 		Vector<HBITMAP> m_textBitmaps;
diff --git a/Engine/Source/Programs/UnrealBuildAccelerator/Visualizer/Private/Windows/UbaVisualizerWin.cpp b/Engine/Source/Programs/UnrealBuildAccelerator/Visualizer/Private/Windows/UbaVisualizerWin.cpp
--- a/Engine/Source/Programs/UnrealBuildAccelerator/Visualizer/Private/Windows/UbaVisualizerWin.cpp
+++ b/Engine/Source/Programs/UnrealBuildAccelerator/Visualizer/Private/Windows/UbaVisualizerWin.cpp
@@ -2,7 +2,7 @@
 
 #include "UbaVisualizerWin.h"
 #include "UbaConfig.h"
-#include "UbaVisualizerBackendGDI.h"
+#include "UbaVisualizerBackendD2D.h"
 
 #include <algorithm>
 #include <uxtheme.h>
@@ -56,8 +56,20 @@ namespace uba
 		TString& m_out;
 	};
 
+	VisualizerBackend& CreateBackend(VisualizerConfig& config, VisualizerWin& v)
+	{
+		if (config.useD2D)
+		{
+			auto d2d = new VisualizerBackendD2D(v);
+			if (d2d->Create())
+				return *d2d;
+			delete d2d;
+		}
+		return *new VisualizerBackendGDI(v);
+	}
+
 	VisualizerWin::VisualizerWin(VisualizerConfig& config, Logger& logger)
-	:	Visualizer(config, logger, *new VisualizerBackendGDI())
+	:	Visualizer(config, logger, CreateBackend(config, *this))
 	{
 		memset(m_activeProcessFont, 0, sizeof(m_activeProcessFont));
 		memset(m_activeProcessCountHistory, 0, sizeof(m_activeProcessCountHistory));
@@ -253,102 +265,6 @@ namespace uba
 		m_locked = lock;
 	}
 
-	void VisualizerWin::InitBrushes()
-	{
-		if (m_config.DarkMode)
-		{
-			m_textColor = RGB(190, 190, 190);
-			m_textWarningColor = RGB(190, 190, 0);
-			m_textErrorColor = RGB(190, 0, 0);
-
-			m_processBrushes[0].inProgress = CreateSolidBrush(RGB(70, 70, 70));
-			m_processBrushes[1].inProgress = CreateSolidBrush(RGB(130, 130, 130));
-			m_processBrushes[0].inProgressRace = CreateSolidBrush(RGB(50, 50, 50));
-			m_processBrushes[1].inProgressRace = CreateSolidBrush(RGB(120, 120, 120));
-
-			m_processBrushes[0].error = CreateSolidBrush(RGB(140, 0, 0));
-			m_processBrushes[1].error = CreateSolidBrush(RGB(190, 0, 0));
-
-			m_processBrushes[0].returned = CreateSolidBrush(RGB(50, 50, 120));
-			m_processBrushes[1].returned = CreateSolidBrush(RGB(70, 70, 160));
-
-			m_processBrushes[0].recv = CreateSolidBrush(RGB(10, 92, 10));
-			m_processBrushes[1].recv = CreateSolidBrush(RGB(10, 130, 10));
-			m_processBrushes[0].success = CreateSolidBrush(RGB(10, 100, 10));
-			m_processBrushes[1].success = CreateSolidBrush(RGB(10, 140, 10));
-			m_processBrushes[0].send = CreateSolidBrush(RGB(10, 115, 10));
-			m_processBrushes[1].send = CreateSolidBrush(RGB(10, 145, 10));
-			m_processBrushes[0].cacheFetchTest = CreateSolidBrush(RGB(24, 120, 110));
-			m_processBrushes[1].cacheFetchTest = CreateSolidBrush(RGB(31, 150, 138));
-			m_processBrushes[0].cacheFetchDownload = CreateSolidBrush(RGB(24, 112, 110));
-			m_processBrushes[1].cacheFetchDownload = CreateSolidBrush(RGB(31, 143, 138));
-			m_processBrushes[0].warning = CreateSolidBrush(RGB(100, 100, 0));
-			m_processBrushes[1].warning = CreateSolidBrush(RGB(100, 100, 0));
-
-			m_backgroundBrush = CreateSolidBrush(0x00252526);
-			m_separatorPen = CreatePen(PS_SOLID, 1, RGB(50, 50, 50));
-			m_tooltipBackgroundBrush = CreateSolidBrush(0x00404040);
-			m_checkboxPen = CreatePen(PS_SOLID, 1, RGB(130, 130, 130));
-
-			m_sendColor = RGB(0, 170, 0);
-			m_recvColor = RGB(0, 170, 255);
-			m_cpuColor = RGB(170, 170, 0);
-			m_memColor = RGB(170, 0, 255);
-			m_driveColor = RGB(170, 65, 55);
-			m_activeProcColor = RGB(0, 170, 170);
-		}
-		else
-		{
-			m_textColor = GetSysColor(COLOR_INFOTEXT);
-			m_textWarningColor = RGB(170, 130, 0);
-			m_textErrorColor = RGB(190, 0, 0);
-
-			m_processBrushes[0].inProgress = CreateSolidBrush(RGB(150, 150, 150));
-			m_processBrushes[1].inProgress = CreateSolidBrush(RGB(180, 180, 180));
-			m_processBrushes[0].inProgressRace = CreateSolidBrush(RGB(130, 130, 130));
-			m_processBrushes[1].inProgressRace = CreateSolidBrush(RGB(170, 170, 170));
-
-			m_processBrushes[0].error = CreateSolidBrush(RGB(255, 70, 70));
-			m_processBrushes[1].error = CreateSolidBrush(RGB(255, 100, 70));
-
-			m_processBrushes[0].returned = CreateSolidBrush(RGB(150, 150, 200));
-			m_processBrushes[1].returned = CreateSolidBrush(RGB(170, 170, 200));
-
-			m_processBrushes[0].recv = CreateSolidBrush(RGB(10, 190, 10));
-			m_processBrushes[1].recv = CreateSolidBrush(RGB(20, 210, 20));
-			m_processBrushes[0].success = CreateSolidBrush(RGB(10, 200, 10));
-			m_processBrushes[1].success = CreateSolidBrush(RGB(20, 220, 20));
-			m_processBrushes[0].send = CreateSolidBrush(RGB(80, 210, 80));
-			m_processBrushes[1].send = CreateSolidBrush(RGB(90, 250, 90));
-			m_processBrushes[0].cacheFetchTest = CreateSolidBrush(RGB(150, 150, 200));
-			m_processBrushes[1].cacheFetchTest = CreateSolidBrush(RGB(170, 170, 200));
-			m_processBrushes[0].cacheFetchDownload = CreateSolidBrush(RGB(150, 150, 200));
-			m_processBrushes[1].cacheFetchDownload = CreateSolidBrush(RGB(170, 170, 200));
-			m_processBrushes[0].warning = CreateSolidBrush(RGB(170, 170, 0));
-			m_processBrushes[1].warning = CreateSolidBrush(RGB(170, 170, 0));
-
-			m_backgroundBrush = GetSysColorBrush(0);
-			m_separatorPen = CreatePen(PS_SOLID, 1, RGB(180, 180, 180));
-			m_tooltipBackgroundBrush = GetSysColorBrush(COLOR_INFOBK);
-			m_checkboxPen = CreatePen(PS_SOLID, 1, RGB(130, 130, 130));
-
-			m_sendColor = RGB(0, 170, 0); // Green
-			m_recvColor = RGB(63, 72, 204); // Blue
-			m_cpuColor = RGB(200, 130, 0); // Orange
-			m_memColor = RGB(170, 0, 255); // Purple
-			m_driveColor = RGB(255, 115, 96);
-			m_activeProcColor = RGB(0, 170, 170);
-		}
-
-		m_textPen = CreatePen(PS_SOLID, 1, m_textColor);
-		m_sendPen = CreatePen(PS_SOLID, 1, m_sendColor);
-		m_recvPen = CreatePen(PS_SOLID, 1, m_recvColor);
-		m_cpuPen = CreatePen(PS_SOLID, 1, m_cpuColor);
-		m_memPen = CreatePen(PS_SOLID, 1, m_memColor);
-		m_drivePen = CreatePen(PS_SOLID, 1, m_driveColor);
-		m_activeProcPen = CreatePen(PS_SOLID, 1, m_activeProcColor);
-	}
-
 	void VisualizerWin::ThreadLoop()
 	{
 		if (m_config.parent)
@@ -356,12 +272,6 @@ namespace uba
 			SetProcessDpiAwareness(PROCESS_SYSTEM_DPI_AWARE);
 		}
 
-		InitBrushes();
-
-		LOGBRUSH br = { 0 };
-		GetObject(m_backgroundBrush, sizeof(br), &br);
-		m_processUpdatePen = CreatePen(PS_SOLID, 2, RGB(GetRValue(br.lbColor), GetGValue(br.lbColor), GetBValue(br.lbColor)));
-
 		HINSTANCE hInstance = GetModuleHandle(NULL);
 		int winPosX = m_config.x;
 		int winPosY = m_config.y;
@@ -401,8 +311,10 @@ namespace uba
 
 		UpdateProcessFont();
 
-		const TCHAR* fontName = TEXT("Consolas");
-		m_popupFont.handle = CreateFontW(-12, 0, 0, 0, FW_NORMAL, false, false, false, ANSI_CHARSET, OUT_DEFAULT_PRECIS, CLIP_DEFAULT_PRECIS, CLEARTYPE_QUALITY, FIXED_PITCH | FF_MODERN, fontName);
+		//const TCHAR* fontName = TEXT("Consolas");
+		//m_popupFont.handle = CreateFontW(-12, 0, 0, 0, FW_NORMAL, false, false, false, ANSI_CHARSET, OUT_DEFAULT_PRECIS, CLIP_DEFAULT_PRECIS, CLEARTYPE_QUALITY, FIXED_PITCH | FF_MODERN, fontName);
+
+		m_popupFont.handle = m_backend.CreateFont(TEXT("Consolas"), 16, false);
 		m_popupFont.height = 14;
 		//m_popupFont = (HFONT)GetStockObject(SYSTEM_FIXED_FONT);
 
@@ -424,7 +336,6 @@ namespace uba
 		GetTitlePrefix(title).Append(L"Initializing...");
 
 		HWND hwnd = CreateWindowEx(exStyle, windowClassName, title.data, windowStyle, winPosX, winPosY, winWidth, winHeight, NULL, NULL, hInstance, this);
-		SetWindowLongPtr(hwnd, GWLP_USERDATA, (LONG_PTR)this);
 		m_hwnd = hwnd;
 
 		auto destroyWin = [&]()
@@ -450,10 +361,17 @@ namespace uba
 		m_clientWidth = r.right;
 		m_clientHeight = r.bottom;
 
+		m_backend.Init(m_hwnd);
+
+		LOGBRUSH br = { 0 };
+		GetObject(m_backgroundBrush, sizeof(br), &br);
+		m_processUpdatePen = CreatePen(PS_SOLID, 2, RGB(GetRValue(br.lbColor), GetGValue(br.lbColor), GetBValue(br.lbColor)));
+
 		{
-			VisualizerBackendGDI::PaintContextGDI context((VisualizerBackendGDI&)m_backend, NULL, NULL, { m_clientWidth, m_clientHeight });
+			PaintContext* context = m_backend.CreateContext(m_clientWidth, m_clientHeight);
 			HitTestResult res;
-			HitTest(res, context, { -1, -1 });
+			HitTest(res, *context, { -1, -1 });
+			m_backend.DeleteContext(context);
 		}
 
 		if (m_config.parent)
@@ -601,7 +519,7 @@ namespace uba
 		if (!full)
 			return;
 
-		m_backend.ResetCaches();
+		m_backend.Reset();
 	}
 
 	void Visualizer::UpdateFont(Font& font, int height, bool createUnderline)
@@ -622,23 +540,23 @@ namespace uba
 		if (height <= 8)
 			++fh;
 		if (height <= 6)
+		{
 			++fh;
+			--font.offset;
+		}
 		if (height <= 4)
 			--font.offset;
 
 		if (font.handle)
-			DeleteObject(font.handle);
+			m_backend.DeleteFont(font.handle);
+		font.handle = 0;
 		if (font.handleUnderlined)
-			DeleteObject(font.handleUnderlined);
-
-		//NONCLIENTMETRICS nonClientMetrics;
-		//nonClientMetrics.cbSize = sizeof(nonClientMetrics);
-		//SystemParametersInfo(SPI_GETNONCLIENTMETRICS, sizeof(nonClientMetrics), &nonClientMetrics, 0);
-		//m_font = (HFONT)CreateFontIndirect(&nonClientMetrics.lfMessageFont);
+			m_backend.DeleteFont(font.handleUnderlined);
+		font.handleUnderlined = 0;
 
-		font.handle = CreateFontW(4 - fh, 0, 0, 0, FW_NORMAL, 0, 0, 0, ANSI_CHARSET, OUT_DEFAULT_PRECIS, CLIP_DEFAULT_PRECIS, CLEARTYPE_QUALITY, DEFAULT_PITCH, m_config.fontName.c_str());
+		font.handle = m_backend.CreateFont(m_config.fontName.c_str(), fh, false);
 		if (createUnderline)
-			font.handleUnderlined = CreateFontW(4 - fh, 0, 0, 0, FW_NORMAL, 0, 1, 0, ANSI_CHARSET, OUT_DEFAULT_PRECIS, CLIP_DEFAULT_PRECIS, CLEARTYPE_QUALITY, DEFAULT_PITCH, m_config.fontName.c_str());
+			font.handleUnderlined = m_backend.CreateFont(m_config.fontName.c_str(), fh, true);
 	}
 
 	void VisualizerWin::UpdateDefaultFont()
@@ -683,57 +601,10 @@ namespace uba
 		KillTimer(m_hwnd, 0);
 	}
 
-	void VisualizerWin::PaintClient(const Function<void(HDC hdc, HDC memDC)>& paintFunc)
-	{
-		HDC hdc = GetDC(m_hwnd);
-
-		HDC memDC = CreateCompatibleDC(hdc);
-
-		if (m_cachedBitmapSize.cx != m_clientWidth || m_cachedBitmapSize.cy != m_clientHeight)
-		{
-			if (m_cachedBitmap)
-				DeleteObject(m_cachedBitmap);
-			m_cachedBitmap = CreateCompatibleBitmap(hdc, m_clientWidth, m_clientHeight);
-			m_cachedBitmapSize = { m_clientWidth, m_clientHeight };
-		}
-		HGDIOBJ oldBmp = SelectObject(memDC, m_cachedBitmap);
-
-		if (!m_isInPaint)
-		{
-			m_isInPaint = true;
-			paintFunc(hdc, memDC);
-			m_isInPaint = false;
-		}
-
-		SelectObject(memDC, oldBmp);
-		DeleteDC(memDC);
-
-		ReleaseDC(m_hwnd, hdc);
-	}
-
 	void VisualizerWin::PaintEmptyScreen(PaintContext& context)
 	{
-		if (m_parentHwnd)
-			return;
-
-		StringBuffer<> str[2];
-		if (!m_fileName.IsEmpty())
-		{
-			str[0].Append(TCV("Loading trace file"));
-		}
-		else
-		{
-			str[1].Append(TCV("Click here to open trace file"));
-			if (m_listenChannel.count)
-			{
-				str[0].Appendf(TC("Listening for new sessions on channel '%s'"), m_listenChannel.data);
-				str[1].Append(TCV(" instead"));
-			}
-			else
-				str[0].Append(TCV("No trace active"));
-		}
-		context.SetFont(m_popupFont);
-		context.DrawCenteredText(str, 2, nullptr);
+		if (!m_parentHwnd)
+			Visualizer::PaintEmptyScreen(context);
 	}
 
 	void VisualizerWin::FinishIncompatible()
@@ -775,8 +646,9 @@ namespace uba
 
 		HitTestResult res;
 		{
-			VisualizerBackendGDI::PaintContextGDI context((VisualizerBackendGDI&)m_backend, NULL, NULL, { m_clientWidth, m_clientHeight });
-			HitTest(res, context, pos);
+			PaintContext* context = m_backend.CreateContext(m_clientWidth, m_clientHeight);
+			HitTest(res, *context, pos);
+			m_backend.DeleteContext(context);
 		}
 		m_activeSection = res.section;
 
@@ -933,29 +805,16 @@ namespace uba
 		case WM_PAINT:
 		{
 			u64 start = GetTime();
-			PaintClient([&](HDC hdc, HDC memDC)
+			m_backend.Paint([&](PaintContext& context, RECT rect)
 				{
-					RECT rect = { 0, 0, m_clientWidth, m_clientHeight };
-					FillRect(memDC, &rect, m_backgroundBrush);
-					SetBkColor(memDC, m_config.DarkMode ? RGB(70, 70, 70) : RGB(180, 180, 180));
-
-					HDC textDC = CreateCompatibleDC(memDC);
-					SetTextColor(textDC, m_textColor);
-					SelectObject(textDC, m_processFont.handle);
-					SelectObject(textDC, GetStockObject(NULL_BRUSH));
-					SetBkMode(textDC, TRANSPARENT);
-
-					HBITMAP nullBmp = CreateCompatibleBitmap(memDC, 1, 1);
-					HBITMAP oldBmp = (HBITMAP)SelectObject(textDC, nullBmp);
-
-					VisualizerBackendGDI::PaintContextGDI context((VisualizerBackendGDI&)m_backend, m_hwnd, memDC, { m_clientWidth, m_clientHeight }, textDC);
 					PaintAll(context, rect);
-
-					SelectObject(textDC, oldBmp);
-					DeleteObject(nullBmp);
-					DeleteDC(textDC);
-
-					BitBlt(hdc, 0, 0, rect.right - rect.left, rect.bottom - rect.top, memDC, 0, 0, SRCCOPY);
+#if UBA_DEBUG
+					StringBuffer<32> buf;
+					buf.Appendf(TC("%llums"), m_lastPaintTimeMs);
+					context.SetFont(m_defaultFont);
+					context.SetTextColor(m_textColor);
+					context.Text(m_clientWidth - 50, m_clientHeight - 50, buf, nullptr);
+#endif
 				});
 			m_lastPaintTimeMs = TimeToMs(GetTime() - start);
 			break;
@@ -972,6 +831,7 @@ namespace uba
 			GetClientRect(m_hwnd, &r);
 			m_clientWidth = r.right;
 			m_clientHeight = r.bottom;
+			m_backend.Resize(m_clientWidth, m_clientHeight);
 			UpdateScrollbars(true);
 			break;
 		}
@@ -1245,8 +1105,9 @@ namespace uba
 					}
 
 					HitTestResult res;
-					VisualizerBackendGDI::PaintContextGDI context((VisualizerBackendGDI&)m_backend, NULL, NULL, { m_clientWidth, m_clientHeight });
-					HitTest(res, context, { -1, -1 });
+					PaintContext* context = m_backend.CreateContext(m_clientWidth, m_clientHeight);
+					HitTest(res, *context, { -1, -1 });
+					m_backend.DeleteContext(context);
 
 					m_scrollPosX = Min(Max(m_scrollPosX, float(m_clientWidth - m_contentWidth)), 0.0f);
 					m_scrollPosY = Min(0.0f, Max(m_scrollPosY, float(m_clientHeight - m_contentHeight)));
@@ -1439,8 +1300,7 @@ namespace uba
 			case Popup_DarkMode:
 			{
 				m_config.DarkMode = !m_config.DarkMode;
-				DirtyBitmaps(false);
-				InitBrushes();
+				DirtyBitmaps(true);
 				UpdateTheme();
 				Redraw(true);
 				break;
@@ -1673,8 +1533,9 @@ namespace uba
 		VisualizerWin* thisPtr;
 		if (Msg == WM_CREATE)
 		{
-			thisPtr = (VisualizerWin*)lParam;
-			SetWindowLongPtr(hWnd, GWLP_USERDATA, lParam);
+			LPCREATESTRUCT pcs = (LPCREATESTRUCT)lParam;
+			thisPtr = (VisualizerWin*)pcs->lpCreateParams;
+			SetWindowLongPtr(hWnd, GWLP_USERDATA, (LONG_PTR)thisPtr);
 			//EnableNonClientDpiScaling(hWnd);
 		}
 		else
diff --git a/Engine/Source/Programs/UnrealBuildAccelerator/Visualizer/Private/Windows/UbaVisualizerWin.h b/Engine/Source/Programs/UnrealBuildAccelerator/Visualizer/Private/Windows/UbaVisualizerWin.h
--- a/Engine/Source/Programs/UnrealBuildAccelerator/Visualizer/Private/Windows/UbaVisualizerWin.h
+++ b/Engine/Source/Programs/UnrealBuildAccelerator/Visualizer/Private/Windows/UbaVisualizerWin.h
@@ -25,9 +25,8 @@ namespace uba
 		HWND GetHwnd();
 		void Lock(bool lock);
 
-	private:
+	//private:
 		bool StartHwndThread();
-		void PaintClient(const Function<void(HDC hdc, HDC memDC)>& paintFunc);
 		virtual void PaintEmptyScreen(PaintContext& context) override;
 		virtual void FinishIncompatible() override;
 
@@ -35,7 +34,6 @@ namespace uba
 		void UnselectAndRedraw();
 		virtual bool UpdateSelection() override;
 		virtual void UpdateScrollbars(bool redraw) override;
-		void InitBrushes();
 		void ThreadLoop();
 		void Pause(bool pause);
 		void StartDragToScroll(const POINT& anchor);
@@ -50,7 +48,7 @@ namespace uba
 		virtual void StopUpdate() override;
 		void UpdateTheme();
 
-		u64 m_lastPaintTimeMs;
+		u64 m_lastPaintTimeMs = 0;
 
 		Atomic<bool> m_looping = false;
 		HWND m_hwnd = 0;
@@ -64,9 +62,6 @@ namespace uba
 		bool m_locked = false;
 		bool m_handlingTimer = false;
 
-		HBITMAP m_cachedBitmap = 0;
-		SIZE m_cachedBitmapSize = { INT_MIN, INT_MIN };
-
 		POINT m_mouseAnchor = {};
 		float m_scrollAtAnchorX = 0;
 		float m_scrollAtAnchorY = 0;
