[UBA]
* Added reason to when scheduler kill process

diff --git a/Engine/Source/Programs/UnrealBuildAccelerator/Common/Private/UbaScheduler.cpp b/Engine/Source/Programs/UnrealBuildAccelerator/Common/Private/UbaScheduler.cpp
--- a/Engine/Source/Programs/UnrealBuildAccelerator/Common/Private/UbaScheduler.cpp
+++ b/Engine/Source/Programs/UnrealBuildAccelerator/Common/Private/UbaScheduler.cpp
@@ -1390,7 +1390,7 @@ namespace uba
 				return;
 		}
 
-		m_session.GetTrace().SchedulerKillProcess(ph.GetId());
+		m_session.GetTrace().SchedulerKillProcess(ph.GetId(), reason);
 
 		u64 startTime = GetTime();
 		if (ProcessReturned(*ph.m_process, true))
@@ -1716,7 +1716,7 @@ namespace uba
 					if (i->HasExited())
 						continue;
 					ProcessHandle ph = *i;
-					m_session.GetTrace().SchedulerKillProcess(ph.GetId());
+					m_session.GetTrace().SchedulerKillProcess(ph.GetId(), TCV("race"));
 					u64 startTime = GetTime();
 					if (!ProcessReturned(*ph.m_process, true))
 						continue;
diff --git a/Engine/Source/Programs/UnrealBuildAccelerator/Common/Private/UbaTrace.cpp b/Engine/Source/Programs/UnrealBuildAccelerator/Common/Private/UbaTrace.cpp
--- a/Engine/Source/Programs/UnrealBuildAccelerator/Common/Private/UbaTrace.cpp
+++ b/Engine/Source/Programs/UnrealBuildAccelerator/Common/Private/UbaTrace.cpp
@@ -644,10 +644,12 @@ namespace uba
 		writer.WriteBool(spawnPaused);
 	}
 
-	void Trace::SchedulerKillProcess(u32 processId)
+	void Trace::SchedulerKillProcess(u32 processId, StringView reason)
 	{
+		u32 stringIndex = AddString(reason);
 		BEGIN_TRACE_ENTRY(SchedulerKillProcess);
 		writer.WriteU32(processId);
+		writer.Write7BitEncoded(stringIndex);
 	}
 
 	void Trace::CacheBeginFetch(u32 fetchId, const tchar* description)
diff --git a/Engine/Source/Programs/UnrealBuildAccelerator/Common/Public/UbaTrace.h b/Engine/Source/Programs/UnrealBuildAccelerator/Common/Public/UbaTrace.h
--- a/Engine/Source/Programs/UnrealBuildAccelerator/Common/Public/UbaTrace.h
+++ b/Engine/Source/Programs/UnrealBuildAccelerator/Common/Public/UbaTrace.h
@@ -66,7 +66,7 @@ namespace uba
 		#undef UBA_TRACE_TYPE
 	};
 
-	static constexpr u32 TraceVersion = 51;
+	static constexpr u32 TraceVersion = 52;
 	static constexpr u32 TraceReadCompatibilityVersion = 6;
 
 	////////////////////////////////////////////////////////////////////////////////////////////////////
@@ -129,7 +129,7 @@ namespace uba
 		void DriveUpdate(const tchar drive, u8 busyPercent, u32 readCount, u64 readBytes, u32 writeCount, u64 writeBytes);
 		void RemoteExecutionDisabled();
 		void SchedulerUpdate(bool spawnPaused);
-		void SchedulerKillProcess(u32 processId);
+		void SchedulerKillProcess(u32 processId, StringView reason);
 
 		void CacheBeginFetch(u32 fetchId, const tchar* description);
 		void CacheEndFetch(u32 fetchId, bool success, const u8* data, u64 dataSize);
diff --git a/Engine/Source/Programs/UnrealBuildAccelerator/Visualizer/Private/UbaTraceReader.cpp b/Engine/Source/Programs/UnrealBuildAccelerator/Visualizer/Private/UbaTraceReader.cpp
--- a/Engine/Source/Programs/UnrealBuildAccelerator/Visualizer/Private/UbaTraceReader.cpp
+++ b/Engine/Source/Programs/UnrealBuildAccelerator/Visualizer/Private/UbaTraceReader.cpp
@@ -75,6 +75,7 @@ namespace uba
 		progressProcessesDone = 0;
 		progressErrorCount = 0;
 		lastKillProcessTime = 0;
+		lastKillReason = nullptr;
 		lastSpawningDelayStartTime = 0;
 		lastSpawningDelayEndTime = 0;
 		remoteExecutionDisabled = false;
@@ -1501,7 +1502,20 @@ namespace uba
 		case TraceType_SchedulerKillProcess:
 		{
 			u32 processId = reader.ReadU32();(void)processId;
+
+			const tchar* killReason = nullptr;
+			if (out.version >= 52)
+				killReason = out.strings[reader.Read7BitEncoded()];
 			out.lastKillProcessTime = time;
+			out.lastKillReason = killReason;
+
+			auto findIt = m_activeProcesses.find(processId);
+			if (findIt != m_activeProcesses.end())
+			{
+				auto& active = findIt->second;
+				TraceView::Process& process = GetSession(out, active.sessionIndex).processors[active.processorIndex].processes[active.processIndex];
+				process.killReason = killReason;
+			}
 			break;
 		}
 		default:
diff --git a/Engine/Source/Programs/UnrealBuildAccelerator/Visualizer/Private/UbaTraceReader.h b/Engine/Source/Programs/UnrealBuildAccelerator/Visualizer/Private/UbaTraceReader.h
--- a/Engine/Source/Programs/UnrealBuildAccelerator/Visualizer/Private/UbaTraceReader.h
+++ b/Engine/Source/Programs/UnrealBuildAccelerator/Visualizer/Private/UbaTraceReader.h
@@ -43,6 +43,7 @@ namespace uba
 			u64 writeFilesTime = 0;
 			Vector<u8> stats;
 			Vector<ProcessLogLine> logLines;
+			const tchar* killReason = nullptr;
 		};
 
 		struct Processor
@@ -209,6 +210,7 @@ namespace uba
 		u64 frequency = 0;
 		double timeMultiplier = 0;
 		u64 lastKillProcessTime = 0;
+		const tchar* lastKillReason = nullptr;
 		u64 lastSpawningDelayStartTime = 0;
 		u64 lastSpawningDelayEndTime = 0;
 		u32 totalProcessActiveCount = 0;
diff --git a/Engine/Source/Programs/UnrealBuildAccelerator/Visualizer/Private/UbaVisualizer.cpp b/Engine/Source/Programs/UnrealBuildAccelerator/Visualizer/Private/UbaVisualizer.cpp
--- a/Engine/Source/Programs/UnrealBuildAccelerator/Visualizer/Private/UbaVisualizer.cpp
+++ b/Engine/Source/Programs/UnrealBuildAccelerator/Visualizer/Private/UbaVisualizer.cpp
@@ -875,8 +875,14 @@ namespace uba
 			{
 				rect.left = start + width + 20;
 				rect.right = rect.left + activeFont.height * 7;
-				context.FillRect(rect, m_processBrushes[0].error);
-				drawStatusText(TCV("Killing (mem)"), LogEntryType_Info, rect.left, rect.right, 0, false, true);
+				const tchar* killReason = m_traceView.lastKillReason;
+				if (!killReason)
+					killReason = TC("mem");
+				HBRUSH brush = TCV("mem").Equals(killReason) ? m_processBrushes[0].error : m_processBrushes[0].returned;
+				context.FillRect(rect, brush);
+				StringBuffer<128> text;
+				text.Appendf(TC("Killing (%s)"), killReason);
+				drawStatusText(text, LogEntryType_Info, rect.left, rect.right, 0, false, true);
 			}
 			else if (m_traceView.lastSpawningDelayStartTime && playTime >= m_traceView.lastSpawningDelayStartTime && playTime < m_traceView.lastSpawningDelayEndTime + MsToTime(500))
 			{
@@ -1799,7 +1805,7 @@ namespace uba
 			if (hasExited && process.exitCode != 0)
 			{
 				if (process.exitCode == ProcessCancelExitCode)
-					logger.Info(TC("  ExitCode: Cancelled"));
+					logger.Info(TC("  ExitCode: Cancelled%s"), process.isRacing ? TC(" (Lost race)") : TC(""));
 				else if (process.exitCode >= 0xC0000000u)
 					logger.Info(TC("  ExitCode: 0x%x"), process.exitCode);
 				else
