[UBA]
* Swapped colors around so active is blue and cancelled (because of memory or race lost) is gray

diff --git a/Engine/Source/Programs/UnrealBuildAccelerator/Visualizer/Private/UbaVisualizer.cpp b/Engine/Source/Programs/UnrealBuildAccelerator/Visualizer/Private/UbaVisualizer.cpp
--- a/Engine/Source/Programs/UnrealBuildAccelerator/Visualizer/Private/UbaVisualizer.cpp
+++ b/Engine/Source/Programs/UnrealBuildAccelerator/Visualizer/Private/UbaVisualizer.cpp
@@ -181,16 +181,16 @@ namespace uba
 			m_textWarningColor = RGB(190, 190, 0);
 			m_textErrorColor = RGB(190, 0, 0);
 
-			m_processBrushes[0].inProgress = CreateBrush(70, 70, 70);
-			m_processBrushes[1].inProgress = CreateBrush(130, 130, 130);
-			m_processBrushes[0].inProgressRace = CreateBrush(50, 50, 50);
-			m_processBrushes[1].inProgressRace = CreateBrush(120, 120, 120);
+			m_processBrushes[0].inProgress = CreateBrush(70, 70, 160);
+			m_processBrushes[1].inProgress = CreateBrush(80, 80, 170);
+			m_processBrushes[0].inProgressRace = CreateBrush(60, 60, 120);
+			m_processBrushes[1].inProgressRace = CreateBrush(80, 80, 140);
 
 			m_processBrushes[0].error = CreateBrush(140, 0, 0);
 			m_processBrushes[1].error = CreateBrush(190, 0, 0);
 
-			m_processBrushes[0].returned = CreateBrush(50, 50, 120);
-			m_processBrushes[1].returned = CreateBrush(70, 70, 160);
+			m_processBrushes[0].returned = CreateBrush(70, 70, 70);
+			m_processBrushes[1].returned = CreateBrush(130, 130, 130);
 
 			m_processBrushes[0].recv = CreateBrush(10, 92, 10);
 			m_processBrushes[1].recv = CreateBrush(10, 130, 10);
@@ -232,22 +232,22 @@ namespace uba
 			m_textWarningColor = RGB(170, 130, 0);
 			m_textErrorColor = RGB(190, 0, 0);
 
-			m_processBrushes[0].inProgress = CreateBrush(150, 150, 150);
-			m_processBrushes[1].inProgress = CreateBrush(180, 180, 180);
+			m_processBrushes[0].inProgress = CreateBrush(150, 150, 200);
+			m_processBrushes[1].inProgress = CreateBrush(170, 170, 200);
 			m_processBrushes[0].inProgressRace = CreateBrush(130, 130, 130);
 			m_processBrushes[1].inProgressRace = CreateBrush(170, 170, 170);
 
 			m_processBrushes[0].error = CreateBrush(255, 70, 70);
 			m_processBrushes[1].error = CreateBrush(255, 100, 70);
 
-			m_processBrushes[0].returned = CreateBrush(150, 150, 200);
-			m_processBrushes[1].returned = CreateBrush(170, 170, 200);
+			m_processBrushes[0].returned = CreateBrush(150, 150, 150);
+			m_processBrushes[1].returned = CreateBrush(180, 180, 180);
 
-			m_processBrushes[0].recv = CreateBrush(10, 190, 10);
+			m_processBrushes[0].recv = CreateBrush(70, 170, 70);
 			m_processBrushes[1].recv = CreateBrush(20, 210, 20);
-			m_processBrushes[0].success = CreateBrush(10, 200, 10);
+			m_processBrushes[0].success = CreateBrush(80, 180, 80);
 			m_processBrushes[1].success = CreateBrush(20, 220, 20);
-			m_processBrushes[0].send = CreateBrush(80, 210, 80);
+			m_processBrushes[0].send = CreateBrush(80, 190, 80);
 			m_processBrushes[1].send = CreateBrush(90, 250, 90);
 			m_processBrushes[0].cacheFetchTest = CreateBrush(150, 150, 200);
 			m_processBrushes[1].cacheFetchTest = CreateBrush(170, 170, 200);
@@ -256,7 +256,7 @@ namespace uba
 			m_processBrushes[0].warning = CreateBrush(170, 170, 0);
 			m_processBrushes[1].warning = CreateBrush(170, 170, 0);
 
-			m_separatorPen = m_backend.CreatePen(RGB(180, 180, 180));
+			m_separatorPen = m_backend.CreatePen(RGB(140, 140, 140));
 			m_backgroundBrush = m_backend.CreateSolidBrush(backgroundColor);
 			m_tooltipBackgroundBrush = m_backend.CreateSolidBrush(tooltipBackgroundColor);
 			m_checkboxPen = m_backend.CreatePen(RGB(130, 130, 130));
@@ -857,7 +857,7 @@ namespace uba
 			rect.right = start + width;
 			rect.top = posY;
 			rect.bottom = posY + activeFont.height;
-			context.FillRect(rect, m_processBrushes[0].inProgress);
+			context.FillRect(rect, m_processBrushes[0].returned);
 
 			rect.right = rect.left + int(progress*float(width));//durationMs/100;
 			context.FillRect(rect, m_traceView.progressErrorCount ? m_processBrushes[0].error : m_processBrushes[0].success);
@@ -1539,7 +1539,7 @@ namespace uba
 							lastDrawnRight = rect.right;
 
 							bool done = stop != ~u64(0);
-							if (done)
+							if (done && m_config.DarkMode) // Color only work in dark mode.. hard to see in light mode)
 							{
 								auto insres = m_coloredBrushes.try_emplace(color);
 								if (insres.second)
@@ -1552,7 +1552,7 @@ namespace uba
 							else
 							{
 								//stop = playTime;
-								brush = m_processBrushes[0].inProgress;
+								brush = m_processBrushes[selected].returned;
 							}
 
 							--rect.top;
