[UBA]
* Bumped reserved address space for MemoryBlock when writing .exi files

diff --git a/Engine/Source/Programs/UnrealBuildAccelerator/Common/Private/UbaSession.cpp b/Engine/Source/Programs/UnrealBuildAccelerator/Common/Private/UbaSession.cpp
--- a/Engine/Source/Programs/UnrealBuildAccelerator/Common/Private/UbaSession.cpp
+++ b/Engine/Source/Programs/UnrealBuildAccelerator/Common/Private/UbaSession.cpp
@@ -4060,7 +4060,7 @@ namespace uba
 		#if UBA_DEBUG
 		verbose = true;
 		#endif
-		MemoryBlock memoryBlock(32*1024*1024);
+		MemoryBlock memoryBlock(64*1024*1024); // Seen files that are over 32mb when verbose = true
 		if (!objectFile->WriteImportsAndExports(m_logger, memoryBlock, verbose))
 			return false;
 
