[UBA]
* Fixed so extra obj file when building merge modules can contain more than 65535 loopbacks.

diff --git a/Engine/Source/Programs/UnrealBuildAccelerator/Common/Private/UbaObjectFileCoff.cpp b/Engine/Source/Programs/UnrealBuildAccelerator/Common/Private/UbaObjectFileCoff.cpp
--- a/Engine/Source/Programs/UnrealBuildAccelerator/Common/Private/UbaObjectFileCoff.cpp
+++ b/Engine/Source/Programs/UnrealBuildAccelerator/Common/Private/UbaObjectFileCoff.cpp
@@ -409,14 +409,29 @@ namespace uba
 		u32 dataRawPos = u32(memoryBlock.writtenSize);
 		allocate(loopbackCount * 8);
 
+		bool useOverflow = loopbackCount > 0xFFFF;
+
 		// Relocations for each 8-byte slot (pointing at the base symbol, i.e., "foo")
 		u32 relocFilePos = u32(memoryBlock.writtenSize);
-		auto relocations = (ImageRelocation*)allocate(loopbackCount * sizeof(ImageRelocation));
+		u32 relocRecordCount = useOverflow ? (loopbackCount + 1) : loopbackCount;
+		auto relocations = (ImageRelocation*)allocate(relocRecordCount * sizeof(ImageRelocation));
 		rdataSection.PointerToRawData   = dataRawPos;
 		rdataSection.SizeOfRawData      = loopbackCount * 8;
 		rdataSection.PointerToRelocations = relocFilePos;
 		rdataSection.NumberOfRelocations  = u16(loopbackCount);
 
+		if (useOverflow)
+		{
+			rdataSection.Characteristics |= IMAGE_SCN_LNK_NRELOC_OVFL;
+			rdataSection.NumberOfRelocations = 0xFFFF;
+
+			// First entry is the overflow header
+			relocations[0].VirtualAddress   = loopbackCount; // actual relocation count
+			relocations[0].SymbolTableIndex = 0;
+			relocations[0].Type             = 0;
+			++relocations;
+		}
+
 		for (u32 i = 0; i != loopbackCount; ++i)
 		{
 			auto& relocation = relocations[i];
