[UBA]
* Replaced HBRUSH/HPEN/HFONT to non-gdi specific types to make it clear they are not these types when using non-gdi backends

diff --git a/Engine/Source/Programs/UnrealBuildAccelerator/Visualizer/Private/UbaVisualizer.cpp b/Engine/Source/Programs/UnrealBuildAccelerator/Visualizer/Private/UbaVisualizer.cpp
--- a/Engine/Source/Programs/UnrealBuildAccelerator/Visualizer/Private/UbaVisualizer.cpp
+++ b/Engine/Source/Programs/UnrealBuildAccelerator/Visualizer/Private/UbaVisualizer.cpp
@@ -18,7 +18,7 @@ namespace uba
 
 	struct Visualizer::DrawTextLogger : public Logger
 	{
-		DrawTextLogger(PaintContext& c, int fh, HBRUSH bb)
+		DrawTextLogger(PaintContext& c, int fh, BrushHandle bb)
 			:	context(c)
 			,	fontHeight(fh)
 			,	backgroundBrush(bb)
@@ -104,7 +104,7 @@ namespace uba
 
 		PaintContext& context;
 		int fontHeight;
-		HBRUSH backgroundBrush;
+		BrushHandle backgroundBrush;
 		COLORREF textColor;
 		bool isFirst = true;
 		using Logger::Log;
@@ -112,7 +112,7 @@ namespace uba
 
 	////////////////////////////////////////////////////////////////////////////////////////////////////
 
-	void DrawCenteredText(VisualizerBackend::PaintContext& context, const StringBuffer<>* rows, u32 rowCount, HBRUSH backgroundBrush)
+	void DrawCenteredText(VisualizerBackend::PaintContext& context, const StringBuffer<>* rows, u32 rowCount, BrushHandle backgroundBrush)
 	{
 		SIZE clientSize = context.GetClientSize();
 		for (u32 i=0; i!=rowCount; ++i)
@@ -202,6 +202,8 @@ namespace uba
 			m_processBrushes[1].cacheFetchTest = CreateBrush(31, 150, 138);
 			m_processBrushes[0].cacheFetchDownload = CreateBrush(24, 112, 110);
 			m_processBrushes[1].cacheFetchDownload = CreateBrush(31, 143, 138);
+			m_processBrushes[0].cacheMiss = CreateBrush(90, 90, 90);
+			m_processBrushes[1].cacheMiss = CreateBrush(130, 130, 130);
 			m_processBrushes[0].warning = CreateBrush(100, 100, 0);
 			m_processBrushes[1].warning = CreateBrush(100, 100, 0);
 
@@ -232,10 +234,10 @@ namespace uba
 			m_textWarningColor = RGB(170, 130, 0);
 			m_textErrorColor = RGB(190, 0, 0);
 
-			m_processBrushes[0].inProgress = CreateBrush(150, 150, 200);
-			m_processBrushes[1].inProgress = CreateBrush(170, 170, 200);
-			m_processBrushes[0].inProgressRace = CreateBrush(130, 130, 130);
-			m_processBrushes[1].inProgressRace = CreateBrush(170, 170, 170);
+			m_processBrushes[0].inProgress = CreateBrush(100, 150, 200);
+			m_processBrushes[1].inProgress = CreateBrush(120, 170, 200);
+			m_processBrushes[0].inProgressRace = CreateBrush(100, 130, 180);
+			m_processBrushes[1].inProgressRace = CreateBrush(130, 180, 210);
 
 			m_processBrushes[0].error = CreateBrush(255, 70, 70);
 			m_processBrushes[1].error = CreateBrush(255, 100, 70);
@@ -253,6 +255,8 @@ namespace uba
 			m_processBrushes[1].cacheFetchTest = CreateBrush(170, 170, 200);
 			m_processBrushes[0].cacheFetchDownload = CreateBrush(150, 150, 200);
 			m_processBrushes[1].cacheFetchDownload = CreateBrush(170, 170, 200);
+			m_processBrushes[0].cacheMiss = CreateBrush(150, 150, 150);
+			m_processBrushes[1].cacheMiss = CreateBrush(180, 180, 180);
 			m_processBrushes[0].warning = CreateBrush(170, 170, 0);
 			m_processBrushes[1].warning = CreateBrush(170, 170, 0);
 
@@ -878,7 +882,7 @@ namespace uba
 				const tchar* killReason = m_traceView.lastKillReason;
 				if (!killReason)
 					killReason = TC("mem");
-				HBRUSH brush = TCV("mem").Equals(killReason) ? m_processBrushes[0].error : m_processBrushes[0].returned;
+				BrushHandle brush = TCV("mem").Equals(killReason) ? m_processBrushes[0].error : m_processBrushes[0].returned;
 				context.FillRect(rect, brush);
 				StringBuffer<128> text;
 				text.Appendf(TC("Killing (%s)"), killReason);
@@ -1173,7 +1177,7 @@ namespace uba
 
 					Vector<POINT> line;
 
-					auto drawGraph = [&](u32 id, auto& values, auto maxValue, double scale, HPEN pen, bool accumulatingValue, int offsetY = 0)
+					auto drawGraph = [&](u32 id, auto& values, auto maxValue, double scale, PenHandle pen, bool accumulatingValue, int offsetY = 0)
 						{
 							bool loop = true;
 							u32 reconnectIndex = 0;
@@ -1500,7 +1504,7 @@ namespace uba
 									continue;
 							}
 
-							HBRUSH brush;
+							BrushHandle brush;
 
 							u64 stop = work.stop;
 
@@ -2110,9 +2114,9 @@ namespace uba
 		auto clampRect = [&](RECT& r) { r.left = Min(Max(r.left, progressRect.left), progressRect.right); r.right = Max(Min(r.right, progressRect.right), progressRect.left); };
 		bool done = process.stop != ~u64(0);
 
-		HBRUSH brush = m_processBrushes[selected].success;
+		BrushHandle brush = m_processBrushes[selected].success;
 		if (!process.returnedReason.empty())
-			brush = m_processBrushes[selected].returned;
+			brush = process.type == TraceView::ProcessType_CacheFetchTest ? m_processBrushes[selected].cacheMiss : m_processBrushes[selected].returned;
 		else if (!done)
 			brush = process.isRacing ? m_processBrushes[selected].inProgressRace : m_processBrushes[selected].inProgress;
 		else if (process.type == TraceView::ProcessType_CacheFetchTest)
diff --git a/Engine/Source/Programs/UnrealBuildAccelerator/Visualizer/Private/UbaVisualizer.h b/Engine/Source/Programs/UnrealBuildAccelerator/Visualizer/Private/UbaVisualizer.h
--- a/Engine/Source/Programs/UnrealBuildAccelerator/Visualizer/Private/UbaVisualizer.h
+++ b/Engine/Source/Programs/UnrealBuildAccelerator/Visualizer/Private/UbaVisualizer.h
@@ -106,16 +106,17 @@ namespace uba
 
 		struct ProcessBrushes
 		{
-			HBRUSH inProgress = 0;
-			HBRUSH inProgressRace = 0;
-			HBRUSH success = 0;
-			HBRUSH error = 0;
-			HBRUSH returned = 0;
-			HBRUSH recv = 0;
-			HBRUSH send = 0;
-			HBRUSH cacheFetchTest = 0;
-			HBRUSH cacheFetchDownload = 0;
-			HBRUSH warning = 0;
+			BrushHandle inProgress = 0;
+			BrushHandle inProgressRace = 0;
+			BrushHandle success = 0;
+			BrushHandle error = 0;
+			BrushHandle returned = 0;
+			BrushHandle recv = 0;
+			BrushHandle send = 0;
+			BrushHandle cacheFetchTest = 0;
+			BrushHandle cacheFetchDownload = 0;
+			BrushHandle cacheMiss = 0;
+			BrushHandle warning = 0;
 		};
 
 		Logger& m_logger;
@@ -150,18 +151,18 @@ namespace uba
 		COLORREF m_memColor = {};
 		COLORREF m_driveColor = {};
 		COLORREF m_activeProcColor = {};
-		HBRUSH m_backgroundBrush = 0;
-		HBRUSH m_tooltipBackgroundBrush = 0;
-		HPEN m_textPen = 0;
-		HPEN m_separatorPen = 0;
-		HPEN m_sendPen = 0;
-		HPEN m_recvPen = 0;
-		HPEN m_cpuPen = 0;
-		HPEN m_memPen = 0;
-		HPEN m_drivePen = 0;
-		HPEN m_activeProcPen = 0;
-		HPEN m_processUpdatePen = 0;
-		HPEN m_checkboxPen = 0;
+		BrushHandle m_backgroundBrush = 0;
+		BrushHandle m_tooltipBackgroundBrush = 0;
+		PenHandle m_textPen = 0;
+		PenHandle m_separatorPen = 0;
+		PenHandle m_sendPen = 0;
+		PenHandle m_recvPen = 0;
+		PenHandle m_cpuPen = 0;
+		PenHandle m_memPen = 0;
+		PenHandle m_drivePen = 0;
+		PenHandle m_activeProcPen = 0;
+		PenHandle m_processUpdatePen = 0;
+		PenHandle m_checkboxPen = 0;
 		int m_sessionStepY = 0;
 
 		int m_processFontOffsetY = 0;
@@ -208,7 +209,7 @@ namespace uba
 
 		TString m_filterString;
 
-		UnorderedMap<Color, HBRUSH> m_coloredBrushes;
+		UnorderedMap<Color, BrushHandle> m_coloredBrushes;
 
 		bool m_workSelected = false;
 		u32 m_workTrack = ~0u;
diff --git a/Engine/Source/Programs/UnrealBuildAccelerator/Visualizer/Private/UbaVisualizerBackend.h b/Engine/Source/Programs/UnrealBuildAccelerator/Visualizer/Private/UbaVisualizerBackend.h
--- a/Engine/Source/Programs/UnrealBuildAccelerator/Visualizer/Private/UbaVisualizerBackend.h
+++ b/Engine/Source/Programs/UnrealBuildAccelerator/Visualizer/Private/UbaVisualizerBackend.h
@@ -11,17 +11,13 @@ namespace uba
 	struct SIZE { int cx; int cy; };
 	struct RECT { int left; int top; int right; int bottom; };
 	using COLORREF = u32;
-	using HFONT = void*;
-	using HPEN = void*;
-	using HBRUSH = void*;
-	#define TRANSPARENT  1
-	#define OPAQUE       2
-	#define TA_LEFT      0
-	#define TA_RIGHT     2
-	#define TA_CENTER    6
-	#define RGB(r,g,b)          ((COLORREF)((u8(r)|(u16(u8(g))<<8))|(u32(u8(b))<<16)))
+	#define RGB(r,g,b) ((COLORREF)((u8(r)|(u16(u8(g))<<8))|(u32(u8(b))<<16)))
 	#endif
 
+	using BrushHandle = void*;
+	using PenHandle = void*;
+	using FontHandle = void*;
+
 	////////////////////////////////////////////////////////////////////////////////////////////////////
 
 	class VisualizerBackend
@@ -29,8 +25,8 @@ namespace uba
 	public:
 		struct Font
 		{
-			HFONT handle = 0;
-			HFONT handleUnderlined = 0;
+			FontHandle handle = 0;
+			FontHandle handleUnderlined = 0;
 			int height = 0;
 			int offset = 0;
 		};
@@ -38,14 +34,14 @@ namespace uba
 		struct PaintContext
 		{
 			virtual SIZE GetClientSize() = 0;
-			virtual void SetPen(HPEN pen) = 0;
-			virtual void SetBrush(HBRUSH b) = 0;
+			virtual void SetPen(PenHandle pen) = 0;
+			virtual void SetBrush(BrushHandle b) = 0;
 			virtual void SetFont(const Font& font, bool underlined = false) = 0;
 			virtual void SetTextColor(COLORREF color) = 0;
 			virtual void Text(int x, int y, StringView str, const RECT* clipRect) = 0;
 			virtual COLORREF GetTextColor() = 0;
 			virtual Font& GetFont() = 0;
-			virtual void FillRect(const RECT& r, HBRUSH b) = 0;
+			virtual void FillRect(const RECT& r, BrushHandle b) = 0;
 			virtual void Rectangle(int left, int top, int right, int bottom) = 0;
 			virtual void Line(int fromX, int fromY, int toX, int toY) = 0;
 			virtual void Polyline(u32 cacheId, const Vector<POINT>& polyLine) = 0;
@@ -66,11 +62,11 @@ namespace uba
 		virtual PaintContext* CreateContext(int clientWidth, int clientHeight) = 0;
 		virtual void DeleteContext(PaintContext*) = 0;
 
-		virtual HBRUSH CreateSolidBrush(COLORREF color) = 0;
-		virtual HPEN CreatePen(COLORREF color) = 0;
+		virtual BrushHandle CreateSolidBrush(COLORREF color) = 0;
+		virtual PenHandle CreatePen(COLORREF color) = 0;
 
-		virtual HFONT CreateFont(const tchar* font, int height, bool createUnderline) = 0;
-		virtual void DeleteFont(HFONT font) = 0;
+		virtual FontHandle CreateFont(const tchar* font, int height, bool createUnderline) = 0;
+		virtual void DeleteFont(FontHandle font) = 0;
 	};
 
 	////////////////////////////////////////////////////////////////////////////////////////////////////
diff --git a/Engine/Source/Programs/UnrealBuildAccelerator/Visualizer/Private/Windows/UbaVisualizerBackendD2D.h b/Engine/Source/Programs/UnrealBuildAccelerator/Visualizer/Private/Windows/UbaVisualizerBackendD2D.h
--- a/Engine/Source/Programs/UnrealBuildAccelerator/Visualizer/Private/Windows/UbaVisualizerBackendD2D.h
+++ b/Engine/Source/Programs/UnrealBuildAccelerator/Visualizer/Private/Windows/UbaVisualizerBackendD2D.h
@@ -40,11 +40,11 @@ namespace uba
 			{
 				return clientSize;
 			}
-			virtual void SetPen(HPEN pen) override
+			virtual void SetPen(PenHandle pen) override
 			{
 				activePen = pen;
 			}
-			virtual void SetBrush(HBRUSH b) override
+			virtual void SetBrush(BrushHandle b) override
 			{
 				activeBrush = b;
 			}
@@ -85,7 +85,7 @@ namespace uba
 				if (clipRect)
 					renderTarget.PushAxisAlignedClip(layoutRect, D2D1_ANTIALIAS_MODE_ALIASED);
 
-				if ((HFONT)activeFontHandle == activeFont.handleUnderlined)
+				if ((FontHandle)activeFontHandle == activeFont.handleUnderlined)
 				{
 					IDWriteTextLayout* layout = nullptr;
 					owner.m_dwFactory->CreateTextLayout(str.data, str.count, activeFontHandle, 10000.0f, 1000.0f, &layout);
@@ -112,7 +112,7 @@ namespace uba
 			{
 				return activeFont;
 			}
-			virtual void FillRect(const RECT& r, HBRUSH b) override
+			virtual void FillRect(const RECT& r, BrushHandle b) override
 			{
 				D2D1_RECT_F rect = D2D1::RectF(float(r.left), float(r.top), float(r.right), float(r.bottom));
 				renderTarget.FillRectangle(rect, (ID2D1SolidColorBrush*)b);
@@ -258,8 +258,8 @@ namespace uba
 			VisualizerBackendD2D& owner;
 			ID2D1RenderTarget& renderTarget;
 			SIZE clientSize;
-			HPEN activePen = 0;
-			HBRUSH activeBrush = 0;
+			PenHandle activePen = 0;
+			BrushHandle activeBrush = 0;
 			Font activeFont;
 			IDWriteTextFormat* activeFontHandle = nullptr;
 			COLORREF activeTextColor = 0;
@@ -356,20 +356,20 @@ namespace uba
 			delete (PaintContextD2D*)context;
 		}
 
-		virtual HBRUSH CreateSolidBrush(COLORREF color) override
+		virtual BrushHandle CreateSolidBrush(COLORREF color) override
 		{
 			ID2D1SolidColorBrush* brush;
 			m_d2dRenderTarget->CreateSolidColorBrush(ToRgbD2D(color), &brush);
 			m_brushes.push_back(brush);
-			return (HBRUSH)brush;
+			return (BrushHandle)brush;
 		}
 
-		virtual HPEN CreatePen(COLORREF color) override
+		virtual PenHandle CreatePen(COLORREF color) override
 		{
-			return (HPEN)CreateSolidBrush(color);
+			return (PenHandle)CreateSolidBrush(color);
 		}
 
-		virtual HFONT CreateFont(const tchar* font, int height, bool createUnderline) override
+		virtual FontHandle CreateFont(const tchar* font, int height, bool createUnderline) override
 		{
 			auto weight = DWRITE_FONT_WEIGHT_NORMAL;
 			auto style = DWRITE_FONT_STYLE_NORMAL;
@@ -379,10 +379,10 @@ namespace uba
 				return 0;
 			textFormat->SetTextAlignment(DWRITE_TEXT_ALIGNMENT_LEADING);
 			textFormat->SetParagraphAlignment(DWRITE_PARAGRAPH_ALIGNMENT_NEAR);
-			return (HFONT)textFormat;
+			return (FontHandle)textFormat;
 		}
 
-		virtual void DeleteFont(HFONT font)
+		virtual void DeleteFont(FontHandle font)
 		{
 			((IDWriteTextFormat*)font)->Release();
 		}
diff --git a/Engine/Source/Programs/UnrealBuildAccelerator/Visualizer/Private/Windows/UbaVisualizerBackendGDI.h b/Engine/Source/Programs/UnrealBuildAccelerator/Visualizer/Private/Windows/UbaVisualizerBackendGDI.h
--- a/Engine/Source/Programs/UnrealBuildAccelerator/Visualizer/Private/Windows/UbaVisualizerBackendGDI.h
+++ b/Engine/Source/Programs/UnrealBuildAccelerator/Visualizer/Private/Windows/UbaVisualizerBackendGDI.h
@@ -20,7 +20,7 @@ namespace uba
 			virtual SIZE GetClientSize() override { return clientSize; }
 			virtual void SetTextColor(COLORREF color) override { ::SetTextColor(hdc, color); }
 			virtual COLORREF GetTextColor() override { return ::GetTextColor(hdc); }
-			virtual void SetPen(HPEN pen) override { ::SelectObject(hdc, pen); }
+			virtual void SetPen(PenHandle pen) override { ::SelectObject(hdc, pen); }
 			virtual void Text(int x, int y, StringView str, const RECT* clipRect) override
 			{
 				::ExtTextOutW(hdc, x, y, clipRect ? ETO_CLIPPED : 0, clipRect, str.data, str.count, NULL);
@@ -48,7 +48,7 @@ namespace uba
 				return activeFont;
 			}
 
-			virtual void SetBrush(HBRUSH b) override
+			virtual void SetBrush(BrushHandle b) override
 			{
 				if (lastSelectedBrush != b)
 				{
@@ -57,7 +57,7 @@ namespace uba
 				}
 			}
 
-			virtual void FillRect(const RECT& r, HBRUSH b) override
+			virtual void FillRect(const RECT& r, BrushHandle b) override
 			{
 				SetBrush(b);
 				PatBlt(hdc, r.left, r.top, r.right - r.left, r.bottom - r.top, PATCOPY);
@@ -65,7 +65,7 @@ namespace uba
 
 			virtual void Rectangle(int left, int top, int right, int bottom) override
 			{
-				SetBrush((HBRUSH)GetStockObject(NULL_BRUSH));
+				SetBrush((BrushHandle)GetStockObject(NULL_BRUSH));
 				::Rectangle(hdc, left, top, right, bottom);
 			}
 
@@ -171,10 +171,10 @@ namespace uba
 			SIZE clientSize;
 
 			Font activeFont;
-			HFONT activeFontHandle;
+			FontHandle activeFontHandle;
 
 			HDC textDC = NULL;
-			HBRUSH lastSelectedBrush = 0;
+			BrushHandle lastSelectedBrush = 0;
 			HBITMAP lastSelectedBitmap = 0;
 		};
 
@@ -224,7 +224,7 @@ namespace uba
 			HGDIOBJ oldBmp = SelectObject(memDC, m_cachedBitmap);
 
 			RECT rect = { 0, 0, clientWidth, clientHeight };
-			FillRect(memDC, &rect, m_owner.m_backgroundBrush);
+			FillRect(memDC, &rect, (HBRUSH)m_owner.m_backgroundBrush);
 			SetBkColor(memDC, m_owner.m_config.DarkMode ? RGB(70, 70, 70) : RGB(180, 180, 180));
 			SetBkMode(memDC, TRANSPARENT);
 
@@ -263,23 +263,23 @@ namespace uba
 			delete (PaintContextGDI*)context;
 		}
 
-		virtual HBRUSH CreateSolidBrush(COLORREF color) override
+		virtual BrushHandle CreateSolidBrush(COLORREF color) override
 		{
 			return ::CreateSolidBrush(color);
 		}
 
-		virtual HPEN CreatePen(COLORREF color) override
+		virtual PenHandle CreatePen(COLORREF color) override
 		{
 			return ::CreatePen(PS_SOLID, 1, color);
 		}
 
-		HFONT CreateFont(const tchar* font, int height, bool createUnderline) override
+		FontHandle CreateFont(const tchar* font, int height, bool createUnderline) override
 		{
 			DWORD underline = createUnderline ? 1 : 0;
 			return ::CreateFontW(4 - height, 0, 0, 0, FW_NORMAL, 0, underline, 0, ANSI_CHARSET, OUT_DEFAULT_PRECIS, CLIP_DEFAULT_PRECIS, CLEARTYPE_QUALITY, DEFAULT_PITCH, font);
 		}
 
-		void DeleteFont(HFONT font)
+		void DeleteFont(FontHandle font)
 		{
 			::DeleteObject(font);
 		}
diff --git a/Engine/Source/Programs/UnrealBuildAccelerator/Visualizer/Private/Windows/UbaVisualizerWin.cpp b/Engine/Source/Programs/UnrealBuildAccelerator/Visualizer/Private/Windows/UbaVisualizerWin.cpp
--- a/Engine/Source/Programs/UnrealBuildAccelerator/Visualizer/Private/Windows/UbaVisualizerWin.cpp
+++ b/Engine/Source/Programs/UnrealBuildAccelerator/Visualizer/Private/Windows/UbaVisualizerWin.cpp
@@ -316,7 +316,7 @@ namespace uba
 
 		m_popupFont.handle = m_backend.CreateFont(TEXT("Consolas"), 16, false);
 		m_popupFont.height = 14;
-		//m_popupFont = (HFONT)GetStockObject(SYSTEM_FIXED_FONT);
+		//m_popupFont = (FontHandle)GetStockObject(SYSTEM_FIXED_FONT);
 
 		DWORD scrollbarFlags = 0;
 		m_verticalScrollBarEnabled = !ActiveProcessesShouldFillHeight();
