[UBA]
* Fixed bug in d2d visualizer backend where DrawText used multiple lines

diff --git a/Engine/Source/Programs/UnrealBuildAccelerator/Visualizer/Private/Windows/UbaVisualizerBackendD2D.h b/Engine/Source/Programs/UnrealBuildAccelerator/Visualizer/Private/Windows/UbaVisualizerBackendD2D.h
--- a/Engine/Source/Programs/UnrealBuildAccelerator/Visualizer/Private/Windows/UbaVisualizerBackendD2D.h
+++ b/Engine/Source/Programs/UnrealBuildAccelerator/Visualizer/Private/Windows/UbaVisualizerBackendD2D.h
@@ -66,24 +66,19 @@ namespace uba
 				x += clientOffsetX;
 				y += clientOffsetY;
 
-				int right = x + 10000;
-				int bottom = y + 100;
+				D2D1_RECT_F layoutRect = D2D1::RectF(float(x), float(y), float(x + 10000), float(y + 100));
+
 				if (clipRect)
 				{
 					//UBA_ASSERT(clipRect->left <= x);
 					//UBA_ASSERT(clipRect->top <= y);
-					right = clipRect->right + clientOffsetX;
-					bottom = clipRect->bottom + clientOffsetY;
+					D2D1_RECT_F cr = layoutRect;
+					cr.right = clipRect->right + float(clientOffsetX);
+					cr.bottom = clipRect->bottom + float(clientOffsetX);
+					if (cr.bottom < 0)
+						return;
+					renderTarget.PushAxisAlignedClip(cr, D2D1_ANTIALIAS_MODE_ALIASED);
 				}
-				if (y > clientSize.cy)
-					return;
-				if (bottom < 0)
-					return;
-
-				D2D1_RECT_F layoutRect = D2D1::RectF(float(x), float(y), float(right), float(bottom));
-
-				if (clipRect)
-					renderTarget.PushAxisAlignedClip(layoutRect, D2D1_ANTIALIAS_MODE_ALIASED);
 
 				if ((FontHandle)activeFontHandle == activeFont.handleUnderlined)
 				{
