[UBA]
* Fixed posix compile error in UbaObjectFileCoff

diff --git a/Engine/Source/Programs/UnrealBuildAccelerator/Common/Private/UbaObjectFileCoff.cpp b/Engine/Source/Programs/UnrealBuildAccelerator/Common/Private/UbaObjectFileCoff.cpp
--- a/Engine/Source/Programs/UnrealBuildAccelerator/Common/Private/UbaObjectFileCoff.cpp
+++ b/Engine/Source/Programs/UnrealBuildAccelerator/Common/Private/UbaObjectFileCoff.cpp
@@ -17,6 +17,7 @@ namespace uba
 	constexpr u32 ImageScnLnkRemove = 0x00000800; // IMAGE_SCN_LNK_REMOVE
 	constexpr u32 ImageScnLnkComdat = 0x00001000; // IMAGE_SCN_LNK_COMDAT
 	constexpr u32 ImageScnAlign1Bytes = 0x00100000; // IMAGE_SCN_ALIGN_1BYTES
+	constexpr u32 ImageScnLnkNRelocOvfl = 0x01000000; // IMAGE_SCN_LNK_NRELOC_OVFL
 	constexpr u32 ImageScnMemExecute = 0x20000000; // IMAGE_SCN_MEM_EXECUTE
 	constexpr u32 ImageScnMemRead = 0x40000000; // IMAGE_SCN_MEM_READ
 
@@ -80,7 +81,7 @@ namespace uba
 		union
 		{
 			u32 VirtualAddress;
-			u32 RelocCount; // Set to the real count when IMAGE_SCN_LNK_NRELOC_OVFL is set
+			u32 RelocCount;
 		};
 		u32 SymbolTableIndex;
 		u16 Type;
@@ -422,7 +423,7 @@ namespace uba
 
 		if (useOverflow)
 		{
-			rdataSection.Characteristics |= IMAGE_SCN_LNK_NRELOC_OVFL;
+			rdataSection.Characteristics |= ImageScnLnkNRelocOvfl;
 			rdataSection.NumberOfRelocations = 0xFFFF;
 
 			// First entry is the overflow header
