[UBA]
* Fixed bug where UpdateDirectoryTableFromServer wait could hang if a message takes so long time so an earlier wait times out. Now a timed out wait sets error flag and tell all other waiters to stop waiting and exit instead
* Downgraded uba errors to warnings

diff --git a/Engine/Source/Programs/UnrealBuildAccelerator/Common/Private/UbaNetworkClient.cpp b/Engine/Source/Programs/UnrealBuildAccelerator/Common/Private/UbaNetworkClient.cpp
--- a/Engine/Source/Programs/UnrealBuildAccelerator/Common/Private/UbaNetworkClient.cpp
+++ b/Engine/Source/Programs/UnrealBuildAccelerator/Common/Private/UbaNetworkClient.cpp
@@ -866,7 +866,7 @@ namespace uba
 			u32 timeoutMs = 10 * 60 * 1000;
 			if (!gotResponse.IsSet(timeoutMs))
 			{
-				m_logger.Error(TC("Timed out after %s waiting for message response from server. (%s)"), TimeToText(GetTime() - waitStart, true).str, MessageToString(message.GetServiceId(), message.GetMessageType()).data);
+				m_logger.Warning(TC("Timed out after %s waiting for message response from server. (%s)"), TimeToText(GetTime() - waitStart, true).str, MessageToString(message.GetServiceId(), message.GetMessageType()).data);
 				message.m_error = 4;
 				OnDisconnected(connection, 4);
 			}
diff --git a/Engine/Source/Programs/UnrealBuildAccelerator/Common/Private/UbaSessionClient.cpp b/Engine/Source/Programs/UnrealBuildAccelerator/Common/Private/UbaSessionClient.cpp
--- a/Engine/Source/Programs/UnrealBuildAccelerator/Common/Private/UbaSessionClient.cpp
+++ b/Engine/Source/Programs/UnrealBuildAccelerator/Common/Private/UbaSessionClient.cpp
@@ -1131,10 +1131,16 @@ namespace uba
 			if (res)
 				return item.success;
 
+			UpdateError(first);
+			
+			if (client.m_directoryTableError)
+				return false;
+			client.m_directoryTableError = true;
+
 			u32 activeCount = 0;
 			for (auto i = first; i; i = i->next)
 				++activeCount;
-			return client.m_logger.Error(TC("Timed out after 5 minutes waiting for update directory message to reach read position %u  (%u active in %s wait)"), readPos, activeCount, hint);
+			return client.m_logger.Warning(TC("Timed out after 5 minutes waiting for update directory message to reach read position %u  (%u active in %s wait)"), readPos, activeCount, hint);
 		}
 
 		static void UpdateReadPosMatching(ActiveUpdateDirectoryEntry*& first, u32 readPos)
