[UBA]
* Added ProcessHandle_IsCancelled to UbaExports

diff --git a/Engine/Source/Programs/UnrealBuildAccelerator/Host/Private/UbaExports.cpp b/Engine/Source/Programs/UnrealBuildAccelerator/Host/Private/UbaExports.cpp
--- a/Engine/Source/Programs/UnrealBuildAccelerator/Host/Private/UbaExports.cpp
+++ b/Engine/Source/Programs/UnrealBuildAccelerator/Host/Private/UbaExports.cpp
@@ -411,6 +411,11 @@ extern "C"
 		return handle->GetExitCode();
 	}
 
+	bool ProcessHandle_IsCancelled(const uba::ProcessHandle* handle)
+	{
+		return handle->GetExitCode() == uba::ProcessCancelExitCode;
+	}
+
 	uba::u8 ProcessHandle_GetExecutionType(const uba::ProcessHandle* handle)
 	{
 		return handle->GetExecutionType();
diff --git a/Engine/Source/Programs/UnrealBuildAccelerator/Host/Public/UbaExports.h b/Engine/Source/Programs/UnrealBuildAccelerator/Host/Public/UbaExports.h
--- a/Engine/Source/Programs/UnrealBuildAccelerator/Host/Public/UbaExports.h
+++ b/Engine/Source/Programs/UnrealBuildAccelerator/Host/Public/UbaExports.h
@@ -92,6 +92,7 @@ extern "C"
 	// ProcessHandle
 	using ProcessHandle_OutputFileCallback = void(void* userData, const uba::tchar* file);
 	UBA_API uba::u32 ProcessHandle_GetExitCode(const uba::ProcessHandle* handle);
+	UBA_API bool ProcessHandle_IsCancelled(const uba::ProcessHandle* handle);
 	UBA_API uba::u8 ProcessHandle_GetExecutionType(const uba::ProcessHandle* handle);
 	UBA_API const uba::tchar* ProcessHandle_GetExecutingHost(const uba::ProcessHandle* handle);
 	UBA_API const uba::tchar* ProcessHandle_GetLogLine(const uba::ProcessHandle* handle, uba::u32 index);
