[UBA]
* Fixed bug with refcount going wrong when not writing received files to disk (Shader compiling). Refcount needed to be increased when added to receive list while reduced from storage server per-connection receive refcount

diff --git a/Engine/Source/Programs/UnrealBuildAccelerator/Common/Private/UbaSessionServer.cpp b/Engine/Source/Programs/UnrealBuildAccelerator/Common/Private/UbaSessionServer.cpp
--- a/Engine/Source/Programs/UnrealBuildAccelerator/Common/Private/UbaSessionServer.cpp
+++ b/Engine/Source/Programs/UnrealBuildAccelerator/Common/Private/UbaSessionServer.cpp
@@ -1362,6 +1362,12 @@ namespace uba
 			if (!rules.IsInvisible(destination))
 				if (!m_storage.FakeCopy(casKey, destination.data, 0, 0, file.attributes))
 					return m_logger.Error(TC("Failed to fake copy cas from %s to %s (%s)"), CasKeyString(casKey).str, destination.data, startInfo.GetDescription());
+
+			// We need to transfer ref count to this list
+			// and release the one inside StorageServer
+			m_storage.AddRef(casKey);
+			m_storage.SkipWrite(casKey, clientId);
+
 			SCOPED_WRITE_LOCK(m_receivedFilesLock, lock);
 			m_receivedFiles.try_emplace(file.destinationKey, casKey);
 		}
diff --git a/Engine/Source/Programs/UnrealBuildAccelerator/Common/Private/UbaStorage.cpp b/Engine/Source/Programs/UnrealBuildAccelerator/Common/Private/UbaStorage.cpp
--- a/Engine/Source/Programs/UnrealBuildAccelerator/Common/Private/UbaStorage.cpp
+++ b/Engine/Source/Programs/UnrealBuildAccelerator/Common/Private/UbaStorage.cpp
@@ -1919,7 +1919,7 @@ namespace uba
 
 	void StorageImpl::ReleaseCasEntryNoLock(CasEntry& casEntry, bool dropIfZero)
 	{
-		UBA_ASSERT(casEntry.refCount > 0);
+		UBA_ASSERTF(casEntry.refCount > 0, TC("Reducing ref count on cas entry that already has ref count zero (%s)"), CasKeyString(casEntry.key).str);
 		if (--casEntry.refCount)
 			return;
 		if (!casEntry.dropped || !dropIfZero)
diff --git a/Engine/Source/Programs/UnrealBuildAccelerator/Common/Private/UbaStorageServer.cpp b/Engine/Source/Programs/UnrealBuildAccelerator/Common/Private/UbaStorageServer.cpp
--- a/Engine/Source/Programs/UnrealBuildAccelerator/Common/Private/UbaStorageServer.cpp
+++ b/Engine/Source/Programs/UnrealBuildAccelerator/Common/Private/UbaStorageServer.cpp
@@ -657,7 +657,7 @@ namespace uba
 
 						if (AsCompressed(checkedCasKey, m_storeCompressed) != casKey)
 						{
-							m_logger.Error(TC("FetchBegin failed for cas file %s (%s). File on disk has different cas %s"), CasKeyString(casKey).str, hint.data, CasKeyString(checkedCasKey).str);
+							m_logger.Error(TC("FetchBegin failed for cas file %s (%s) requested by client %u. File on disk has different cas %s"), CasKeyString(casKey).str, hint.data, connectionInfo.GetId(), CasKeyString(checkedCasKey).str);
 							writer.WriteU16(0);
 							return false;
 						}
@@ -665,7 +665,7 @@ namespace uba
 						bool deferCreation = false;
 						if (!AddCasFile(fileNameKey, hint.data, casKey, deferCreation, InvalidValue, 0, nullptr))
 						{
-							m_logger.Error(TC("FetchBegin failed for cas file %s (%s). Can't add cas file to database"), CasKeyString(casKey).str, hint.data);
+							m_logger.Error(TC("FetchBegin failed for cas file %s (%s) requested by client %u. Can't add cas file to database"), CasKeyString(casKey).str, hint.data, connectionInfo.GetId());
 							writer.WriteU16(0);
 							return true;
 						}
