[UBA]
* Fixed so WriteTrackedInputs works properly on mac

diff --git a/Engine/Source/Programs/UnrealBuildAccelerator/Common/Private/UbaProcess.cpp b/Engine/Source/Programs/UnrealBuildAccelerator/Common/Private/UbaProcess.cpp
--- a/Engine/Source/Programs/UnrealBuildAccelerator/Common/Private/UbaProcess.cpp
+++ b/Engine/Source/Programs/UnrealBuildAccelerator/Common/Private/UbaProcess.cpp
@@ -1290,6 +1290,7 @@ namespace uba
 		info.logLineFunc = [](void* userData, const tchar* line, u32 length, LogEntryType type) { ((ProcessImpl*)userData)->LogLine(false, TString(line, length), type); };
 		info.startSuspended = startSuspended;
 		info.rootsHandle = m_startInfo.rootsHandle;
+		info.trackInputs = m_startInfo.trackInputs;
 
 		ProcessImpl* parent = isChild ? this : nullptr;
 		ProcessHandle h = m_session.InternalRunProcess(info, true, parent, true, nullptr);
diff --git a/Engine/Source/Programs/UnrealBuildAccelerator/Common/Private/UbaScheduler.cpp b/Engine/Source/Programs/UnrealBuildAccelerator/Common/Private/UbaScheduler.cpp
--- a/Engine/Source/Programs/UnrealBuildAccelerator/Common/Private/UbaScheduler.cpp
+++ b/Engine/Source/Programs/UnrealBuildAccelerator/Common/Private/UbaScheduler.cpp
@@ -209,7 +209,7 @@ namespace uba
 		}
 
 		if (m_writeTrackedInputs)
-			DirectoryCache().CreateDirectory(m_session.GetLogger(), StringBuffer<>(m_session.GetSessionDir()).Append(TCV("Inputs")).data);
+			DirectoryCache().CreateDirectory(m_session.GetLogger(), StringBuffer<>(m_session.GetSessionDir()).Append(TCV("inputs")).data);
 
 		m_session.SetOuterScheduler(this);
 	}
