[UBA]
* Removed log spam when force dropping cas entries that has refcount higher than 0 (which can happen quite often now when everything keeps cas entries alive through refcounts)

diff --git a/Engine/Source/Programs/UnrealBuildAccelerator/Common/Private/UbaStorage.cpp b/Engine/Source/Programs/UnrealBuildAccelerator/Common/Private/UbaStorage.cpp
--- a/Engine/Source/Programs/UnrealBuildAccelerator/Common/Private/UbaStorage.cpp
+++ b/Engine/Source/Programs/UnrealBuildAccelerator/Common/Private/UbaStorage.cpp
@@ -2247,58 +2247,57 @@ namespace uba
 		lookupLock.Leave();
 	
 		CasEntryWriteLock casEntryLock(casEntry);
+		casEntry.dropped = true;
 
-		if (forceDelete)
+		if (!forceDelete)
+			return true;
+
+		StringBuffer<> casFile;
+		if (!StorageImpl::GetCasFileName(casFile, casKey))
+			return false;
+		if (reduceRefCount)
 		{
-			StringBuffer<> casFile;
-			if (!StorageImpl::GetCasFileName(casFile, casKey))
-				return false;
-			if (reduceRefCount)
-			{
-				UBA_ASSERTF(casEntry.refCount > 0, TC("Reducing refcount of cas %s which already has 0 refcount (%s)"), casFile.data, hint);
-				--casEntry.refCount;
-			}
-			if (casEntry.refCount != 0)
-				return m_logger.Info(TC("Skipped dropping cas %s because refcount was not 0 (%s)"), casFile.data, hint).ToFalse();
+			UBA_ASSERTF(casEntry.refCount > 0, TC("Reducing refcount of cas %s which already has 0 refcount (%s)"), casFile.data, hint);
+			--casEntry.refCount;
+		}
 
-			//UBA_ASSERTF(casEntry.refCount == 0, TC("Force deleting cas %s with refcount %u (%s)"), casFile.data, casEntry.refCount, hint);
-			u64 sizeDeleted = 0;
+		// It could be that this cas is used somewhere and in that case we can't force delete it
+		if (casEntry.refCount != 0)
+			return false;//m_logger.Info(TC("Skipped dropping cas %s because refcount was %u and not 0 (%s)"), casFile.data, u32(casEntry.refCount), hint).ToFalse();
 
-			if (casEntry.memoryHandle.IsValid())
+		//UBA_ASSERTF(casEntry.refCount == 0, TC("Force deleting cas %s with refcount %u (%s)"), casFile.data, casEntry.refCount, hint);
+		u64 sizeDeleted = 0;
+
+		if (casEntry.memoryHandle.IsValid())
+		{
+			if (m_casDataBuffer.IsIndependent(casEntry.memoryHandle))
 			{
-				if (m_casDataBuffer.IsIndependent(casEntry.memoryHandle))
-				{
-					m_casDataBuffer.CloseIndependent(casEntry.memoryHandle, casEntry.memorySize, TC("DropCasFile"));
-					casEntry.memoryHandle = {};
-				}
+				m_casDataBuffer.CloseIndependent(casEntry.memoryHandle, casEntry.memorySize, TC("DropCasFile"));
+				casEntry.memoryHandle = {};
 			}
-			else
-			{
-				if (DeleteFileW(casFile.data) == 0)
-				{
-					u32 lastError = GetLastError();
-					if (lastError != ERROR_FILE_NOT_FOUND && lastError != ERROR_PATH_NOT_FOUND)
-						return m_logger.Error(TC("Failed to drop cas %s (%s) (%s)"), casFile.data, hint, LastErrorToText(lastError).data);
-				}
-				else
-				{
-					m_casDroppedBytes += casEntry.size;
-					m_casDroppedCount++;
-					//m_logger.Debug(TC("Evicted %s from cache (%llukb)"), casFile, casEntry.size/1024);
-					sizeDeleted = casEntry.size;
-				}
-			}
-			casEntry.verified = true;
-			casEntry.exists = false;
-			casEntry.dropped = false;
-			casEntryLock.Leave();
-	
-			CasEntryDeleted(casEntry, sizeDeleted);
 		}
 		else
 		{
-			casEntry.dropped = true;
+			if (DeleteFileW(casFile.data) == 0)
+			{
+				u32 lastError = GetLastError();
+				if (lastError != ERROR_FILE_NOT_FOUND && lastError != ERROR_PATH_NOT_FOUND)
+					return m_logger.Error(TC("Failed to drop cas %s (%s) (%s)"), casFile.data, hint, LastErrorToText(lastError).data);
+			}
+			else
+			{
+				m_casDroppedBytes += casEntry.size;
+				m_casDroppedCount++;
+				//m_logger.Debug(TC("Evicted %s from cache (%llukb)"), casFile, casEntry.size/1024);
+				sizeDeleted = casEntry.size;
+			}
 		}
+		casEntry.verified = true;
+		casEntry.exists = false;
+		casEntry.dropped = false;
+		casEntryLock.Leave();
+	
+		CasEntryDeleted(casEntry, sizeDeleted);
 
 		return true;
 	}
