[UBA]
* Fixed so process reuse is using writeOutputFilesOnFail and IsExitCodeSuccess to decide if received files should be written or not

diff --git a/Engine/Source/Programs/UnrealBuildAccelerator/Common/Private/UbaSessionServer.cpp b/Engine/Source/Programs/UnrealBuildAccelerator/Common/Private/UbaSessionServer.cpp
--- a/Engine/Source/Programs/UnrealBuildAccelerator/Common/Private/UbaSessionServer.cpp
+++ b/Engine/Source/Programs/UnrealBuildAccelerator/Common/Private/UbaSessionServer.cpp
@@ -2007,7 +2007,9 @@ namespace uba
 		auto receivedFiles = std::move(remoteProcess.m_receivedFiles);
 
 		Timer writeFilesTimer; // TODO: This should be added somewhere
-		if (prevExitCode == 0)
+		
+		auto& startInfo = remoteProcess.m_startInfo; // This is kind of wrong with processor reuse but we know the rules is the same
+		if (startInfo.writeOutputFilesOnFail || startInfo.rules->IsExitCodeSuccess(prevExitCode))
 		{
 			if (!WriteReceivedFiles(receivedFiles, remoteProcess, writeFilesTimer, connectionInfo))
 				prevExitCode = UBA_EXIT_CODE(20);
