[UBA]
* Fixed visualizer bug where scrollbars were not properly updated when showing/hiding things using the checkboxes

diff --git a/Engine/Source/Programs/UnrealBuildAccelerator/Visualizer/Private/UbaVisualizer.cpp b/Engine/Source/Programs/UnrealBuildAccelerator/Visualizer/Private/UbaVisualizer.cpp
--- a/Engine/Source/Programs/UnrealBuildAccelerator/Visualizer/Private/UbaVisualizer.cpp
+++ b/Engine/Source/Programs/UnrealBuildAccelerator/Visualizer/Private/UbaVisualizer.cpp
@@ -440,10 +440,8 @@ namespace uba
 			posY += 3;
 		}
 
-		if (pos.y < posY)
-			return;
-
-		outResult.section = 1;
+		if (pos.y >= posY)
+			outResult.section = 1;
 
 		if (m_config.showActiveProcessGraph)
 		{
@@ -503,10 +501,8 @@ namespace uba
 				return;
 		}
 
-		if (pos.y < posY)
-			return;
-
-		outResult.section = 2;
+		if (pos.y >= posY)
+			outResult.section = 2;
 
 		SessionRec sortedSessions[1024];
 		Populate(sortedSessions, m_traceView, m_config.SortActiveRemoteSessions);
diff --git a/Engine/Source/Programs/UnrealBuildAccelerator/Visualizer/Private/Windows/UbaVisualizerWin.cpp b/Engine/Source/Programs/UnrealBuildAccelerator/Visualizer/Private/Windows/UbaVisualizerWin.cpp
--- a/Engine/Source/Programs/UnrealBuildAccelerator/Visualizer/Private/Windows/UbaVisualizerWin.cpp
+++ b/Engine/Source/Programs/UnrealBuildAccelerator/Visualizer/Private/Windows/UbaVisualizerWin.cpp
@@ -1055,9 +1055,10 @@ namespace uba
 				bool* values = &m_config.showProgress;
 				values[m_buttonSelected] = !values[m_buttonSelected];
 
-				VisualizerBackendGDI::PaintContextGDI context((VisualizerBackendGDI&)m_backend, NULL, NULL, { m_clientWidth, m_clientHeight });
+				PaintContext* context = m_backend.CreateContext(m_clientWidth, m_clientHeight);
 				HitTestResult res;
-				HitTest(res, context, { -1, -1 });
+				HitTest(res, *context, { -1, -1 });
+				m_backend.DeleteContext(context);
 				UpdateScrollbars(true);
 				Redraw(false);
 			}
