[UBA]
* Changed so WSAPoll/poll always only polls with 10 seconds timeout max. Reason for this is because shutting down sockets does not trigger the poll which means that sometimes a shutdown can happen without the receiving thread knowing and taking 10 minutes for it to actually realize

diff --git a/Engine/Source/Programs/UnrealBuildAccelerator/Common/Private/UbaNetworkBackendTcp.cpp b/Engine/Source/Programs/UnrealBuildAccelerator/Common/Private/UbaNetworkBackendTcp.cpp
--- a/Engine/Source/Programs/UnrealBuildAccelerator/Common/Private/UbaNetworkBackendTcp.cpp
+++ b/Engine/Source/Programs/UnrealBuildAccelerator/Common/Private/UbaNetworkBackendTcp.cpp
@@ -124,6 +124,7 @@ namespace uba
 
 		Logger& logger;
 		Atomic<Socket> socket;
+		Atomic<bool> shutdownCalled = false;
 
 		Event ready;
 		Guid uid;
@@ -351,6 +352,7 @@ namespace uba
 		if (conn.socket == InvalidSocket)
 			return;
 		ShutdownSocket(conn.logger, conn.socket, TC("Shutdown"));
+		conn.shutdownCalled = true;
 	}
 
 	bool NetworkBackendTcp::Send(Logger& logger, void* connection, const void* data, u32 dataSize, SendContext& sendContext, const tchar* sendHint)
@@ -1084,22 +1086,24 @@ namespace uba
 			return false;
 		}
 
-		u32 timeoutMs = 40 * 1000; // 40 seconds timeout to begin with
+		u32 timeoutMs = 10 * 1000;
 
 #if PLATFORM_WINDOWS
 		if (connection.wsaEvent != WSA_INVALID_EVENT)
 		{
-			for (u32 i = 0; i != 2; ++i)
+			u32 timeoutCounter = 0;
+			while (true)
 			{
 				DWORD wr = WSAWaitForMultipleEvents(1, &connection.wsaEvent, FALSE, timeoutMs, FALSE);
 				if (wr == WSA_WAIT_TIMEOUT)
 				{
-					if (i == 0)
-					{
+					if (connection.shutdownCalled)
+						return false;
+					++timeoutCounter;
+					if (timeoutCounter == 4)
 						logger.Info(TC("SendSocket - WSAWaitForMultipleEvents returned timeout after 40 seconds for socket %i (%s)"), socket, hint);
-						timeoutMs = DefaultNetworkSendTimeoutSeconds * 1000; // 10 minutes send timeout
+					if (timeoutCounter < 6*10) // Loop for 10 minutes
 						continue;
-					}
 					logger.Info(TC("SendSocket - WSAWaitForMultipleEvents returned timeout after 10 minutes for socket %i (%s)"), socket, hint);
 					return false;
 				}
@@ -1125,18 +1129,20 @@ namespace uba
 		else
 #endif
 		{
-			for (u32 i = 0; i != 2; ++i)
+			u32 timeoutCounter = 0;
+			while (true)
 			{
 				int revents = 0;
 				int res = SocketPoll(socket, PollOUT, timeoutMs, &revents);
 				if (!res)
 				{
-					if (i == 0)
-					{
+					if (connection.shutdownCalled)
+						return false;
+					++timeoutCounter;
+					if (timeoutCounter == 4)
 						logger.Info(TC("SendSocket - WSAPoll returned timeout after 40 seconds for socket %i (%s)"), socket, hint);
-						timeoutMs = DefaultNetworkSendTimeoutSeconds * 1000; // 10 minutes send timeout
+					if (timeoutCounter < 6*10) // Loop for 10 minutes
 						continue;
-					}
 					logger.Info(TC("SendSocket - WSAPoll returned timeout after 10 minutes for socket %i (%s)"), socket, hint);
 					return false;
 				}
@@ -1459,33 +1465,54 @@ namespace uba
 				if (!timeoutMs)
 					timeoutMs = DefaultNetworkReceiveTimeoutSeconds * 1000; // 10 minutes time out
 
-				int revents = 0;
-				int res = SocketPoll(socket, PollIN, timeoutMs, &revents);
-				if (!res)
+				// Since a socket shutdown does not signal the poll to stop we need to timeout earlier and check shutdownCalled and then poll again
+				// We timeout 10 seconds at the time until we reach the "real" timeout
+				u32 timeoutSliceMs = Min(timeoutMs, 10u*1000);
+				u64 startTime = GetTime();
+				while (true)
 				{
-					if (connection.recvTimeoutCallback)
+					int revents = 0;
+					int res = SocketPoll(socket, PollIN, timeoutSliceMs, &revents);
+					if (!res)
 					{
-						if (connection.recvTimeoutCallback(connection.recvTimeoutContext, timeoutMs, connection.recvHint, hint))
+						if (connection.shutdownCalled)
+							return false;
+
+						u64 now = GetTime();
+						u32 timePassedMs = u32(TimeToMs(now - startTime));
+						if (timePassedMs < timeoutMs)
+						{
+							// If less than 10s left to timeout we just timeout to that time
+							timeoutSliceMs = Min(timeoutMs - timePassedMs + 1000U, timeoutSliceMs);
 							continue;
+						}
+
+						if (connection.recvTimeoutCallback)
+						{
+							if (connection.recvTimeoutCallback(connection.recvTimeoutContext, timeoutMs, connection.recvHint, hint))
+								break;
+							return false;
+						}
+						connection.logger.Info(TC("RecvSocket - WSAPoll returned timeout for socket %i and connection %s after %s (%s%s)"), socket, GuidToString(connection.uid).str, TimeToText(MsToTime(timeoutMs)).str, connection.recvHint, hint);
 						return false;
 					}
-					connection.logger.Info(TC("RecvSocket - WSAPoll returned timeout for socket %i and connection %s after %s (%s%s)"), socket, GuidToString(connection.uid).str, TimeToText(MsToTime(timeoutMs)).str, connection.recvHint, hint);
-					return false;
-				}
-				if (res < 0)
-				{
-					#if UBA_LOG_SOCKET_ERRORS
-					connection.logger.Info(TC("RecvSocket - WSAPoll returned an err-or for socket %i and connection %s: %s (%s%s)"), socket, GuidToString(connection.uid).str, LastErrorToText(WSAGetLastError()).data, connection.recvHint, hint);
-					#endif
-					return false;
-				}
-				if ((revents & PollIN) == 0)
-				{
-					#if UBA_LOG_SOCKET_ERRORS
-					connection.logger.Info(TC("RecvSocket - WSAPoll revents contained error for socket %i and connection %s: %s (%s%s)"), socket, GuidToString(connection.uid).str, LastErrorToText(WSAGetLastError()).data, connection.recvHint, hint);
-					#endif
-					return false;
+					if (res < 0)
+					{
+						#if UBA_LOG_SOCKET_ERRORS
+						connection.logger.Info(TC("RecvSocket - WSAPoll returned an err-or for socket %i and connection %s: %s (%s%s)"), socket, GuidToString(connection.uid).str, LastErrorToText(WSAGetLastError()).data, connection.recvHint, hint);
+						#endif
+						return false;
+					}
+					if ((revents & PollIN) == 0)
+					{
+						#if UBA_LOG_SOCKET_ERRORS
+						connection.logger.Info(TC("RecvSocket - WSAPoll revents contained error for socket %i and connection %s: %s (%s%s)"), socket, GuidToString(connection.uid).str, LastErrorToText(WSAGetLastError()).data, connection.recvHint, hint);
+						#endif
+						return false;
+					}
+					break;
 				}
+
 				continue;
 			}
 
