[UBA]
* Most of visualizer rendering for mac implemented

diff --git a/Engine/Source/Programs/UnrealBuildAccelerator/Visualizer/Private/Mac/UbaVisualizerApp.mm b/Engine/Source/Programs/UnrealBuildAccelerator/Visualizer/Private/Mac/UbaVisualizerApp.mm
--- a/Engine/Source/Programs/UnrealBuildAccelerator/Visualizer/Private/Mac/UbaVisualizerApp.mm
+++ b/Engine/Source/Programs/UnrealBuildAccelerator/Visualizer/Private/Mac/UbaVisualizerApp.mm
@@ -16,162 +16,323 @@
 
 namespace uba
 {
-    class VisualizerBackendMac : public VisualizerBackend
-    {
-    public:
-        struct PaintContextMac : PaintContext
-        {
+	class VisualizerBackendMac : public VisualizerBackend
+	{
+	public:
+		struct PaintContextMac : PaintContext
+		{
 			virtual SIZE GetClientSize() override
-            {
-                return clientSize;
-            }
-			virtual void SetPen(HPEN pen) override {}
-			virtual void SetBrush(HBRUSH b) override {}
-			virtual void SetFont(const Font& font, bool underlined = false) override { activeFont = font; }
-			virtual void SetTextColor(COLORREF color) override {}
-
-            int textOffsetX = 0;
-            int textOffsetY = 0;
-            NSRect textRect;
-
-			virtual void Text(int x, int y, const RECT* lprect, StringView str) override
-            {         
-                NSString *text = [NSString stringWithUTF8String:str.data];
-
-                if (!lprect)
-                {
-                    NSDictionary *attrs = @{
-                        NSFontAttributeName: [NSFont systemFontOfSize:10.0],
-                        NSForegroundColorAttributeName: [NSColor blackColor],
-                    };
-                    NSPoint pos = NSMakePoint(x + textOffsetX, y + textOffsetY);
-                    [text drawAtPoint:pos withAttributes:attrs];
-                }
-                else
-                {
-                    NSMutableParagraphStyle *style = [[NSMutableParagraphStyle alloc] init];
-                    style.lineBreakMode = NSLineBreakByTruncatingTail; // ... at the end
-                    style.alignment = NSTextAlignmentLeft;
-
-                    NSDictionary *attrs = @{
-                        NSFontAttributeName: [NSFont systemFontOfSize:10.0],
-                        NSForegroundColorAttributeName: [NSColor blackColor],
-                        NSParagraphStyleAttributeName: style
-                    };
-                    NSRect textRect = NSMakeRect(x + textOffsetX, y + textOffsetY, lprect->right - lprect->left, lprect->bottom - lprect->right); // narrow rect
-
-                    NSAttributedString *attrText = [[NSAttributedString alloc] initWithString:text attributes:attrs];
-                    [attrText drawWithRect:textRect
-                                options:NSStringDrawingUsesLineFragmentOrigin
-                                context:nil];
-                }
-            }
-
-			virtual COLORREF GetTextColor() { return {}; }
-			virtual Font& GetFont() { return activeFont; }
-			virtual void FillRect(const RECT& r, HBRUSH b)
-            {
-                [[NSColor greenColor] setFill];
-                NSRect rect = NSMakeRect(r.left, r.top, r.right - r.left, r.bottom - r.top);
-                NSRectFill(rect);
-            }
-			virtual void Rectangle(int left, int top, int right, int bottom) {}
-			virtual void Line(int fromX, int fromY, int toX, int toY) {}
-			virtual void Polyline(u32 cacheId, const Vector<POINT>& polyLine) {}
-			virtual void PolyPolyline(const POINT* points, const u32* lines, u32 lineCount) {}
-			virtual void DrawCenteredText(const StringBuffer<>* rows, u32 rowCount, HBRUSH backgroundBrush) {}
-			virtual SIZE GetTextSize(StringView str)
-            {
-                return { int(str.count)*8, 10 };
-            }
-			virtual void DrawCached(TraceView::BitmapCache& process, RECT rect, int offsetY, int textHeight, Font& processFont, RECT clippingRect, const Function<void(PaintContext& textCacheContext)>& func)
-            {
-                textOffsetX = rect.left;
-                textOffsetY = rect.top;
-                func(*this);
-                textOffsetX = 0;
-                textOffsetY = 0;
-            }
-
-			virtual POINT GetCursorClientPos() { return {}; }
-
-            SIZE clientSize;
+ 			{
+				return clientSize;
+			}
+			virtual void SetPen(PenHandle pen) override
+			{
+				CGContextSetStrokeColorWithColor(ctx, (CGColorRef)pen);
+			}
+			virtual void SetBrush(BrushHandle b) override
+			{
+				CGContextSetFillColorWithColor(ctx, (CGColorRef)b);
+			}
+			virtual void SetFont(const Font& font, bool underlined = false) override
+			{
+				activeFont = font;
+			}
+			virtual void SetTextColor(COLORREF color) override
+			{
+				activeTextColor = color;
+				if (activeTextColorRef)
+					CFRelease(activeTextColorRef);
+				float r = float(GET_RED(color))/255.0f;
+				float g = float(GET_GREEN(color))/255.0f;
+				float b = float(GET_BLUE(color))/255.0f;
+				CGFloat comps[4] = { r, g, b, 1.0 }; // black
+				activeTextColorRef = CGColorCreate(owner.colorSpaceRef, comps);
+			}
+			virtual void Text(int x, int y, StringView str, const RECT* clipRect) override
+			{
+				// --- Convert UTF-8 to CFString ---
+				CFStringRef cfStr = CFStringCreateWithBytes(
+					kCFAllocatorDefault,
+					reinterpret_cast<const UInt8*>(str.data),
+					static_cast<CFIndex>(str.count),
+					kCFStringEncodingUTF8,
+					kCFStringEncodingUTF8 == kCFStringEncodingUTF8 ? false : false // just "false"
+				);
+				if (!cfStr)
+					return;
+
+				y -= 2;
+
+				const void* keys[]   = { kCTFontAttributeName,			kCTForegroundColorAttributeName };
+				const void* values[] = { (CTFontRef)activeFont.handle,	activeTextColorRef };
+				CFDictionaryRef attrs = CFDictionaryCreate(kCFAllocatorDefault, keys, values, 2, &kCFTypeDictionaryKeyCallBacks, &kCFTypeDictionaryValueCallBacks);
+
+				CFAttributedStringRef attrStr = CFAttributedStringCreate(kCFAllocatorDefault, cfStr, attrs);
+				CTLineRef line = CTLineCreateWithAttributedString(attrStr);
+
+				// Logical top-left position in your coordinate system
+				CGFloat topLeftX = (CGFloat)x;
+				CGFloat topLeftY = (CGFloat)y;
+
+				CGContextSaveGState(ctx);
+
+				// Optional clip in your normal (0,0 top-left, y-down) coordinates
+				CGFloat maxWidth = 0.0f;
+				if (clipRect)
+				{
+					CGFloat clipW = (CGFloat)(clipRect->right  - clipRect->left);
+					CGFloat clipH = (CGFloat)(clipRect->bottom - clipRect->top);
+					CGRect cgClip = CGRectMake(topLeftX, topLeftY, clipW, clipH);
+					CGContextClipToRect(ctx, cgClip);
+
+					maxWidth = clipW;
+				}
+
+				// Start with an identity text matrix
+				CGContextSetTextMatrix(ctx, CGAffineTransformIdentity);
+
+				// Now make a local transform so Core Text sees y-up, but
+				// the top of the text is at (topLeftX, topLeftY) in your coords.
+				//
+				// 1) Move origin to the *baseline*: (topLeftX, topLeftY + activeFont.height)
+				// 2) Flip vertically so Core Text sees y-up.
+				CGContextTranslateCTM(ctx, topLeftX, topLeftY + activeFont.height);
+				CGContextScaleCTM(ctx, 1.0, -1.0);
+
+				// Baseline at (0,0) in this local text space
+				CGContextSetTextPosition(ctx, 0.0, 0.0);
+
+				if (!clipRect)
+				{
+					CTLineDraw(line, ctx);
+				}
+				else
+				{
+					// Tail truncation to clip width
+					CFStringRef ellipsisStr = CFSTR("\u2026");
+					CFAttributedStringRef ellAttrStr = CFAttributedStringCreate(kCFAllocatorDefault, ellipsisStr, attrs);
+					CTLineRef ellLine = CTLineCreateWithAttributedString(ellAttrStr);
+					CTLineRef truncLine = CTLineCreateTruncatedLine(line, maxWidth, kCTLineTruncationEnd, ellLine);
+					CTLineDraw(truncLine ? truncLine : line, ctx);
+					CFRelease(ellLine);
+					CFRelease(ellAttrStr);
+					if (truncLine)
+						CFRelease(truncLine);
+				}
+
+				CGContextRestoreGState(ctx);
+
+				CFRelease(line);
+				CFRelease(attrStr);
+				CFRelease(attrs);
+				CFRelease(cfStr);
+			}
+
+			virtual COLORREF GetTextColor() override
+			{
+				return activeTextColor;
+			}
+			virtual Font& GetFont() override
+			{
+				return activeFont;
+			}
+			virtual void FillRect(const RECT& r, BrushHandle b) override
+			{
+				CGContextSetFillColorWithColor(ctx, (CGColorRef)b);
+				CGRect rect = CGRectMake(r.left, r.top, r.right - r.left, r.bottom - r.top);
+				CGContextFillRect(ctx, rect);
+			}
+			virtual void Rectangle(int left, int top, int right, int bottom) override
+			{
+				CGRect rect = CGRectMake(left, top, right - left, bottom - top);
+				CGContextSetLineWidth(ctx, 1.0);
+				CGContextStrokeRect(ctx, rect);
+			}
+			virtual void Line(int fromX, int fromY, int toX, int toY) override
+			{
+				CGContextSetLineWidth(ctx, 1.0);
+				CGContextMoveToPoint(ctx, fromX, fromY);
+				CGContextAddLineToPoint(ctx, toX, toY);
+				CGContextStrokePath(ctx);
+			}
+			virtual void Polyline(u32 cacheId, const Vector<POINT>& polyLine) override
+			{
+				Vector<CGPoint> points;
+				points.resize(polyLine.size());
+				u32 i = 0;
+				for (auto& p : polyLine)
+					points[i++] = CGPoint(p.x, p.y);
+				CGContextBeginPath(ctx);
+				CGContextAddLines(ctx, points.data(), points.size());
+				CGContextStrokePath(ctx);
+			}
+			virtual void PolyPolyline(const POINT* points, const u32* lines, u32 lineCount) override
+			{
+				Vector<CGPoint> p;
+				for (u32 i=0; i!=lineCount; ++i)
+				{
+					u32 pointCount = lines[i];
+					p.resize(pointCount);
+					for (u32 j=0; j!=pointCount; ++j, ++points)
+						p[j] = CGPoint(points->x, points->y);
+					CGContextBeginPath(ctx);
+					CGContextAddLines(ctx, p.data(), p.size());
+					CGContextStrokePath(ctx);
+				}
+			}
+			virtual SIZE GetTextSize(StringView str) override
+			{
+				return { int(str.count)*8, 10 };
+			}
+			virtual void DrawCached(TraceView::BitmapCache& process, RECT rect, int offsetY, int textHeight, Font& processFont, RECT clippingRect, const Function<void(PaintContext& textCacheContext)>& func) override
+			{
+				CGContextSaveGState(ctx);
+				CGContextTranslateCTM(ctx, rect.left, rect.top);
+				func(*this);
+				CGContextRestoreGState(ctx);
+			}
+
+			virtual POINT GetCursorClientPos() override { return {}; }
+
+			PaintContextMac(VisualizerBackendMac& o) : owner(o)
+			{
+				ctx = [[NSGraphicsContext currentContext] CGContext];
+			}
+
+			VisualizerBackendMac& owner;
+			CGContextRef ctx;
+			SIZE clientSize;
 			Font activeFont;
-        };
-
-		virtual void ResetCaches() override {}
-        virtual HBRUSH CreateSolidBrush(COLORREF color) override { return 0; }
-    };
-
-    class VisualizerMac : public Visualizer
-    {
-    public:
-        VisualizerMac(VisualizerConfig& config, Logger& logger, VisualizerBackend& backend) : Visualizer(config, logger, backend)
-        {
-            char path[512];
-            GetFullPathNameW("~/Trace.uba", 512, path, nullptr);
-
-            m_fileName.Append(path);
-            if (true)
-            {
-                m_trace.ReadFile(m_traceView, path, true);
-
-                m_startTime = GetTime();
-                
-                bool changed = true;
-                u64 time = 0;
-                m_traceView.finished = false;
-                m_trace.UpdateReadFile(m_traceView, time, changed);
-                m_pauseStart = m_startTime + time;
-                m_pauseTime = m_startTime - m_pauseStart;
-                m_autoScroll = true;
-                m_replay = 1;
-            }
-            else
-                m_trace.ReadFile(m_traceView, path, false);
-
-		    m_defaultFont.height = 14;
-		    m_processFont.height = 14;
-		    m_timelineFont.height = 14;
-    		m_popupFont.height = 14;
-        }
+			COLORREF activeTextColor;
+			CGColorRef activeTextColorRef = nullptr;
+		};
+
+		VisualizerBackendMac()
+		{
+			colorSpaceRef = CGColorSpaceCreateDeviceRGB();
+		}
+
+		virtual const tchar* GetName() override
+		{
+			return "";
+		}
+		virtual void Init(void* windowHandle) override
+		{
+		}
+		virtual void Reset() override
+		{
+			for (auto c : colors)
+				CFRelease(c);
+			colors.clear();
+		}
+		virtual void Resize(int width, int height) override
+		{
+		}
+		virtual void Paint(const Function<void(PaintContext& context, RECT rect)>& func) override
+		{
+		}
+		virtual PaintContext* CreateContext(int clientWidth, int clientHeight) override { return nullptr; }
+		virtual void DeleteContext(PaintContext*) override {}
+		virtual BrushHandle CreateSolidBrush(COLORREF color) override
+		{
+				float r = float(GET_RED(color))/255.0f;
+				float g = float(GET_GREEN(color))/255.0f;
+				float b = float(GET_BLUE(color))/255.0f;
+				CGFloat comps[4] = { r, g, b, 1.0 }; // black
+				CGColorRef ref = CGColorCreate(colorSpaceRef, comps);
+				colors.push_back(ref);
+				return ref;
+		}
+		virtual PenHandle CreatePen(COLORREF color) override
+		{
+			return CreateSolidBrush(color);
+		}
+		virtual FontHandle CreateFont(const tchar* font, int height, bool createUnderline) override
+		{
+			const CGFloat fontSize = height;
+			CFStringRef fontStr = CFStringCreateWithCString(kCFAllocatorDefault, font, kCFStringEncodingUTF8);
+			auto res = (FontHandle)CTFontCreateWithName(fontStr, fontSize-4, nullptr);
+			UBA_ASSERT(res);
+			CFRelease(fontStr);
+			return res;
+		}
+		virtual void DeleteFont(FontHandle font) override
+		{
+			CFRelease((CTFontRef)font);
+		}
+
+		CGColorSpaceRef colorSpaceRef;
+		List<CGColorRef> colors;
+	};
+
+	class VisualizerMac : public Visualizer
+	{
+	public:
+		VisualizerMac(VisualizerConfig& config, Logger& logger, VisualizerBackend& backend) : Visualizer(config, logger, backend)
+		{
+			char path[512];
+			GetFullPathNameW("~/Trace.uba", 512, path, nullptr);
+
+			m_fileName.Append(path);
+			if (true)
+			{
+				m_trace.ReadFile(m_traceView, path, true);
+
+				m_startTime = GetTime();
+
+				bool changed = true;
+				u64 time = 0;
+				m_traceView.finished = false;
+				m_trace.UpdateReadFile(m_traceView, time, changed);
+				m_pauseStart = m_startTime + time;
+				m_pauseTime = m_startTime - m_pauseStart;
+				m_autoScroll = true;
+				m_replay = 1;
+			}
+			else
+				m_trace.ReadFile(m_traceView, path, false);
+
+			UpdateDefaultFont();
+			UpdateProcessFont();
+			m_popupFont.handle = m_backend.CreateFont(TC("Consolas"), 16, false);
+			m_popupFont.height = 14;
+
+			InitBrushes();
+		}
 		virtual void PaintEmptyScreen(PaintContext& context) override
-        {
-            [[NSColor blackColor] setFill];
-            NSRect rect = NSMakeRect(50, 50, 200, 100);
-            NSRectFill(rect);
-        }
+		{
+			[[NSColor blackColor] setFill];
+			NSRect rect = NSMakeRect(50, 50, 200, 100);
+			NSRectFill(rect);
+		}
 		virtual void FinishIncompatible()
-        {
-        }
-
-        void Paint(SIZE size)
-        {
-            m_clientWidth = size.cx;
-            m_clientHeight = size.cy;
-
-            uba::VisualizerBackendMac::PaintContextMac context;
-            context.clientSize = size;
-            RECT r { 0, 0, size.cx, size.cy };
-            PaintAll(context, r);
-        }
-
-        void Update()
-        {
-            UpdateTrace();
-        }
-    };
-
-	void Visualizer::UpdateFont(Font& font, int height, bool createUnderline)
-    {
-        font.height = height;
-    }
+		{
+		}
+
+		void Paint(SIZE size)
+		{
+			m_clientWidth = size.cx;
+			m_clientHeight = size.cy;
+
+			uba::VisualizerBackendMac::PaintContextMac context((VisualizerBackendMac&)m_backend);
+			context.clientSize = size;
+			RECT r { 0, 0, size.cx, size.cy };
+
+			context.FillRect(r, m_backgroundBrush);
+			PaintAll(context, r);
+		}
+
+		void Update()
+		{
+			UpdateTrace();
+		}
+	};
 }
 
 @interface MyCanvasView : NSView
 {
-    NSTimeInterval _time;
-    uba::VisualizerMac* _visualizer;
+	NSTimeInterval _time;
+	uba::VisualizerMac* _visualizer;
 }
 - (void)tick:(NSTimer *)timer;
 @end
@@ -180,65 +341,50 @@ @implementation MyCanvasView
 
 
 - (instancetype)initWithFrame:(NSRect)frame {
-    self = [super initWithFrame:frame];
-    if (self) {
-        _time = 0.0;
-
-        uba::VisualizerConfig config("~/UbaVisualizerConfig.toml");
-        config.boxHeight = 14;
-        config.horizontalScaleValue = 1.0f;
-        auto logger = new uba::LoggerWithWriter(uba::g_consoleLogWriter);
-        auto backend = new uba::VisualizerBackendMac;
-
-        _visualizer = new uba::VisualizerMac(config, *logger, *backend);
-
-        // 60 FPS timer
-        [NSTimer scheduledTimerWithTimeInterval:(1.0/60.0)
-                                         target:self
-                                       selector:@selector(tick:)
-                                       userInfo:nil
-                                        repeats:YES];
-    }
-    return self;
+	self = [super initWithFrame:frame];
+	if (self) {
+		_time = 0.0;
+
+		uba::VisualizerConfig config("~/UbaVisualizerConfig.toml");
+		config.boxHeight = 14;
+		config.horizontalScaleValue = 1.0f;
+		config.DarkMode = true;
+		auto logger = new uba::LoggerWithWriter(uba::g_consoleLogWriter);
+		auto backend = new uba::VisualizerBackendMac;
+
+		_visualizer = new uba::VisualizerMac(config, *logger, *backend);
+
+		// 60 FPS timer
+		[NSTimer scheduledTimerWithTimeInterval:(1.0/60.0)
+										 target:self
+									   selector:@selector(tick:)
+									   userInfo:nil
+										repeats:YES];
+	}
+	return self;
 }
 
 - (void)tick:(NSTimer *)timer {
-    _time += 1.0 / 60.0;
+	_time += 1.0 / 60.0;
 
-    _visualizer->Update();
+	_visualizer->Update();
 
-    // Request a redraw
-    [self setNeedsDisplay:YES];
+	// Request a redraw
+	[self setNeedsDisplay:YES];
 }
 
 
 
 - (BOOL)isFlipped {
-    // Use top-left origin like many 2D APIs (optional)
-    return YES;
+	// Use top-left origin like many 2D APIs (optional)
+	return YES;
 }
 
 - (void)drawRect:(NSRect)dirtyRect {
    [super drawRect:dirtyRect];
-
-    NSRect bounds = self.bounds;
-
-    [[NSColor whiteColor] setFill];
-    NSRectFill(bounds);
-/*
-    // Example: animate the line using _time
-    [[NSColor blackColor] setStroke];
-    NSBezierPath *path = [NSBezierPath bezierPath];
-
-    CGFloat x = 10 + 50 * sin(_time * 2.0); // wiggle a bit
-
-    [path moveToPoint:NSMakePoint(x, 10)];
-    [path lineToPoint:NSMakePoint(bounds.size.width - x,
-                                  bounds.size.height - 10)];
-    [path stroke];
-*/
-    uba::SIZE size = { (int)bounds.size.width, (int)bounds.size.height };
-    _visualizer->Paint(size);
+	NSRect bounds = self.bounds;
+	uba::SIZE size = { (int)bounds.size.width, (int)bounds.size.height };
+	_visualizer->Paint(size);
 }
 
 @end
@@ -248,62 +394,62 @@ @interface AppDelegate : NSObject <NSWindowDelegate>
 
 @implementation AppDelegate
 - (void)windowWillClose:(NSNotification *)notification {
-    [NSApp terminate:nil];
+	[NSApp terminate:nil];
 }
 @end
 
 int main(int argc, const char * argv[]) {
-    @autoreleasepool {
-        [NSApplication sharedApplication];
-        [NSApp setActivationPolicy:NSApplicationActivationPolicyRegular];
-
-        // --- Minimal menu bar with "Quit" (Cmd+Q) ---
-        NSMenu *menubar = [[NSMenu alloc] init];
-        NSMenuItem *appMenuItem = [[NSMenuItem alloc] init];
-        [menubar addItem:appMenuItem];
-        [NSApp setMainMenu:menubar];
-
-        NSMenu *appMenu = [[NSMenu alloc] initWithTitle:@""];
-        NSString *appName = [[NSProcessInfo processInfo] processName];
-        NSMenuItem *quitItem =
-            [[NSMenuItem alloc] initWithTitle:[NSString stringWithFormat:@"Quit %@", appName]
-                                       action:@selector(terminate:)
-                                keyEquivalent:@"q"];
-        // Make sure it's Cmd+Q, not just 'q'
-        [quitItem setKeyEquivalentModifierMask:NSEventModifierFlagCommand];
-
-        [appMenu addItem:quitItem];
-        [appMenuItem setSubmenu:appMenu];
-        // --- end menu setup ---
-
-        // Create the window
-        NSRect frame = NSMakeRect(200, 200, 800, 600);
-        NSWindowStyleMask style =
-            NSWindowStyleMaskTitled |
-            NSWindowStyleMaskClosable |
-            NSWindowStyleMaskResizable |
-            NSWindowStyleMaskMiniaturizable;
-
-        NSWindow *window = [[NSWindow alloc]
-                            initWithContentRect:frame
-                            styleMask:style
-                            backing:NSBackingStoreBuffered
-                            defer:NO];
-
-        AppDelegate *delegate = [[AppDelegate alloc] init];
-        [window setDelegate:delegate];
-        [window setTitle:@"UbaVisualizer"];
-
-        // Attach our canvas view
-        MyCanvasView *view = [[MyCanvasView alloc] initWithFrame:frame];
-        [window setContentView:view];
-        [window makeKeyAndOrderFront:nil];
-
-        // Make sure app shows in Dock and gets focus
-        [NSApp activateIgnoringOtherApps:YES];
-
-        // Enter the Cocoa event loop
-        [NSApp run];
-    }
-    return 0;
-}
\ No newline at end of file
+	@autoreleasepool {
+		[NSApplication sharedApplication];
+		[NSApp setActivationPolicy:NSApplicationActivationPolicyRegular];
+
+		// --- Minimal menu bar with "Quit" (Cmd+Q) ---
+		NSMenu *menubar = [[NSMenu alloc] init];
+		NSMenuItem *appMenuItem = [[NSMenuItem alloc] init];
+		[menubar addItem:appMenuItem];
+		[NSApp setMainMenu:menubar];
+
+		NSMenu *appMenu = [[NSMenu alloc] initWithTitle:@""];
+		NSString *appName = [[NSProcessInfo processInfo] processName];
+		NSMenuItem *quitItem =
+			[[NSMenuItem alloc] initWithTitle:[NSString stringWithFormat:@"Quit %@", appName]
+									   action:@selector(terminate:)
+								keyEquivalent:@"q"];
+		// Make sure it's Cmd+Q, not just 'q'
+		[quitItem setKeyEquivalentModifierMask:NSEventModifierFlagCommand];
+
+		[appMenu addItem:quitItem];
+		[appMenuItem setSubmenu:appMenu];
+		// --- end menu setup ---
+
+		// Create the window
+		NSRect frame = NSMakeRect(200, 200, 800, 600);
+		NSWindowStyleMask style =
+			NSWindowStyleMaskTitled |
+			NSWindowStyleMaskClosable |
+			NSWindowStyleMaskResizable |
+			NSWindowStyleMaskMiniaturizable;
+
+		NSWindow *window = [[NSWindow alloc]
+							initWithContentRect:frame
+							styleMask:style
+							backing:NSBackingStoreBuffered
+							defer:NO];
+
+		AppDelegate *delegate = [[AppDelegate alloc] init];
+		[window setDelegate:delegate];
+		[window setTitle:@"UbaVisualizer"];
+
+		// Attach our canvas view
+		MyCanvasView *view = [[MyCanvasView alloc] initWithFrame:frame];
+		[window setContentView:view];
+		[window makeKeyAndOrderFront:nil];
+
+		// Make sure app shows in Dock and gets focus
+		[NSApp activateIgnoringOtherApps:YES];
+
+		// Enter the Cocoa event loop
+		[NSApp run];
+	}
+	return 0;
+}
diff --git a/Engine/Source/Programs/UnrealBuildAccelerator/Visualizer/Private/UbaVisualizer.cpp b/Engine/Source/Programs/UnrealBuildAccelerator/Visualizer/Private/UbaVisualizer.cpp
--- a/Engine/Source/Programs/UnrealBuildAccelerator/Visualizer/Private/UbaVisualizer.cpp
+++ b/Engine/Source/Programs/UnrealBuildAccelerator/Visualizer/Private/UbaVisualizer.cpp
@@ -1381,9 +1381,9 @@ namespace uba
 								context.DrawCached(process, rect, offsetY, textHeight, m_processFont, progressRect, [&](PaintContext& textCacheContext)
 									{
 										RECT rect2{ 0, int(process.bitmapOffset), 256, int(process.bitmapOffset) + m_processFont.height };
-										RECT rect3{ 0, int(process.bitmapOffset), processWidth, int(process.bitmapOffset) + m_processFont.height };
-										if (!done)
-											rect3.right = 256;
+										RECT rect3 = rect2;
+										if (done)
+											rect3.right = processWidth;
 
 										PaintProcessRect(textCacheContext, process, rect3, rect2, selected, true);
 
@@ -2821,5 +2821,82 @@ namespace uba
 		return false;
 	}
 
+	void Visualizer::UpdateFont(Font& font, int height, bool createUnderline)
+	{
+		font.height = height;
+
+		int fh = height;
+		font.offset = 0;
+		if (height <= 13)
+		{
+			++fh;
+			--font.offset;
+		}
+		if (height <= 11)
+			++fh;
+		if (height <= 9)
+			++fh;
+		if (height <= 8)
+			++fh;
+		if (height <= 6)
+		{
+			++fh;
+			--font.offset;
+		}
+		if (height <= 4)
+			--font.offset;
+
+		if (font.handle)
+			m_backend.DeleteFont(font.handle);
+		font.handle = 0;
+		if (font.handleUnderlined)
+			m_backend.DeleteFont(font.handleUnderlined);
+		font.handleUnderlined = 0;
+
+		font.handle = m_backend.CreateFont(m_config.fontName.c_str(), fh, false);
+		if (createUnderline)
+			font.handleUnderlined = m_backend.CreateFont(m_config.fontName.c_str(), fh, true);
+	}
+
+	void Visualizer::UpdateDefaultFont()
+	{
+		UpdateFont(m_defaultFont, m_config.fontSize, true);
+		m_sessionStepY = m_defaultFont.height + 4;
+		m_timelineFont = m_defaultFont;
+	}
+
+	void Visualizer::UpdateProcessFont()
+	{
+		m_zoomValue = 1.0f + float(m_config.boxHeight) / 30.0f;
+		int fontHeight = Max(int(m_config.boxHeight) - 2, 1);
+		UpdateFont(m_processFont, fontHeight, false);
+		m_progressRectLeft = int(13 + float(m_processFont.height) * 1.5f);
+	}
+
+	void Visualizer::DirtyBitmaps(bool full)
+	{
+		for (auto& session : m_traceView.sessions)
+			for (auto& processor : session.processors)
+				for (auto& process : processor.processes)
+				{
+					process.bitmapDirty = true;
+					if (full)
+						process.bitmap = 0;
+				}
+
+		for (auto& workTrack : m_traceView.workTracks)
+			for (auto& work : workTrack.records)
+			{
+				work.bitmapDirty = true;
+				if (full)
+					work.bitmap = 0;
+			}
+
+		if (!full)
+			return;
+
+		m_backend.Reset();
+	}
+
 	////////////////////////////////////////////////////////////////////////////////////////////////////
 }
diff --git a/Engine/Source/Programs/UnrealBuildAccelerator/Visualizer/Private/UbaVisualizer.h b/Engine/Source/Programs/UnrealBuildAccelerator/Visualizer/Private/UbaVisualizer.h
--- a/Engine/Source/Programs/UnrealBuildAccelerator/Visualizer/Private/UbaVisualizer.h
+++ b/Engine/Source/Programs/UnrealBuildAccelerator/Visualizer/Private/UbaVisualizer.h
@@ -53,6 +53,9 @@ namespace uba
 		void Reset();
 		bool Unselect();
 		void UpdateFont(Font& font, int height, bool createUnderline);
+		void UpdateDefaultFont();
+		void UpdateProcessFont();
+		void DirtyBitmaps(bool full);
 
 		enum VisualizerFlag
 		{
diff --git a/Engine/Source/Programs/UnrealBuildAccelerator/Visualizer/Private/UbaVisualizerBackend.h b/Engine/Source/Programs/UnrealBuildAccelerator/Visualizer/Private/UbaVisualizerBackend.h
--- a/Engine/Source/Programs/UnrealBuildAccelerator/Visualizer/Private/UbaVisualizerBackend.h
+++ b/Engine/Source/Programs/UnrealBuildAccelerator/Visualizer/Private/UbaVisualizerBackend.h
@@ -2,6 +2,8 @@
 
 #pragma once
 
+#include "UbaTraceReader.h"
+
 namespace uba
 {
 	////////////////////////////////////////////////////////////////////////////////////////////////////
@@ -12,6 +14,9 @@ namespace uba
 	struct RECT { int left; int top; int right; int bottom; };
 	using COLORREF = u32;
 	#define RGB(r,g,b) ((COLORREF)((u8(r)|(u16(u8(g))<<8))|(u32(u8(b))<<16)))
+	#define GET_RED(c) u8(c & 0xff)
+	#define GET_GREEN(c) u8((c >> 8) & 0xff)
+	#define GET_BLUE(c) u8((c >> 16) & 0xff)
 	#endif
 
 	using BrushHandle = void*;
diff --git a/Engine/Source/Programs/UnrealBuildAccelerator/Visualizer/Private/Windows/UbaVisualizerWin.cpp b/Engine/Source/Programs/UnrealBuildAccelerator/Visualizer/Private/Windows/UbaVisualizerWin.cpp
--- a/Engine/Source/Programs/UnrealBuildAccelerator/Visualizer/Private/Windows/UbaVisualizerWin.cpp
+++ b/Engine/Source/Programs/UnrealBuildAccelerator/Visualizer/Private/Windows/UbaVisualizerWin.cpp
@@ -497,85 +497,6 @@ namespace uba
 		m_config.Save(m_logger);
 	}
 
-	void VisualizerWin::DirtyBitmaps(bool full)
-	{
-		for (auto& session : m_traceView.sessions)
-			for (auto& processor : session.processors)
-				for (auto& process : processor.processes)
-				{
-					process.bitmapDirty = true;
-					if (full)
-						process.bitmap = 0;
-				}
-
-		for (auto& workTrack : m_traceView.workTracks)
-			for (auto& work : workTrack.records)
-			{
-				work.bitmapDirty = true;
-				if (full)
-					work.bitmap = 0;
-			}
-
-		if (!full)
-			return;
-
-		m_backend.Reset();
-	}
-
-	void Visualizer::UpdateFont(Font& font, int height, bool createUnderline)
-	{
-		font.height = height;
-
-		int fh = height;
-		font.offset = 0;
-		if (height <= 13)
-		{
-			++fh;
-			--font.offset;
-		}
-		if (height <= 11)
-			++fh;
-		if (height <= 9)
-			++fh;
-		if (height <= 8)
-			++fh;
-		if (height <= 6)
-		{
-			++fh;
-			--font.offset;
-		}
-		if (height <= 4)
-			--font.offset;
-
-		if (font.handle)
-			m_backend.DeleteFont(font.handle);
-		font.handle = 0;
-		if (font.handleUnderlined)
-			m_backend.DeleteFont(font.handleUnderlined);
-		font.handleUnderlined = 0;
-
-		font.handle = m_backend.CreateFont(m_config.fontName.c_str(), fh, false);
-		if (createUnderline)
-			font.handleUnderlined = m_backend.CreateFont(m_config.fontName.c_str(), fh, true);
-	}
-
-	void VisualizerWin::UpdateDefaultFont()
-	{
-		UpdateFont(m_defaultFont, m_config.fontSize, true);
-		m_sessionStepY = m_defaultFont.height + 4;
-
-		m_timelineFont = m_defaultFont;
-	}
-
-	void VisualizerWin::UpdateProcessFont()
-	{
-		m_zoomValue = 1.0f + float(m_config.boxHeight) / 30.0f;
-		int fontHeight = Max(int(m_config.boxHeight) - 2, 1);
-		UpdateFont(m_processFont, fontHeight, false);
-		m_progressRectLeft = int(13 + float(m_processFont.height) * 1.5f);
-		DirtyBitmaps(true);
-	}
-
 	void VisualizerWin::ChangeFontSize(int offset)
 	{
 		m_config.fontSize += offset;
diff --git a/Engine/Source/Programs/UnrealBuildAccelerator/Visualizer/Private/Windows/UbaVisualizerWin.h b/Engine/Source/Programs/UnrealBuildAccelerator/Visualizer/Private/Windows/UbaVisualizerWin.h
--- a/Engine/Source/Programs/UnrealBuildAccelerator/Visualizer/Private/Windows/UbaVisualizerWin.h
+++ b/Engine/Source/Programs/UnrealBuildAccelerator/Visualizer/Private/Windows/UbaVisualizerWin.h
@@ -39,10 +39,7 @@ namespace uba
 		void StartDragToScroll(const POINT& anchor);
 		void StopDragToScroll();
 		void SaveSettings();
-		void DirtyBitmaps(bool full);
 
-		void UpdateDefaultFont();
-		void UpdateProcessFont();
 		void ChangeFontSize(int offset);
 		virtual void Redraw(bool now) override;
 		virtual void StopUpdate() override;
