[UBA]
* Implemented the DrawCached method for mac

diff --git a/Engine/Source/Programs/UnrealBuildAccelerator/Visualizer/Private/Mac/UbaVisualizerApp.mm b/Engine/Source/Programs/UnrealBuildAccelerator/Visualizer/Private/Mac/UbaVisualizerApp.mm
--- a/Engine/Source/Programs/UnrealBuildAccelerator/Visualizer/Private/Mac/UbaVisualizerApp.mm
+++ b/Engine/Source/Programs/UnrealBuildAccelerator/Visualizer/Private/Mac/UbaVisualizerApp.mm
@@ -19,6 +19,18 @@
 	class VisualizerBackendMac : public VisualizerBackend
 	{
 	public:
+		static constexpr int BitmapCacheHeight = 16*1024;
+
+		struct Bitmap
+		{
+			size_t w, h;
+			size_t bytesPerRow;
+			CGColorSpaceRef cs;
+			CGContextRef bmpCtx;
+			CGImageRef img;
+			void* pixels;
+		};
+
 		struct PaintContextMac : PaintContext
 		{
 			virtual SIZE GetClientSize() override
@@ -184,19 +196,116 @@ virtual SIZE GetTextSize(StringView str) override
 			{
 				return { int(str.count)*8, 10 };
 			}
+			CGRect SrcTopLeftToImagePixels(CGImageRef img, CGRect topLeftPx)
+			{
+				size_t H = BitmapCacheHeight;//CGImageGetHeight(img);
+				// topLeftPx.origin.y is measured from top
+				return CGRectMake(topLeftPx.origin.x,
+								(CGFloat)H - topLeftPx.origin.y - topLeftPx.size.height,
+								topLeftPx.size.width,
+								topLeftPx.size.height);
+			}			
 			virtual void DrawCached(TraceView::BitmapCache& process, RECT rect, int offsetY, int textHeight, Font& processFont, RECT clippingRect, const Function<void(PaintContext& textCacheContext)>& func) override
 			{
-				CGContextSaveGState(ctx);
-				CGContextTranslateCTM(ctx, rect.left, rect.top);
-				func(*this);
-				CGContextRestoreGState(ctx);
+				//CGContextSaveGState(ctx);
+				//CGContextTranslateCTM(ctx, rect.left, rect.top);
+				//func(*this);
+				//CGContextRestoreGState(ctx);
+
+				if (!process.bitmap || process.bitmapDirty)
+				{
+					if (!process.bitmap)
+					{
+						if (owner.m_lastBitmapOffset + processFont.height > BitmapCacheHeight)
+						{
+							if (owner.m_lastBitmap)
+								owner.m_textBitmaps.push_back(owner.m_lastBitmap);
+
+							auto& b = *new Bitmap;
+							b.w = 256;
+							b.h = BitmapCacheHeight;
+							b.bytesPerRow = b.w * 4;
+							b.cs = CGColorSpaceCreateDeviceRGB();
+							b.pixels = calloc(b.h, b.bytesPerRow);
+							CGBitmapInfo info = kCGBitmapByteOrder32Little | kCGImageAlphaPremultipliedFirst;
+							b.bmpCtx = CGBitmapContextCreate(b.pixels, b.w, b.h, 8, b.bytesPerRow, b.cs, info);
+
+							owner.m_lastBitmapOffset = 0;
+							owner.m_lastBitmap = &b;
+						}
+						process.bitmap = (HBITMAP)owner.m_lastBitmap;
+						process.bitmapOffset = owner.m_lastBitmapOffset;
+						owner.m_lastBitmapOffset += processFont.height;
+					}
+					if (lastSelectedBitmap != (Bitmap*)process.bitmap)
+						lastSelectedBitmap = (Bitmap*)process.bitmap;
+
+					PaintContextMac textCacheContext{ owner, lastSelectedBitmap->bmpCtx, { 256, BitmapCacheHeight } };
+					textCacheContext.SetFont(activeFont);
+					func(textCacheContext);
+
+					CGImageRelease(lastSelectedBitmap->img);
+					lastSelectedBitmap->img = CGBitmapContextCreateImage(lastSelectedBitmap->bmpCtx);
+				}
+
+				if (lastSelectedBitmap != (Bitmap*)process.bitmap)
+					lastSelectedBitmap = (Bitmap*)process.bitmap;
+
+				int processWidth = rect.right - rect.left;
+				int width = Min(processWidth, 256);
+				int bitmapOffsetY = process.bitmapOffset;
+				int bltOffsetY = offsetY;
+				if (bltOffsetY < 0)
+				{
+					bitmapOffsetY -= bltOffsetY;
+					bltOffsetY = 0;
+				}
+				int height = Min(textHeight, processFont.height);
+				if (bltOffsetY + height > textHeight)
+					height = textHeight - bltOffsetY;
+
+				int left = rect.left;
+				if (left <= -256 || height <= 0)
+					return;
+
+				int bitmapOffsetX = rect.left - left;
+
+				if (left < clippingRect.left)
+				{
+					int diff = clippingRect.left - left;
+					rect.left = clippingRect.left;
+					width -= diff;
+					bitmapOffsetX += diff;
+				}
+
+				CGImageRef img = lastSelectedBitmap->img;
+				CGRect destRect = CGRectMake(float(rect.left), float(rect.top + bltOffsetY), float(width), float(height));
+				CGRect sourceRect = CGRectMake(float(bitmapOffsetX), float(bitmapOffsetY), float(width), float(height));
+
+				if (false)
+				{
+					CGImageRef sub = CGImageCreateWithImageInRect(img, SrcTopLeftToImagePixels(img, sourceRect));
+					CGContextDrawImage(ctx, destRect, sub);
+					CGImageRelease(sub);
+				}
+				else
+				{
+					CGContextSaveGState(ctx);
+					CGContextClipToRect(ctx, destRect);
+					CGFloat sx = destRect.size.width  / sourceRect.size.width;
+					CGFloat sy = destRect.size.height / sourceRect.size.height;
+					CGContextTranslateCTM(ctx, destRect.origin.x - sourceRect.origin.x * sx, destRect.origin.y - sourceRect.origin.y * sy);
+					CGContextScaleCTM(ctx, sx, sy);
+					CGRect full = CGRectMake(0, 0, (CGFloat)256, (CGFloat)BitmapCacheHeight); // Draw whole image; clip ensures only desired region shows
+					CGContextDrawImage(ctx, full, img);
+					CGContextRestoreGState(ctx);
+				}
 			}
 
 			virtual POINT GetCursorClientPos() override { return {}; }
 
-			PaintContextMac(VisualizerBackendMac& o) : owner(o)
+			PaintContextMac(VisualizerBackendMac& o, CGContextRef c, SIZE s) : owner(o), ctx(c), clientSize(s)
 			{
-				ctx = [[NSGraphicsContext currentContext] CGContext];
 			}
 
 			VisualizerBackendMac& owner;
@@ -205,6 +314,7 @@ virtual void DrawCached(TraceView::BitmapCache& process, RECT rect, int offsetY,
 			Font activeFont;
 			COLORREF activeTextColor;
 			CGColorRef activeTextColorRef = nullptr;
+			Bitmap* lastSelectedBitmap = 0;
 		};
 
 		VisualizerBackendMac()
@@ -263,6 +373,10 @@ virtual void DeleteFont(FontHandle font) override
 
 		CGColorSpaceRef colorSpaceRef;
 		List<CGColorRef> colors;
+
+		Bitmap* m_lastBitmap = 0;
+		int m_lastBitmapOffset = BitmapCacheHeight;
+		Vector<Bitmap*> m_textBitmaps;
 	};
 
 	class VisualizerMac : public Visualizer
@@ -292,6 +406,8 @@ virtual void DeleteFont(FontHandle font) override
 			else
 				m_trace.ReadFile(m_traceView, path, false);
 
+			//m_config.boxHeight = 20;
+
 			UpdateDefaultFont();
 			UpdateProcessFont();
 			m_popupFont.handle = m_backend.CreateFont(TC("Consolas"), 16, false);
@@ -314,8 +430,8 @@ void Paint(SIZE size)
 			m_clientWidth = size.cx;
 			m_clientHeight = size.cy;
 
-			uba::VisualizerBackendMac::PaintContextMac context((VisualizerBackendMac&)m_backend);
-			context.clientSize = size;
+			auto ctx = [[NSGraphicsContext currentContext] CGContext];
+			uba::VisualizerBackendMac::PaintContextMac context((VisualizerBackendMac&)m_backend, ctx, size);
 			RECT r { 0, 0, size.cx, size.cy };
 
 			context.FillRect(r, m_backgroundBrush);
