[UBA]
* Added mouse tracking on visualizer for mac

diff --git a/Engine/Source/Programs/UnrealBuildAccelerator/Visualizer/Private/Mac/UbaVisualizerApp.mm b/Engine/Source/Programs/UnrealBuildAccelerator/Visualizer/Private/Mac/UbaVisualizerApp.mm
--- a/Engine/Source/Programs/UnrealBuildAccelerator/Visualizer/Private/Mac/UbaVisualizerApp.mm
+++ b/Engine/Source/Programs/UnrealBuildAccelerator/Visualizer/Private/Mac/UbaVisualizerApp.mm
@@ -16,7 +16,7 @@
 
 namespace uba
 {
-	class VisualizerBackendMac : public VisualizerBackend
+	class VisualizerBackendMac final : public VisualizerBackend
 	{
 	public:
 		static constexpr int BitmapCacheHeight = 16*1024;
@@ -31,7 +31,7 @@
 			void* pixels;
 		};
 
-		struct PaintContextMac : PaintContext
+		struct PaintContextMac final : PaintContext
 		{
 			virtual SIZE GetClientSize() override
  			{
@@ -73,7 +73,7 @@ virtual void Text(int x, int y, StringView str, const RECT* clipRect) override
 				if (!cfStr)
 					return;
 
-				y -= 2;
+				y -= 2; // TODO: This should not be here
 
 				const void* keys[]   = { kCTFontAttributeName,			kCTForegroundColorAttributeName };
 				const void* values[] = { (CTFontRef)activeFont.handle,	activeTextColorRef };
@@ -194,24 +194,25 @@ virtual void PolyPolyline(const POINT* points, const u32* lines, u32 lineCount)
 			}
 			virtual SIZE GetTextSize(StringView str) override
 			{
-				return { int(str.count)*8, 10 };
+				CFStringRef text = CFStringCreateWithBytes(kCFAllocatorDefault, reinterpret_cast<const UInt8*>(str.data), static_cast<CFIndex>(str.count), kCFStringEncodingUTF8, false);
+				const void* keys[]   = { kCTFontAttributeName };
+				const void* values[] = { activeFont.handle };
+				CFDictionaryRef attrs = CFDictionaryCreate(kCFAllocatorDefault, keys, values, 1, &kCFTypeDictionaryKeyCallBacks, &kCFTypeDictionaryValueCallBacks);
+				CFAttributedStringRef attr = CFAttributedStringCreate(kCFAllocatorDefault, text, attrs);
+				CTLineRef line = CTLineCreateWithAttributedString(attr);
+
+				double ascent, descent, leading;
+				double width = CTLineGetTypographicBounds(line, &ascent, &descent, &leading);
+				double height = ascent + descent + leading;
+
+				CFRelease(line);
+				CFRelease(attr);
+				CFRelease(attrs);
+				CFRelease(text);
+				return { int(width), int(height) };
 			}
-			CGRect SrcTopLeftToImagePixels(CGImageRef img, CGRect topLeftPx)
-			{
-				size_t H = BitmapCacheHeight;//CGImageGetHeight(img);
-				// topLeftPx.origin.y is measured from top
-				return CGRectMake(topLeftPx.origin.x,
-								(CGFloat)H - topLeftPx.origin.y - topLeftPx.size.height,
-								topLeftPx.size.width,
-								topLeftPx.size.height);
-			}			
 			virtual void DrawCached(TraceView::BitmapCache& process, RECT rect, int offsetY, int textHeight, Font& processFont, RECT clippingRect, const Function<void(PaintContext& textCacheContext)>& func) override
 			{
-				//CGContextSaveGState(ctx);
-				//CGContextTranslateCTM(ctx, rect.left, rect.top);
-				//func(*this);
-				//CGContextRestoreGState(ctx);
-
 				if (!process.bitmap || process.bitmapDirty)
 				{
 					if (!process.bitmap)
@@ -282,27 +283,21 @@ virtual void DrawCached(TraceView::BitmapCache& process, RECT rect, int offsetY,
 				CGRect destRect = CGRectMake(float(rect.left), float(rect.top + bltOffsetY), float(width), float(height));
 				CGRect sourceRect = CGRectMake(float(bitmapOffsetX), float(bitmapOffsetY), float(width), float(height));
 
-				if (false)
-				{
-					CGImageRef sub = CGImageCreateWithImageInRect(img, SrcTopLeftToImagePixels(img, sourceRect));
-					CGContextDrawImage(ctx, destRect, sub);
-					CGImageRelease(sub);
-				}
-				else
-				{
-					CGContextSaveGState(ctx);
-					CGContextClipToRect(ctx, destRect);
-					CGFloat sx = destRect.size.width  / sourceRect.size.width;
-					CGFloat sy = destRect.size.height / sourceRect.size.height;
-					CGContextTranslateCTM(ctx, destRect.origin.x - sourceRect.origin.x * sx, destRect.origin.y - sourceRect.origin.y * sy);
-					CGContextScaleCTM(ctx, sx, sy);
-					CGRect full = CGRectMake(0, 0, (CGFloat)256, (CGFloat)BitmapCacheHeight); // Draw whole image; clip ensures only desired region shows
-					CGContextDrawImage(ctx, full, img);
-					CGContextRestoreGState(ctx);
-				}
+				CGContextSaveGState(ctx);
+				CGContextClipToRect(ctx, destRect);
+				CGFloat sx = destRect.size.width  / sourceRect.size.width;
+				CGFloat sy = destRect.size.height / sourceRect.size.height;
+				CGContextTranslateCTM(ctx, destRect.origin.x - sourceRect.origin.x * sx, destRect.origin.y - sourceRect.origin.y * sy);
+				CGContextScaleCTM(ctx, sx, sy);
+				CGRect full = CGRectMake(0, 0, (CGFloat)256, (CGFloat)BitmapCacheHeight); // Draw whole image; clip ensures only desired region shows
+				CGContextDrawImage(ctx, full, img);
+				CGContextRestoreGState(ctx);
 			}
 
-			virtual POINT GetCursorClientPos() override { return {}; }
+			virtual POINT GetCursorClientPos() override
+			{
+				return owner.m_mousePos;
+			}
 
 			PaintContextMac(VisualizerBackendMac& o, CGContextRef c, SIZE s) : owner(o), ctx(c), clientSize(s)
 			{
@@ -341,8 +336,14 @@ virtual void Resize(int width, int height) override
 		virtual void Paint(const Function<void(PaintContext& context, RECT rect)>& func) override
 		{
 		}
-		virtual PaintContext* CreateContext(int clientWidth, int clientHeight) override { return nullptr; }
-		virtual void DeleteContext(PaintContext*) override {}
+		virtual PaintContext* CreateContext(int clientWidth, int clientHeight) override
+		{
+			return new PaintContextMac(*this, 0, {clientWidth, clientHeight});
+		}
+		virtual void DeleteContext(PaintContext* context) override
+		{
+			delete (PaintContextMac*)context;
+		}
 		virtual BrushHandle CreateSolidBrush(COLORREF color) override
 		{
 				float r = float(GET_RED(color))/255.0f;
@@ -377,6 +378,8 @@ virtual void DeleteFont(FontHandle font) override
 		Bitmap* m_lastBitmap = 0;
 		int m_lastBitmapOffset = BitmapCacheHeight;
 		Vector<Bitmap*> m_textBitmaps;
+
+		POINT m_mousePos;
 	};
 
 	class VisualizerMac : public Visualizer
@@ -410,7 +413,7 @@ virtual void DeleteFont(FontHandle font) override
 
 			UpdateDefaultFont();
 			UpdateProcessFont();
-			m_popupFont.handle = m_backend.CreateFont(TC("Consolas"), 16, false);
+			m_popupFont.handle = m_backend.CreateFont(TC("Menlo"), 16, false);
 			m_popupFont.height = 14;
 
 			InitBrushes();
@@ -421,8 +424,12 @@ virtual void PaintEmptyScreen(PaintContext& context) override
 			NSRect rect = NSMakeRect(50, 50, 200, 100);
 			NSRectFill(rect);
 		}
-		virtual void FinishIncompatible()
+		virtual void FinishIncompatible() override
+		{
+		}
+		bool UpdateSelection() override
 		{
+			return Visualizer::UpdateSelection(((VisualizerBackendMac&)m_backend).m_mousePos);
 		}
 
 		void Paint(SIZE size)
@@ -442,12 +449,38 @@ void Update()
 		{
 			UpdateTrace();
 		}
+
+		void MouseMove(CGFloat x, CGFloat y)
+		{
+			((VisualizerBackendMac&)m_backend).m_mousePos =  { (int)x, (int)y };
+			if (UpdateSelection() || m_config.showCursorLine)
+				;//Redraw(false);
+		}
+
+		void MouseDown(CGFloat x, CGFloat y)
+		{
+			((VisualizerBackendMac&)m_backend).m_mousePos =  { (int)x, (int)y };
+
+			if (m_buttonSelected != ~0u)
+			{
+				bool* values = &m_config.showProgress;
+				values[m_buttonSelected] = !values[m_buttonSelected];
+
+				PaintContext* context = m_backend.CreateContext(m_clientWidth, m_clientHeight);
+				HitTestResult res;
+				HitTest(res, *context, { -1, -1 });
+				m_backend.DeleteContext(context);
+				UpdateScrollbars(true);
+				Redraw(false);
+			}
+		}
 	};
 }
 
 @interface MyCanvasView : NSView
 {
 	NSTimeInterval _time;
+    NSTrackingArea* _tracking;
 	uba::VisualizerMac* _visualizer;
 }
 - (void)tick:(NSTimer *)timer;
@@ -503,6 +536,38 @@ - (void)drawRect:(NSRect)dirtyRect {
 	_visualizer->Paint(size);
 }
 
+- (void)updateTrackingAreas {
+    [super updateTrackingAreas];
+
+    if (_tracking)
+	{
+        [self removeTrackingArea:_tracking];
+        _tracking = nil;
+    }
+
+    NSTrackingAreaOptions opts = NSTrackingMouseMoved | NSTrackingActiveInKeyWindow | NSTrackingInVisibleRect;
+
+    _tracking = [[NSTrackingArea alloc] initWithRect:NSZeroRect options:opts owner:self userInfo:nil];
+    [self addTrackingArea:_tracking];
+}
+
+- (BOOL)acceptsFirstResponder { return YES; }
+
+- (void)viewDidMoveToWindow {
+    [super viewDidMoveToWindow];
+    [self.window setAcceptsMouseMovedEvents:YES];
+}
+
+- (void)mouseMoved:(NSEvent*)event {
+    NSPoint p = [self convertPoint:event.locationInWindow fromView:nil];
+	_visualizer->MouseMove(p.x, p.y);
+}
+
+- (void)mouseDown:(NSEvent*)event
+{
+    NSPoint p = [self convertPoint:event.locationInWindow fromView:nil];
+	_visualizer->MouseDown(p.x, p.y);
+}
 @end
 
 @interface AppDelegate : NSObject <NSWindowDelegate>
diff --git a/Engine/Source/Programs/UnrealBuildAccelerator/Visualizer/Private/UbaVisualizer.cpp b/Engine/Source/Programs/UnrealBuildAccelerator/Visualizer/Private/UbaVisualizer.cpp
--- a/Engine/Source/Programs/UnrealBuildAccelerator/Visualizer/Private/UbaVisualizer.cpp
+++ b/Engine/Source/Programs/UnrealBuildAccelerator/Visualizer/Private/UbaVisualizer.cpp
@@ -11,7 +11,7 @@ namespace uba
 	u64 ConvertTime(const TraceView& view, u64 time);
 
 	#if !PLATFORM_WINDOWS
-	void OffsetRect(RECT* r, int x, int y) { r->left -= x; r->right -= x; r->top -= y; r->bottom -= y; }
+	void OffsetRect(RECT* r, int x, int y) { r->left += x; r->right += x; r->top += y; r->bottom += y; }
 	#endif
 
 	////////////////////////////////////////////////////////////////////////////////////////////////////
@@ -63,7 +63,11 @@ namespace uba
 			SIZE clientSize = context.GetClientSize();
 
 			if (r.right > clientSize.cx)
+			{
 				OffsetRect(&r, -width - 15, 0);
+				if (r.left < 0)
+					OffsetRect(&r, -r.left, 0);
+			}
 			if (r.bottom > clientSize.cy)
 			{
 				OffsetRect(&r, 0, clientSize.cy - r.bottom);
@@ -1123,9 +1127,9 @@ namespace uba
 
 						//text.Appendf(TC(" - Cpu: %.1f%%"), cpuLoad * 100.0f);
 						//if (memAvail)
-						//	text.Appendf(TC(" Mem: %ls/%ls"), BytesToText(session.memTotal - memAvail).str, BytesToText(session.memTotal).str);
+						//	text.Appendf(TC(" Mem: %s/%s"), BytesToText(session.memTotal - memAvail).str, BytesToText(session.memTotal).str);
 						if (ping)
-							text.Appendf(TC(" Ping: %ls"), TimeToText(ping, false, m_traceView.frequency).str);
+							text.Appendf(TC(" Ping: %s"), TimeToText(ping, false, m_traceView.frequency).str);
 						if (!session.notification.empty())
 							text.Append(TCV(" - ")).Append(session.notification);
 					}
@@ -1792,14 +1796,14 @@ namespace uba
 			logger.AddWidth(3);
 
 			logger.AddSpace(2);
-			logger.Info(TC("  %ls"), process.description.c_str());
+			logger.Info(TC("  %s"), process.description.c_str());
 			if (process.type != TraceView::ProcessType_Task)
 			{
-				logger.Info(TC("  Host:        %ls"), m_processSelectedLocation.sessionIndex == 0 ? TC("local") : m_traceView.GetSession(m_processSelectedLocation).name.c_str());
+				logger.Info(TC("  Host:        %s"), m_processSelectedLocation.sessionIndex == 0 ? TC("local") : m_traceView.GetSession(m_processSelectedLocation).name.c_str());
 				logger.Info(TC("  ProcessId:  %6u"), process.id);
 			}
-			logger.Info(TC("  Start:     %7ls (%ls)"), TimeToText(process.start, true).str, GetWorldTime(process.start).data);
-			logger.Info(TC("  Duration:  %7ls"), TimeToText(duration, true).str);
+			logger.Info(TC("  Start:     %7s (%s)"), TimeToText(process.start, true).str, GetWorldTime(process.start).data);
+			logger.Info(TC("  Duration:  %7s"), TimeToText(duration, true).str);
 			if (!process.returnedReason.empty())
 				logger.Info(TC("  Returned: %7s"), process.returnedReason.data());
 			if (hasExited && process.exitCode != 0)
@@ -1945,9 +1949,9 @@ namespace uba
 			DrawTextLogger logger(context, m_popupFont.height, m_tooltipBackgroundBrush);
 
 			logger.AddSpace();
-			logger.Info(TC("  %ls"), selectedWork.description);
-			logger.Info(TC("  Start:     %ls (%ls)"), TimeToText(selectedWork.start, true).str, GetWorldTime(selectedWork.start).data);
-			logger.Info(TC("  Duration:  %ls"), TimeToText(duration, true).str);
+			logger.Info(TC("  %s"), selectedWork.description);
+			logger.Info(TC("  Start:     %s (%s)"), TimeToText(selectedWork.start, true).str, GetWorldTime(selectedWork.start).data);
+			logger.Info(TC("  Duration:  %s"), TimeToText(duration, true).str);
 			if (!selectedWork.entries.empty())
 				logger.AddSpace();
 			for (u32 i=0,e=u32(selectedWork.entries.size()); i!=e; ++i)
@@ -1958,7 +1962,7 @@ namespace uba
 				{
 					u64 start = entry.startTime ? entry.startTime : entry.time;
 					if (TimeToMs(start - selectedWork.start) > 1)
-						logger.Info(TC("   Start (%ls)"), TimeToText(start - selectedWork.start).str);
+						logger.Info(TC("   Start (%s)"), TimeToText(start - selectedWork.start).str);
 				}
 
 				if (entry.startTime)
@@ -1977,9 +1981,9 @@ namespace uba
 				if (!time && selectedWork.stop != ~0u)
 					time = selectedWork.stop - entry.time;
 				if (entry.count == 1)
-					logger.Info(TC("%s  %ls (%ls)"), (entry.startTime ? TC(" ") : TC("")), entry.text, TimeToText(time).str);
+					logger.Info(TC("%s  %s (%s)"), (entry.startTime ? TC(" ") : TC("")), entry.text, TimeToText(time).str);
 				else
-					logger.Info(TC("%s  %ls (%ls %u)"), (entry.startTime ? TC(" ") : TC("")), entry.text, TimeToText(time).str, entry.count);
+					logger.Info(TC("%s  %s (%s %u)"), (entry.startTime ? TC(" ") : TC("")), entry.text, TimeToText(time).str, entry.count);
 			}
 			logger.AddSpace();
 			logger.DrawAtCursor();
@@ -2028,16 +2032,16 @@ namespace uba
 			DrawTextLogger logger(context, m_popupFont.height, m_tooltipBackgroundBrush);
 			logger.AddSpace(3);
 			logger.SetColor(m_cpuColor).Info(TC("  Cpu: %.1f%%"), m_stats.cpuLoad * 100.0f);
-			logger.SetColor(m_memColor).Info(TC("  Mem: %ls/%ls"), BytesToText(m_stats.memTotal - m_stats.memAvail).str, BytesToText(m_stats.memTotal).str);
+			logger.SetColor(m_memColor).Info(TC("  Mem: %s/%s"), BytesToText(m_stats.memTotal - m_stats.memAvail).str, BytesToText(m_stats.memTotal).str);
 			if (m_stats.recvBytes || m_stats.sendBytes)
 			{
-				logger.SetColor(m_recvColor).Info(TC("  Recv: %lsit/s"), BytesToText(m_stats.recvBytesPerSecond*8).str);
-				logger.SetColor(m_sendColor).Info(TC("  Send: %lsit/s"), BytesToText(m_stats.sendBytesPerSecond*8).str);
+				logger.SetColor(m_recvColor).Info(TC("  Recv: %sit/s"), BytesToText(m_stats.recvBytesPerSecond*8).str);
+				logger.SetColor(m_sendColor).Info(TC("  Send: %sit/s"), BytesToText(m_stats.sendBytesPerSecond*8).str);
 			}
 			logger.SetColor(m_textColor);
 
 			if (m_stats.ping)
-				logger.Info(TC("  Ping: %ls"), TimeToText(m_stats.ping, false, m_traceView.frequency).str);
+				logger.Info(TC("  Ping: %s"), TimeToText(m_stats.ping, false, m_traceView.frequency).str);
 			if (m_stats.connectionCount)
 				logger.Info(TC("  Conn: %u"), u32(m_stats.connectionCount));
 
@@ -2058,16 +2062,16 @@ namespace uba
 		}
 		else if (m_buttonSelected != ~0u)
 		{
-			const wchar_t* tooltip[] =
+			const tchar* tooltip[] =
 			{
-#define UBA_VISUALIZER_FLAG(name, defaultValue, desc) desc,
+#define UBA_VISUALIZER_FLAG(name, defaultValue, desc) TC(desc),
 				UBA_VISUALIZER_FLAGS1
 #undef UBA_VISUALIZER_FLAG
 			};
 
 			context.SetFont(m_popupFont);
 			DrawTextLogger logger(context, m_popupFont.height, m_tooltipBackgroundBrush);
-			logger.Info(TC("%ls %ls"), TC("Show"), tooltip[m_buttonSelected]);
+			logger.Info(TC("%s %s"), TC("Show"), tooltip[m_buttonSelected]);
 			logger.DrawAtCursor();
 		}
 		else if (m_fetchedFilesSelected != ~0u)
@@ -2392,9 +2396,9 @@ namespace uba
 			if (session.disconnectTime == ~u64(0))
 			{
 				if (session.proxyCreated)
-					drawText(TC("Proxy(HOSTED): %ls"), session.proxyName.c_str());
+					drawText(TC("Proxy(HOSTED): %s"), session.proxyName.c_str());
 				else if (!session.proxyName.empty())
-					drawText(TC("Proxy: %ls"), session.proxyName.c_str());
+					drawText(TC("Proxy: %s"), session.proxyName.c_str());
 				else
 					drawText(TC("Proxy: None"));
 			}
@@ -2420,8 +2424,8 @@ namespace uba
 					sendPerS = u64(float(updateSend - session.prevSend) / duration);
 					recvPerS = u64(float(updateRecv - session.prevRecv) / duration);
 				}
-				drawText(TC("Recv: %ls (%sit/s)"), BytesToText(updateRecv), BytesToText(recvPerS*8));
-				drawText(TC("Send: %ls (%sit/s)"), BytesToText(updateSend), BytesToText(sendPerS*8));
+				drawText(TC("Recv: %s (%sit/s)"), BytesToText(updateRecv), BytesToText(recvPerS*8));
+				drawText(TC("Send: %s (%sit/s)"), BytesToText(updateSend), BytesToText(sendPerS*8));
 			}
 
 			int fileWidth = 700;
@@ -2430,7 +2434,7 @@ namespace uba
 				{
 					textRect.left = posX;
 					textRect.right = posX + fileWidth;
-					drawText(TC("%ls Files: %u (%u) %s"), fileType, u32(count), u32(activeCount), BytesToText(bytes));
+					drawText(TC("%s Files: %u (%u) %s"), fileType, u32(count), u32(activeCount), BytesToText(bytes));
 					u32 fileCount = 0;
 					for (auto rit = files.rbegin(), rend = files.rend(); rit != rend; ++rit)
 					{
@@ -2498,8 +2502,8 @@ namespace uba
 						recvPerS = u64(float(updateRecv - session.prevRecv) / duration);
 					}
 
-					drawText(TC("Recv: %ls (%sit/s)"), BytesToText(updateRecv), BytesToText(recvPerS*8));
-					drawText(TC("Send: %ls (%sit/s)"), BytesToText(updateSend), BytesToText(sendPerS*8));
+					drawText(TC("Recv: %s (%sit/s)"), BytesToText(updateRecv), BytesToText(recvPerS*8));
+					drawText(TC("Send: %s (%sit/s)"), BytesToText(updateSend), BytesToText(sendPerS*8));
 				}
 			}
 
@@ -2523,7 +2527,7 @@ namespace uba
 						writePerS = u64(float(drive.writeBytes.back()) / duration);
 					}
 
-					drawText(TC("%c: Rd %ls (%s/s) Wr %ls (%s/s) "), pair.first, BytesToText(drive.totalReadBytes).str, BytesToText(readPerS).str, BytesToText(drive.totalWriteBytes).str, BytesToText(writePerS).str);
+					drawText(TC("%c: Rd %s (%s/s) Wr %s (%s/s) "), pair.first, BytesToText(drive.totalReadBytes).str, BytesToText(readPerS).str, BytesToText(drive.totalWriteBytes).str, BytesToText(writePerS).str);
 				}
 			}
 		}
@@ -2564,11 +2568,11 @@ namespace uba
 	void Visualizer::WriteProcessStats(Logger& out, const TraceView::Process& process)
 	{
 		bool hasExited = process.stop != ~u64(0);
-		out.Info(TC("  %ls"), process.description.c_str());
+		out.Info(TC("  %s"), process.description.c_str());
 		out.Info(TC("  ProcessId: %u"), process.id);
-		out.Info(TC("  Start:     %ls"), TimeToText(process.start, true).str);
+		out.Info(TC("  Start:     %s"), TimeToText(process.start, true).str);
 		if (hasExited)
-			out.Info(TC("  Duration:  %ls"), TimeToText(process.stop - process.start, true).str);
+			out.Info(TC("  Duration:  %s"), TimeToText(process.stop - process.start, true).str);
 		if (hasExited && process.exitCode != 0)
 			out.Info(TC("  ExitCode:  %u"), process.exitCode);
 
@@ -2621,12 +2625,12 @@ namespace uba
 
 	void Visualizer::WriteWorkStats(Logger& out, const TraceView::WorkRecord& record)
 	{
-		out.Info(TC("  %ls"), record.description);
-		out.Info(TC("  Start:     %ls"), TimeToText(record.start, true).str);
+		out.Info(TC("  %s"), record.description);
+		out.Info(TC("  Start:     %s"), TimeToText(record.start, true).str);
 		if (record.stop != ~u64(0))
-			out.Info(TC("  Duration:  %ls"), TimeToText(record.stop - record.start, true).str);
+			out.Info(TC("  Duration:  %s"), TimeToText(record.stop - record.start, true).str);
 		for (auto& e : record.entries)
-			out.Info(TC("   %ls (%ls)"), e.text, TimeToText(e.time - record.start).str);
+			out.Info(TC("   %s (%s)"), e.text, TimeToText(e.time - record.start).str);
 	}
 
 	void Visualizer::PaintEmptyScreen(PaintContext& context)
@@ -2898,5 +2902,42 @@ namespace uba
 		m_backend.Reset();
 	}
 
+	bool Visualizer::UpdateSelection(POINT pos)
+	{
+		HitTestResult res;
+		{
+			PaintContext* context = m_backend.CreateContext(m_clientWidth, m_clientHeight);
+			HitTest(res, *context, pos);
+			m_backend.DeleteContext(context);
+		}
+		m_activeSection = res.section;
+
+		if (res.processSelected == m_processSelected && res.processLocation == m_processSelectedLocation &&
+			res.sessionSelectedIndex == m_sessionSelectedIndex &&
+			res.statsSelected == m_statsSelected && res.stats == m_stats &&
+			res.buttonSelected == m_buttonSelected && res.timelineSelected == m_timelineSelected &&
+			res.activeProcessGraphSelected == m_activeProcessGraphSelected && 
+			res.activeProcessCount == m_activeProcessCount && 
+			res.fetchedFilesSelected == m_fetchedFilesSelected &&
+			res.workSelected == m_workSelected && res.workTrack == m_workTrack && res.workIndex == m_workIndex &&
+			res.hyperLink == m_hyperLinkSelected)
+			return false;
+		m_processSelected = res.processSelected;
+		m_processSelectedLocation = res.processLocation;
+		m_sessionSelectedIndex = res.sessionSelectedIndex;
+		m_statsSelected = res.statsSelected;
+		m_stats = res.stats;
+		m_activeProcessGraphSelected = res.activeProcessGraphSelected;
+		m_activeProcessCount = res.activeProcessCount;
+		m_buttonSelected = res.buttonSelected;
+		m_timelineSelected = res.timelineSelected;
+		m_fetchedFilesSelected = res.fetchedFilesSelected;
+		m_workSelected = res.workSelected;
+		m_workTrack = res.workTrack;
+		m_workIndex = res.workIndex;
+		m_hyperLinkSelected = res.hyperLink;
+		return true;
+	}
+
 	////////////////////////////////////////////////////////////////////////////////////////////////////
 }
diff --git a/Engine/Source/Programs/UnrealBuildAccelerator/Visualizer/Private/UbaVisualizer.h b/Engine/Source/Programs/UnrealBuildAccelerator/Visualizer/Private/UbaVisualizer.h
--- a/Engine/Source/Programs/UnrealBuildAccelerator/Visualizer/Private/UbaVisualizer.h
+++ b/Engine/Source/Programs/UnrealBuildAccelerator/Visualizer/Private/UbaVisualizer.h
@@ -56,6 +56,7 @@ namespace uba
 		void UpdateDefaultFont();
 		void UpdateProcessFont();
 		void DirtyBitmaps(bool full);
+		bool UpdateSelection(POINT pos);
 
 		enum VisualizerFlag
 		{
diff --git a/Engine/Source/Programs/UnrealBuildAccelerator/Visualizer/Private/UbaVisualizerConfig.h b/Engine/Source/Programs/UnrealBuildAccelerator/Visualizer/Private/UbaVisualizerConfig.h
--- a/Engine/Source/Programs/UnrealBuildAccelerator/Visualizer/Private/UbaVisualizerConfig.h
+++ b/Engine/Source/Programs/UnrealBuildAccelerator/Visualizer/Private/UbaVisualizerConfig.h
@@ -10,31 +10,31 @@ namespace uba
 	class Logger;
 
 	#define UBA_VISUALIZER_FLAGS1 \
-		UBA_VISUALIZER_FLAG(Progress, true, L"progress") \
-		UBA_VISUALIZER_FLAG(Status, true, L"status") \
-		UBA_VISUALIZER_FLAG(ActiveProcesses, false, L"active processes") \
-		UBA_VISUALIZER_FLAG(TitleBars, true, L"instance title bars") \
-		UBA_VISUALIZER_FLAG(DetailedData, false, L"detailed data (use -UbaDetailedTrace for even more)") \
-		UBA_VISUALIZER_FLAG(CpuMemStats, true, L"cpu/mem stats") \
-		UBA_VISUALIZER_FLAG(NetworkStats, true, L"network stats") \
-		UBA_VISUALIZER_FLAG(DriveStats, true, L"drive stats") \
-		UBA_VISUALIZER_FLAG(ActiveProcessGraph, false, L"graph of active processes over time") \
-		UBA_VISUALIZER_FLAG(ProcessBars, true, L"process bars") \
-		UBA_VISUALIZER_FLAG(FinishedProcesses, true, L"finished process bars") \
-		UBA_VISUALIZER_FLAG(Timeline, true, L"timeline") \
-		UBA_VISUALIZER_FLAG(Workers, false, L"workers (threads on host taking care of requests from helpers)") \
-		UBA_VISUALIZER_FLAG(CursorLine, false, L"cursor (vertical line)") \
+		UBA_VISUALIZER_FLAG(Progress, true, "progress") \
+		UBA_VISUALIZER_FLAG(Status, true, "status") \
+		UBA_VISUALIZER_FLAG(ActiveProcesses, false, "active processes") \
+		UBA_VISUALIZER_FLAG(TitleBars, true, "instance title bars") \
+		UBA_VISUALIZER_FLAG(DetailedData, false, "detailed data (use -UbaDetailedTrace for even more)") \
+		UBA_VISUALIZER_FLAG(CpuMemStats, true, "cpu/mem stats") \
+		UBA_VISUALIZER_FLAG(NetworkStats, true, "network stats") \
+		UBA_VISUALIZER_FLAG(DriveStats, true, "drive stats") \
+		UBA_VISUALIZER_FLAG(ActiveProcessGraph, false, "graph of active processes over time") \
+		UBA_VISUALIZER_FLAG(ProcessBars, true, "process bars") \
+		UBA_VISUALIZER_FLAG(FinishedProcesses, true, "finished process bars") \
+		UBA_VISUALIZER_FLAG(Timeline, true, "timeline") \
+		UBA_VISUALIZER_FLAG(Workers, false, "workers (threads on host taking care of requests from helpers)") \
+		UBA_VISUALIZER_FLAG(CursorLine, false, "cursor (vertical line)") \
 
 	#define UBA_VISUALIZER_FLAGS2 \
-		UBA_VISUALIZER_FLAG(ShowProcessText, true, L"Show text in process bars") \
-		UBA_VISUALIZER_FLAG(ShowReadWriteColors, true, L"Show colors for read/write times in process bars") \
-		UBA_VISUALIZER_FLAG(ScaleHorizontalWithScrollWheel, false, L"Use scroll wheel to scale horizontally") \
-		UBA_VISUALIZER_FLAG(DarkMode, false, L"Use dark mode to draw visualizer") \
-		UBA_VISUALIZER_FLAG(AutoSaveSettings, true, L"Auto save Position/Settings on close") \
-		UBA_VISUALIZER_FLAG(ShowAllTraces, true, L"Show all traces started on channel") \
-		UBA_VISUALIZER_FLAG(SortActiveRemoteSessions, true, L"Sort active sessions on top") \
-		UBA_VISUALIZER_FLAG(AutoScaleHorizontal, true, L"Automatically scale horizontally to fit processes") \
-		UBA_VISUALIZER_FLAG(LockTimelineToBottom, true, L"Lock timeline to always paint at bottom") \
+		UBA_VISUALIZER_FLAG(ShowProcessText, true, "Show text in process bars") \
+		UBA_VISUALIZER_FLAG(ShowReadWriteColors, true, "Show colors for read/write times in process bars") \
+		UBA_VISUALIZER_FLAG(ScaleHorizontalWithScrollWheel, false, "Use scroll wheel to scale horizontally") \
+		UBA_VISUALIZER_FLAG(DarkMode, false, "Use dark mode to draw visualizer") \
+		UBA_VISUALIZER_FLAG(AutoSaveSettings, true, "Auto save Position/Settings on close") \
+		UBA_VISUALIZER_FLAG(ShowAllTraces, true, "Show all traces started on channe") \
+		UBA_VISUALIZER_FLAG(SortActiveRemoteSessions, true, "Sort active sessions on top") \
+		UBA_VISUALIZER_FLAG(AutoScaleHorizontal, true, "Automatically scale horizontally to fit processes") \
+		UBA_VISUALIZER_FLAG(LockTimelineToBottom, true, "Lock timeline to always paint at bottom") \
 
 	struct VisualizerConfig
 	{
diff --git a/Engine/Source/Programs/UnrealBuildAccelerator/Visualizer/Private/Windows/UbaVisualizerWin.cpp b/Engine/Source/Programs/UnrealBuildAccelerator/Visualizer/Private/Windows/UbaVisualizerWin.cpp
--- a/Engine/Source/Programs/UnrealBuildAccelerator/Visualizer/Private/Windows/UbaVisualizerWin.cpp
+++ b/Engine/Source/Programs/UnrealBuildAccelerator/Visualizer/Private/Windows/UbaVisualizerWin.cpp
@@ -564,40 +564,7 @@ namespace uba
 		POINT pos;
 		GetCursorPos(&pos);
 		ScreenToClient(m_hwnd, &pos);
-
-		HitTestResult res;
-		{
-			PaintContext* context = m_backend.CreateContext(m_clientWidth, m_clientHeight);
-			HitTest(res, *context, pos);
-			m_backend.DeleteContext(context);
-		}
-		m_activeSection = res.section;
-
-		if (res.processSelected == m_processSelected && res.processLocation == m_processSelectedLocation &&
-			res.sessionSelectedIndex == m_sessionSelectedIndex &&
-			res.statsSelected == m_statsSelected && res.stats == m_stats &&
-			res.buttonSelected == m_buttonSelected && res.timelineSelected == m_timelineSelected &&
-			res.activeProcessGraphSelected == m_activeProcessGraphSelected && 
-			res.activeProcessCount == m_activeProcessCount && 
-			res.fetchedFilesSelected == m_fetchedFilesSelected &&
-			res.workSelected == m_workSelected && res.workTrack == m_workTrack && res.workIndex == m_workIndex &&
-			res.hyperLink == m_hyperLinkSelected)
-			return false;
-		m_processSelected = res.processSelected;
-		m_processSelectedLocation = res.processLocation;
-		m_sessionSelectedIndex = res.sessionSelectedIndex;
-		m_statsSelected = res.statsSelected;
-		m_stats = res.stats;
-		m_activeProcessGraphSelected = res.activeProcessGraphSelected;
-		m_activeProcessCount = res.activeProcessCount;
-		m_buttonSelected = res.buttonSelected;
-		m_timelineSelected = res.timelineSelected;
-		m_fetchedFilesSelected = res.fetchedFilesSelected;
-		m_workSelected = res.workSelected;
-		m_workTrack = res.workTrack;
-		m_workIndex = res.workIndex;
-		m_hyperLinkSelected = res.hyperLink;
-		return true;
+		return Visualizer::UpdateSelection(pos);
 	}
 
 	void VisualizerWin::UpdateScrollbars(bool redraw)
@@ -1111,7 +1078,7 @@ namespace uba
 			ClientToScreen(hWnd, &point);
 
 #define UBA_VISUALIZER_FLAG(name, defaultValue, desc) \
-				AppendMenuW(hMenu, MF_STRING | (m_config.name ? MF_CHECKED : 0), Popup_##name, desc);
+				AppendMenuW(hMenu, MF_STRING | (m_config.name ? MF_CHECKED : 0), Popup_##name, TC(desc));
 			UBA_VISUALIZER_FLAGS2
 #undef UBA_VISUALIZER_FLAG
 
