[UBA]
* Fixed so visualizer now can listen and start showing new traces on trace channel "Default"

diff --git a/Engine/Source/Programs/UnrealBuildAccelerator/Common/Private/UbaTraceChannel.cpp b/Engine/Source/Programs/UnrealBuildAccelerator/Common/Private/UbaTraceChannel.cpp
--- a/Engine/Source/Programs/UnrealBuildAccelerator/Common/Private/UbaTraceChannel.cpp
+++ b/Engine/Source/Programs/UnrealBuildAccelerator/Common/Private/UbaTraceChannel.cpp
@@ -245,6 +245,13 @@ namespace uba
 	bool TraceChannel::Read(StringBufferBase& outTraceName)
 	{
 		#if UBA_USE_TRACE_CHANNEL_FILE
+		struct stat info;
+		if (stat(m_channelFile.c_str(), &info) == -1)
+			return false;
+		if (info.st_mtimespec.tv_nsec == m_lastModified.tv_nsec && info.st_mtimespec.tv_sec == m_lastModified.tv_sec)
+			return true;
+		m_lastModified = info.st_mtimespec;
+
 		return UpdateTraceFile(m_logger, m_channelFile, TC(""), TC(""), [&](StringView trace)
 			{
 				if (!outTraceName.count)
diff --git a/Engine/Source/Programs/UnrealBuildAccelerator/Common/Public/UbaTraceChannel.h b/Engine/Source/Programs/UnrealBuildAccelerator/Common/Public/UbaTraceChannel.h
--- a/Engine/Source/Programs/UnrealBuildAccelerator/Common/Public/UbaTraceChannel.h
+++ b/Engine/Source/Programs/UnrealBuildAccelerator/Common/Public/UbaTraceChannel.h
@@ -27,7 +27,7 @@ namespace uba
 
 		#if UBA_USE_TRACE_CHANNEL_FILE
 		TString m_channelFile;
-
+		struct timespec m_lastModified;
 		#else
 		MutexHandle m_mutex = InvalidMutexHandle;
 		FileMappingHandle m_memHandle;
diff --git a/Engine/Source/Programs/UnrealBuildAccelerator/Visualizer/Private/Mac/UbaVisualizerApp.mm b/Engine/Source/Programs/UnrealBuildAccelerator/Visualizer/Private/Mac/UbaVisualizerApp.mm
--- a/Engine/Source/Programs/UnrealBuildAccelerator/Visualizer/Private/Mac/UbaVisualizerApp.mm
+++ b/Engine/Source/Programs/UnrealBuildAccelerator/Visualizer/Private/Mac/UbaVisualizerApp.mm
@@ -387,28 +387,30 @@ virtual void DeleteFont(FontHandle font) override
 	public:
 		VisualizerMac(VisualizerConfig& config, Logger& logger, VisualizerBackend& backend) : Visualizer(config, logger, backend)
 		{
-			char path[512];
-			GetFullPathNameW("~/Trace.uba", 512, path, nullptr);
-
-			m_fileName.Append(path);
-			if (true)
+			if (false)
 			{
-				m_trace.ReadFile(m_traceView, path, true);
-
-				m_startTime = GetTime();
-
-				bool changed = true;
-				u64 time = 0;
-				m_traceView.finished = false;
-				m_trace.UpdateReadFile(m_traceView, time, changed);
-				m_pauseStart = m_startTime + time;
-				m_pauseTime = m_startTime - m_pauseStart;
-				m_autoScroll = true;
-				m_replay = 1;
-			}
-			else
-				m_trace.ReadFile(m_traceView, path, false);
+				char path[512];
+				GetFullPathNameW("~/Trace.uba", 512, path, nullptr);
 
+				m_fileName.Append(path);
+				if (true)
+				{
+					m_trace.ReadFile(m_traceView, path, true);
+
+					m_startTime = GetTime();
+
+					bool changed = true;
+					u64 time = 0;
+					m_traceView.finished = false;
+					m_trace.UpdateReadFile(m_traceView, time, changed);
+					m_pauseStart = m_startTime + time;
+					m_pauseTime = m_startTime - m_pauseStart;
+					m_autoScroll = true;
+					m_replay = 1;
+				}
+				else
+					m_trace.ReadFile(m_traceView, path, false);
+			}
 			//m_config.boxHeight = 20;
 
 			UpdateDefaultFont();
@@ -416,13 +418,13 @@ virtual void DeleteFont(FontHandle font) override
 			m_popupFont.handle = m_backend.CreateFont(TC("Menlo"), 16, false);
 			m_popupFont.height = 14;
 
-			InitBrushes();
+			m_listenThread.Start([this]() { ThreadListen(); return 0; });
 		}
 		virtual void PaintEmptyScreen(PaintContext& context) override
 		{
-			[[NSColor blackColor] setFill];
-			NSRect rect = NSMakeRect(50, 50, 200, 100);
-			NSRectFill(rect);
+			//[[NSColor blackColor] setFill];
+			//NSRect rect = NSMakeRect(50, 50, 200, 100);
+			//NSRectFill(rect);
 		}
 		virtual void FinishIncompatible() override
 		{
@@ -434,11 +436,16 @@ bool UpdateSelection() override
 
 		void Paint(SIZE size)
 		{
+			auto& backend = (VisualizerBackendMac&)m_backend;
+
+			if (backend.colors.empty())
+				InitBrushes();
+
 			m_clientWidth = size.cx;
 			m_clientHeight = size.cy;
 
 			auto ctx = [[NSGraphicsContext currentContext] CGContext];
-			uba::VisualizerBackendMac::PaintContextMac context((VisualizerBackendMac&)m_backend, ctx, size);
+			uba::VisualizerBackendMac::PaintContextMac context(backend, ctx, size);
 			RECT r { 0, 0, size.cx, size.cy };
 
 			context.FillRect(r, m_backgroundBrush);
@@ -447,6 +454,19 @@ void Paint(SIZE size)
 
 		void Update()
 		{
+			//if (!m_started && m_hasTrace)
+			{
+				SCOPED_FUTEX(m_newTraceLock, lock);
+				if (!m_newTraceName.Equals(m_newTrace))
+				{
+					m_newTraceName.Clear().Append(m_newTrace);
+					m_usingNamed = true;
+					//m_hasTrace = true;
+					//m_started = true;
+					NewTrace(0, false);
+				}
+			}
+
 			UpdateTrace();
 		}
 
@@ -474,6 +494,51 @@ void MouseDown(CGFloat x, CGFloat y)
 				Redraw(false);
 			}
 		}
+
+		bool ThreadListen()
+		{
+			const tchar* channelName = TC("Default");
+
+			TraceChannel channel(m_logger);
+			if (!channel.Init(channelName))
+			{
+				m_logger.Error(TC("TODO"));
+				return false;
+			}
+
+			m_listenTimeout.Create(EventResetType_Auto);
+			
+			StringBuffer<256> traceName;
+			while (true)
+			{
+				traceName.Clear();
+				if (!channel.Read(traceName))
+				{
+					m_logger.Error(TC("TODO2"));
+					return false;
+				}
+				if (traceName.count)
+				{
+					if (!traceName.Equals(m_newTrace))
+					{
+						SCOPED_FUTEX(m_newTraceLock, lock);
+						m_newTrace = traceName.ToString();
+					}
+				}
+				else
+				{
+					SCOPED_FUTEX(m_newTraceLock, lock);
+					m_newTrace.clear();
+				}
+				Sleep(500);
+			}
+			return true;
+		}
+
+		Event m_listenTimeout;
+		Thread m_listenThread;
+		Futex m_newTraceLock;
+		TString m_newTrace;
 	};
 }
 
@@ -503,12 +568,7 @@ - (instancetype)initWithFrame:(NSRect)frame {
 
 		_visualizer = new uba::VisualizerMac(config, *logger, *backend);
 
-		// 60 FPS timer
-		[NSTimer scheduledTimerWithTimeInterval:(1.0/60.0)
-										 target:self
-									   selector:@selector(tick:)
-									   userInfo:nil
-										repeats:YES];
+		[NSTimer scheduledTimerWithTimeInterval:(1.0/60.0) target:self selector:@selector(tick:) userInfo:nil repeats:YES];
 	}
 	return self;
 }
diff --git a/Engine/Source/Programs/UnrealBuildAccelerator/Visualizer/Private/UbaTraceReader.cpp b/Engine/Source/Programs/UnrealBuildAccelerator/Visualizer/Private/UbaTraceReader.cpp
--- a/Engine/Source/Programs/UnrealBuildAccelerator/Visualizer/Private/UbaTraceReader.cpp
+++ b/Engine/Source/Programs/UnrealBuildAccelerator/Visualizer/Private/UbaTraceReader.cpp
@@ -322,11 +322,13 @@ namespace uba
 
 #if PLATFORM_WINDOWS
 		if (m_hostProcess && WaitForSingleObject(m_hostProcess, 0) != WAIT_TIMEOUT)
+#else
+		if (m_hostProcess && kill(m_hostProcess, 0) != 0)
+#endif
 		{
 			StopAllActive(out, GetTime() - out.realStartTime);
 			keepRunning = false;
 		}
-#endif
 
 		if (keepRunning || m_namedTrace.empty())
 			return true;
@@ -373,11 +375,15 @@ namespace uba
 			bool replay = maxTime != ~u64(0);
 
 			u32 hostProcessId = traceReader.ReadU32();
-#if PLATFORM_WINDOWS
 			m_hostProcess = 0;
 			if (trackHost && !replay)
+			{
+#if PLATFORM_WINDOWS
 				m_hostProcess = OpenProcess(SYNCHRONIZE, FALSE, hostProcessId);
+#else
+				m_hostProcess = hostProcessId;
 #endif
+			}
 			u64 traceSystemStartTimeUs = 0;
 
 			if (version < TraceReadCompatibilityVersion || version > TraceVersion)
@@ -1562,10 +1568,11 @@ namespace uba
 		u64 unmapSize = 0;
 		if (m_hostProcess)
 			CloseHandle(m_hostProcess);
-		m_hostProcess = 0;
 #else
 		u64 unmapSize = 128 * 1024 * 1024;
 #endif
+		m_hostProcess = 0;
+
 		if (m_memoryBegin)
 			UnmapViewOfFile(m_logger, m_memoryBegin, unmapSize, TC("TraceReader"));
 		m_memoryBegin = nullptr;
diff --git a/Engine/Source/Programs/UnrealBuildAccelerator/Visualizer/Private/UbaTraceReader.h b/Engine/Source/Programs/UnrealBuildAccelerator/Visualizer/Private/UbaTraceReader.h
--- a/Engine/Source/Programs/UnrealBuildAccelerator/Visualizer/Private/UbaTraceReader.h
+++ b/Engine/Source/Programs/UnrealBuildAccelerator/Visualizer/Private/UbaTraceReader.h
@@ -280,6 +280,8 @@ namespace uba
 
 		#if PLATFORM_WINDOWS
 		HANDLE m_hostProcess = NULL;
+		#else
+		pid_t m_hostProcess = 0;
 		#endif
 	};
 }
\ No newline at end of file
