[UBA]
* Implemented most of visualizer for Mac. Things missing are popup menu, network backend and scrollbars

diff --git a/Engine/Source/Programs/UnrealBuildAccelerator/Visualizer/Private/Mac/UbaVisualizerApp.mm b/Engine/Source/Programs/UnrealBuildAccelerator/Visualizer/Private/Mac/UbaVisualizerApp.mm
--- a/Engine/Source/Programs/UnrealBuildAccelerator/Visualizer/Private/Mac/UbaVisualizerApp.mm
+++ b/Engine/Source/Programs/UnrealBuildAccelerator/Visualizer/Private/Mac/UbaVisualizerApp.mm
@@ -5,14 +5,22 @@
 #include "UbaFile.h"
 #include "UbaVisualizer.h"
 
-// If you want to call into C++ code, declare your C++ functions here:
-#ifdef __cplusplus
-extern "C" {
-#endif
-void cpp_render_example(float width, float height);
-#ifdef __cplusplus
+namespace uba { class VisualizerMac; }
+
+@interface VisualizerCanvasView : NSView
+{
+    NSTrackingArea* _tracking;
+	NSTimer* _tickTimer;
+	uba::VisualizerMac* _visualizer;
+    CGFloat _magnifyAccum;
+    CGFloat _magnifyTotal;
+	bool _mouseInside;
 }
-#endif
+- (void)startTimer;
+- (void)stopTimer;
+- (void)tick:(NSTimer *)timer;
+@end
+
 
 namespace uba
 {
@@ -25,8 +33,7 @@
 		{
 			size_t w, h;
 			size_t bytesPerRow;
-			CGColorSpaceRef cs;
-			CGContextRef bmpCtx;
+			CGContextRef ctx;
 			CGImageRef img;
 			void* pixels;
 		};
@@ -226,27 +233,27 @@ virtual void DrawCached(TraceView::BitmapCache& process, RECT rect, int offsetY,
 							b.w = 256;
 							b.h = BitmapCacheHeight;
 							b.bytesPerRow = b.w * 4;
-							b.cs = CGColorSpaceCreateDeviceRGB();
 							b.pixels = calloc(b.h, b.bytesPerRow);
 							CGBitmapInfo info = kCGBitmapByteOrder32Little | kCGImageAlphaPremultipliedFirst;
-							b.bmpCtx = CGBitmapContextCreate(b.pixels, b.w, b.h, 8, b.bytesPerRow, b.cs, info);
+							b.ctx = CGBitmapContextCreate(b.pixels, b.w, b.h, 8, b.bytesPerRow, owner.colorSpaceRef, info);
 
 							owner.m_lastBitmapOffset = 0;
 							owner.m_lastBitmap = &b;
 						}
-						process.bitmap = (HBITMAP)owner.m_lastBitmap;
+						process.bitmap = owner.m_lastBitmap;
 						process.bitmapOffset = owner.m_lastBitmapOffset;
 						owner.m_lastBitmapOffset += processFont.height;
 					}
 					if (lastSelectedBitmap != (Bitmap*)process.bitmap)
 						lastSelectedBitmap = (Bitmap*)process.bitmap;
 
-					PaintContextMac textCacheContext{ owner, lastSelectedBitmap->bmpCtx, { 256, BitmapCacheHeight } };
+					PaintContextMac textCacheContext{ owner, lastSelectedBitmap->ctx, { 256, BitmapCacheHeight } };
 					textCacheContext.SetFont(activeFont);
 					func(textCacheContext);
 
-					CGImageRelease(lastSelectedBitmap->img);
-					lastSelectedBitmap->img = CGBitmapContextCreateImage(lastSelectedBitmap->bmpCtx);
+					if (lastSelectedBitmap->img)
+						CGImageRelease(lastSelectedBitmap->img);
+					lastSelectedBitmap->img = nullptr;
 				}
 
 				if (lastSelectedBitmap != (Bitmap*)process.bitmap)
@@ -279,21 +286,53 @@ virtual void DrawCached(TraceView::BitmapCache& process, RECT rect, int offsetY,
 					bitmapOffsetX += diff;
 				}
 
-				CGImageRef img = lastSelectedBitmap->img;
-				CGRect destRect = CGRectMake(float(rect.left), float(rect.top + bltOffsetY), float(width), float(height));
-				CGRect sourceRect = CGRectMake(float(bitmapOffsetX), float(bitmapOffsetY), float(width), float(height));
+				BlitInfo info
+				{
+					*lastSelectedBitmap,
+					CGRectMake(float(rect.left), float(rect.top + bltOffsetY), float(width), float(height)),
+					CGRectMake(float(bitmapOffsetX), float(bitmapOffsetY), float(width), float(height)),
+				};
+
+				if (info.b.img)
+					Blit(info);
+				else
+					m_deferredBlits.emplace_back(info);
+			}
+
+			struct BlitInfo
+			{
+				Bitmap& b;
+				CGRect dest;
+				CGRect source;
+			};
 
+			void Blit(BlitInfo& info)
+			{
+				CGImageRef img = info.b.img;
 				CGContextSaveGState(ctx);
-				CGContextClipToRect(ctx, destRect);
-				CGFloat sx = destRect.size.width  / sourceRect.size.width;
-				CGFloat sy = destRect.size.height / sourceRect.size.height;
-				CGContextTranslateCTM(ctx, destRect.origin.x - sourceRect.origin.x * sx, destRect.origin.y - sourceRect.origin.y * sy);
+				CGContextClipToRect(ctx, info.dest);
+				CGFloat sx = info.dest.size.width  / info.source.size.width;
+				CGFloat sy = info.dest.size.height / info.source.size.height;
+				CGContextTranslateCTM(ctx, info.dest.origin.x - info.source.origin.x * sx, info.dest.origin.y - info.source.origin.y * sy);
 				CGContextScaleCTM(ctx, sx, sy);
 				CGRect full = CGRectMake(0, 0, (CGFloat)256, (CGFloat)BitmapCacheHeight); // Draw whole image; clip ensures only desired region shows
 				CGContextDrawImage(ctx, full, img);
 				CGContextRestoreGState(ctx);
 			}
 
+			Vector<BlitInfo> m_deferredBlits;
+
+			virtual void EndCached() override
+			{
+				for (auto& blit : m_deferredBlits)
+				{
+					if (!blit.b.img)
+						blit.b.img = CGBitmapContextCreateImage(blit.b.ctx);
+					Blit(blit);
+				}
+				m_deferredBlits.clear();
+			}
+
 			virtual POINT GetCursorClientPos() override
 			{
 				return owner.m_mousePos;
@@ -329,6 +368,17 @@ virtual void Reset() override
 			for (auto c : colors)
 				CFRelease(c);
 			colors.clear();
+
+			for (auto b : m_textBitmaps)
+			{
+				CFRelease(b->ctx);
+				CFRelease(b->img);
+				free(b->pixels);
+				delete b;
+			}
+			m_textBitmaps.clear();
+			m_lastBitmap = 0;
+			m_lastBitmapOffset = BitmapCacheHeight;
 		}
 		virtual void Resize(int width, int height) override
 		{
@@ -385,17 +435,27 @@ virtual void DeleteFont(FontHandle font) override
 	class VisualizerMac : public Visualizer
 	{
 	public:
-		VisualizerMac(VisualizerConfig& config, Logger& logger, VisualizerBackend& backend) : Visualizer(config, logger, backend)
+		VisualizerMac(VisualizerConfig& config, Logger& logger, VisualizerBackend& backend, VisualizerCanvasView* view) : Visualizer(config, logger, backend)
 		{
-			if (false)
+			m_view = view;
+			//m_config.boxHeight = 20;
+
+			UpdateDefaultFont();
+			UpdateProcessFont();
+			m_popupFont.handle = m_backend.CreateFont(TC("Menlo"), 16, false);
+			m_popupFont.height = 14;
+
+			NSArray<NSString *> *args = [[NSProcessInfo processInfo] arguments];
+			if (args.count > 1)
 			{
-				char path[512];
-				GetFullPathNameW("~/Trace.uba", 512, path, nullptr);
+				auto filePath = [args[1] UTF8String];
+				char fixedPath[512];
+				GetFullPathNameW(filePath, 512, fixedPath, nullptr);
 
-				m_fileName.Append(path);
-				if (true)
+				m_fileName.Append(fixedPath);
+				if (false)
 				{
-					m_trace.ReadFile(m_traceView, path, true);
+					m_trace.ReadFile(m_traceView, fixedPath, true);
 
 					m_startTime = GetTime();
 
@@ -409,16 +469,12 @@ virtual void DeleteFont(FontHandle font) override
 					m_replay = 1;
 				}
 				else
-					m_trace.ReadFile(m_traceView, path, false);
+					m_trace.ReadFile(m_traceView, fixedPath, false);
+			}
+			else
+			{
+				m_listenThread.Start([this]() { ThreadListen(); return 0; });
 			}
-			//m_config.boxHeight = 20;
-
-			UpdateDefaultFont();
-			UpdateProcessFont();
-			m_popupFont.handle = m_backend.CreateFont(TC("Menlo"), 16, false);
-			m_popupFont.height = 14;
-
-			m_listenThread.Start([this]() { ThreadListen(); return 0; });
 		}
 		virtual void PaintEmptyScreen(PaintContext& context) override
 		{
@@ -426,6 +482,14 @@ virtual void PaintEmptyScreen(PaintContext& context) override
 			//NSRect rect = NSMakeRect(50, 50, 200, 100);
 			//NSRectFill(rect);
 		}
+		virtual void Redraw(bool now) override
+		{
+			[m_view setNeedsDisplay:YES];
+		}
+		virtual void StopUpdate() override
+		{
+			[m_view stopTimer];
+		}
 		virtual void FinishIncompatible() override
 		{
 		}
@@ -434,7 +498,7 @@ bool UpdateSelection() override
 			return Visualizer::UpdateSelection(((VisualizerBackendMac&)m_backend).m_mousePos);
 		}
 
-		void Paint(SIZE size)
+		void Paint(CGContextRef ctx, SIZE size)
 		{
 			auto& backend = (VisualizerBackendMac&)m_backend;
 
@@ -444,7 +508,6 @@ void Paint(SIZE size)
 			m_clientWidth = size.cx;
 			m_clientHeight = size.cy;
 
-			auto ctx = [[NSGraphicsContext currentContext] CGContext];
 			uba::VisualizerBackendMac::PaintContextMac context(backend, ctx, size);
 			RECT r { 0, 0, size.cx, size.cy };
 
@@ -452,29 +515,24 @@ void Paint(SIZE size)
 			PaintAll(context, r);
 		}
 
-		void Update()
+		void PostNewTrace(const TString& newTrace)
 		{
-			//if (!m_started && m_hasTrace)
-			{
-				SCOPED_FUTEX(m_newTraceLock, lock);
-				if (!m_newTraceName.Equals(m_newTrace))
-				{
-					m_newTraceName.Clear().Append(m_newTrace);
-					m_usingNamed = true;
-					//m_hasTrace = true;
-					//m_started = true;
-					NewTrace(0, false);
-				}
-			}
+			m_newTraceName.Clear().Append(newTrace);
+			m_usingNamed = true;
+			NewTrace(0, false);
+			[m_view startTimer];
+		}
 
-			UpdateTrace();
+		bool Update()
+		{
+			return UpdateTrace();
 		}
 
 		void MouseMove(CGFloat x, CGFloat y)
 		{
 			((VisualizerBackendMac&)m_backend).m_mousePos =  { (int)x, (int)y };
 			if (UpdateSelection() || m_config.showCursorLine)
-				;//Redraw(false);
+				Redraw(false);
 		}
 
 		void MouseDown(CGFloat x, CGFloat y)
@@ -493,6 +551,25 @@ void MouseDown(CGFloat x, CGFloat y)
 				UpdateScrollbars(true);
 				Redraw(false);
 			}
+			else if (!m_hyperLinkSelected.empty())
+			{
+				NSString* s = [NSString stringWithUTF8String:m_hyperLinkSelected.c_str()];
+				if (NSURL* url = [NSURL URLWithString:s])
+					[[NSWorkspace sharedWorkspace] openURL:url];
+			}
+			else if (m_sessionSelectedIndex != ~0u)
+			{
+				auto& session = m_traceView.sessions[m_sessionSelectedIndex];
+				for (auto& info : session.infos)
+				{
+					if (info.hyperlink.empty())
+						continue;
+					NSString* s = [NSString stringWithUTF8String:info.hyperlink.c_str()];
+					if (NSURL* url = [NSURL URLWithString:s])
+						[[NSWorkspace sharedWorkspace] openURL:url];
+					break;
+				}
+			}
 		}
 
 		bool ThreadListen()
@@ -508,6 +585,7 @@ bool ThreadListen()
 
 			m_listenTimeout.Create(EventResetType_Auto);
 			
+			TString newTrace;
 			StringBuffer<256> traceName;
 			while (true)
 			{
@@ -519,46 +597,40 @@ bool ThreadListen()
 				}
 				if (traceName.count)
 				{
-					if (!traceName.Equals(m_newTrace))
+					if (!traceName.Equals(newTrace))
 					{
-						SCOPED_FUTEX(m_newTraceLock, lock);
-						m_newTrace = traceName.ToString();
+						newTrace = traceName.ToString();
+						dispatch_sync(dispatch_get_main_queue(), ^{
+							PostNewTrace(newTrace);
+						});
 					}
 				}
-				else
+				else if (!newTrace.empty())
 				{
-					SCOPED_FUTEX(m_newTraceLock, lock);
-					m_newTrace.clear();
+					dispatch_sync(dispatch_get_main_queue(), ^{
+							PostNewTrace("");
+					});
+					newTrace.clear();
 				}
 				Sleep(500);
 			}
 			return true;
 		}
 
+		VisualizerCanvasView* m_view;
 		Event m_listenTimeout;
 		Thread m_listenThread;
-		Futex m_newTraceLock;
-		TString m_newTrace;
 	};
 }
 
-@interface MyCanvasView : NSView
-{
-	NSTimeInterval _time;
-    NSTrackingArea* _tracking;
-	uba::VisualizerMac* _visualizer;
-}
-- (void)tick:(NSTimer *)timer;
-@end
 
-@implementation MyCanvasView
+@implementation VisualizerCanvasView
 
 
 - (instancetype)initWithFrame:(NSRect)frame {
 	self = [super initWithFrame:frame];
-	if (self) {
-		_time = 0.0;
-
+	if (self)
+	{
 		uba::VisualizerConfig config("~/UbaVisualizerConfig.toml");
 		config.boxHeight = 14;
 		config.horizontalScaleValue = 1.0f;
@@ -566,34 +638,35 @@ - (instancetype)initWithFrame:(NSRect)frame {
 		auto logger = new uba::LoggerWithWriter(uba::g_consoleLogWriter);
 		auto backend = new uba::VisualizerBackendMac;
 
-		_visualizer = new uba::VisualizerMac(config, *logger, *backend);
-
-		[NSTimer scheduledTimerWithTimeInterval:(1.0/60.0) target:self selector:@selector(tick:) userInfo:nil repeats:YES];
+		_visualizer = new uba::VisualizerMac(config, *logger, *backend, self);
 	}
 	return self;
 }
 
-- (void)tick:(NSTimer *)timer {
-	_time += 1.0 / 60.0;
-
-	_visualizer->Update();
-
-	// Request a redraw
-	[self setNeedsDisplay:YES];
+- (void)startTimer {
+	if (_tickTimer) return;
+	_tickTimer = [NSTimer scheduledTimerWithTimeInterval:(1.0/60.0) target:self selector:@selector(tick:) userInfo:nil repeats:YES];
 }
 
+- (void)stopTimer {
+	if (!_tickTimer) return;
+	[_tickTimer invalidate];
+	_tickTimer = nil;
+}
 
-
-- (BOOL)isFlipped {
-	// Use top-left origin like many 2D APIs (optional)
-	return YES;
+- (void)tick:(NSTimer *)timer {
+	if (_visualizer->Update())
+		[self setNeedsDisplay:YES];
 }
 
+- (BOOL)isFlipped { return YES; } // top-left is 0,0
+
 - (void)drawRect:(NSRect)dirtyRect {
    [super drawRect:dirtyRect];
+	CGContextRef ctx = [[NSGraphicsContext currentContext] CGContext];
 	NSRect bounds = self.bounds;
 	uba::SIZE size = { (int)bounds.size.width, (int)bounds.size.height };
-	_visualizer->Paint(size);
+	_visualizer->Paint(ctx, size);
 }
 
 - (void)updateTrackingAreas {
@@ -605,7 +678,7 @@ - (void)updateTrackingAreas {
         _tracking = nil;
     }
 
-    NSTrackingAreaOptions opts = NSTrackingMouseMoved | NSTrackingActiveInKeyWindow | NSTrackingInVisibleRect;
+    NSTrackingAreaOptions opts = NSTrackingMouseMoved | NSTrackingActiveInKeyWindow | NSTrackingInVisibleRect | NSTrackingMouseEnteredAndExited;
 
     _tracking = [[NSTrackingArea alloc] initWithRect:NSZeroRect options:opts owner:self userInfo:nil];
     [self addTrackingArea:_tracking];
@@ -619,8 +692,22 @@ - (void)viewDidMoveToWindow {
 }
 
 - (void)mouseMoved:(NSEvent*)event {
-    NSPoint p = [self convertPoint:event.locationInWindow fromView:nil];
-	_visualizer->MouseMove(p.x, p.y);
+	if (_mouseInside)
+	{
+		NSPoint p = [self convertPoint:event.locationInWindow fromView:nil];
+		_visualizer->MouseMove(p.x, p.y);
+	}
+}
+
+- (void)mouseEntered:(NSEvent*)event
+{
+    _mouseInside = true;
+}
+
+- (void)mouseExited:(NSEvent*)event
+{
+    _mouseInside = false;
+	_visualizer->UnselectAndRedraw();
 }
 
 - (void)mouseDown:(NSEvent*)event
@@ -628,6 +715,59 @@ - (void)mouseDown:(NSEvent*)event
     NSPoint p = [self convertPoint:event.locationInWindow fromView:nil];
 	_visualizer->MouseDown(p.x, p.y);
 }
+
+- (void)scrollWheel:(NSEvent *)event
+{
+	using namespace uba;
+    CGFloat dx = event.scrollingDeltaX;
+    CGFloat dy = event.scrollingDeltaY;
+	_visualizer->m_scrollPosY = Min(0.0f, _visualizer->m_scrollPosY + float(dy));
+	_visualizer->m_scrollPosX = Min(0.0f, _visualizer->m_scrollPosX + float(dx));
+	_visualizer->UpdateSelection();
+	[self setNeedsDisplay:YES];
+}
+
+- (void)magnifyWithEvent:(NSEvent *)event
+{
+	using namespace uba;
+
+	if (event.phase == NSEventPhaseBegan)
+	{
+		_magnifyAccum = 0.0;
+		_magnifyTotal = 0.0;
+	}
+
+	_magnifyAccum += event.magnification;
+	_magnifyTotal += event.magnification;
+
+	const CGFloat step = 0.10;
+
+	bool changed = false;
+	while (_magnifyAccum >= step) {
+		_magnifyAccum -= step;
+		_visualizer->m_config.boxHeight = u32(Min(int(_visualizer->m_config.boxHeight + 1), 20));
+		changed = true;
+	}
+	while (_magnifyAccum <= -step) {
+		_magnifyAccum += step;
+		_visualizer->m_config.boxHeight = u32(Max(int(_visualizer->m_config.boxHeight - 1), 2));
+		changed = true;
+	}
+
+	if (changed)
+	{
+		_visualizer->UpdateProcessFont();
+		_visualizer->DirtyBitmaps(true);
+		_visualizer->UpdateSelection();
+		[self setNeedsDisplay:YES];
+	}
+
+	if (event.phase == NSEventPhaseEnded || event.phase == NSEventPhaseCancelled)
+		_magnifyAccum = 0.0;
+
+    //CGFloat magnification = event.magnification;
+	//_visualizer->m_zoomValue += float(magnification);
+}
 @end
 
 @interface AppDelegate : NSObject <NSWindowDelegate>
@@ -663,33 +803,18 @@ int main(int argc, const char * argv[]) {
 		[appMenuItem setSubmenu:appMenu];
 		// --- end menu setup ---
 
-		// Create the window
 		NSRect frame = NSMakeRect(200, 200, 800, 600);
-		NSWindowStyleMask style =
-			NSWindowStyleMaskTitled |
-			NSWindowStyleMaskClosable |
-			NSWindowStyleMaskResizable |
-			NSWindowStyleMaskMiniaturizable;
-
-		NSWindow *window = [[NSWindow alloc]
-							initWithContentRect:frame
-							styleMask:style
-							backing:NSBackingStoreBuffered
-							defer:NO];
-
-		AppDelegate *delegate = [[AppDelegate alloc] init];
+		NSWindowStyleMask style = NSWindowStyleMaskTitled | NSWindowStyleMaskClosable | NSWindowStyleMaskResizable | NSWindowStyleMaskMiniaturizable;
+		NSWindow* window = [[NSWindow alloc] initWithContentRect:frame styleMask:style backing:NSBackingStoreBuffered defer:NO];
+
+		AppDelegate* delegate = [[AppDelegate alloc] init];
 		[window setDelegate:delegate];
 		[window setTitle:@"UbaVisualizer"];
 
-		// Attach our canvas view
-		MyCanvasView *view = [[MyCanvasView alloc] initWithFrame:frame];
+		VisualizerCanvasView *view = [[VisualizerCanvasView alloc] initWithFrame:frame];
 		[window setContentView:view];
 		[window makeKeyAndOrderFront:nil];
-
-		// Make sure app shows in Dock and gets focus
-		[NSApp activateIgnoringOtherApps:YES];
-
-		// Enter the Cocoa event loop
+		[NSApp activateIgnoringOtherApps:YES]; // Make sure app shows in Dock and gets focus
 		[NSApp run];
 	}
 	return 0;
diff --git a/Engine/Source/Programs/UnrealBuildAccelerator/Visualizer/Private/UbaVisualizer.cpp b/Engine/Source/Programs/UnrealBuildAccelerator/Visualizer/Private/UbaVisualizer.cpp
--- a/Engine/Source/Programs/UnrealBuildAccelerator/Visualizer/Private/UbaVisualizer.cpp
+++ b/Engine/Source/Programs/UnrealBuildAccelerator/Visualizer/Private/UbaVisualizer.cpp
@@ -1589,6 +1589,8 @@ namespace uba
 			context.SetFont(m_defaultFont);
 		}
 
+		context.EndCached();
+
 		if (true)
 		{
 			Font& activeFont = context.GetFont();
@@ -2939,5 +2941,11 @@ namespace uba
 		return true;
 	}
 
+	void Visualizer::UnselectAndRedraw()
+	{
+		if (Unselect() || m_config.showCursorLine)
+			Redraw(false);
+	}
+
 	////////////////////////////////////////////////////////////////////////////////////////////////////
 }
diff --git a/Engine/Source/Programs/UnrealBuildAccelerator/Visualizer/Private/UbaVisualizer.h b/Engine/Source/Programs/UnrealBuildAccelerator/Visualizer/Private/UbaVisualizer.h
--- a/Engine/Source/Programs/UnrealBuildAccelerator/Visualizer/Private/UbaVisualizer.h
+++ b/Engine/Source/Programs/UnrealBuildAccelerator/Visualizer/Private/UbaVisualizer.h
@@ -57,6 +57,7 @@ namespace uba
 		void UpdateProcessFont();
 		void DirtyBitmaps(bool full);
 		bool UpdateSelection(POINT pos);
+		void UnselectAndRedraw();
 
 		enum VisualizerFlag
 		{
diff --git a/Engine/Source/Programs/UnrealBuildAccelerator/Visualizer/Private/UbaVisualizerBackend.h b/Engine/Source/Programs/UnrealBuildAccelerator/Visualizer/Private/UbaVisualizerBackend.h
--- a/Engine/Source/Programs/UnrealBuildAccelerator/Visualizer/Private/UbaVisualizerBackend.h
+++ b/Engine/Source/Programs/UnrealBuildAccelerator/Visualizer/Private/UbaVisualizerBackend.h
@@ -53,6 +53,7 @@ namespace uba
 			virtual void PolyPolyline(const POINT* points, const u32* lines, u32 lineCount) = 0;
 			virtual SIZE GetTextSize(StringView str) = 0;
 			virtual void DrawCached(TraceView::BitmapCache& process, RECT rect, int offsetY, int textHeight, Font& processFont, RECT clippingRect, const Function<void(PaintContext& textCacheContext)>& func) = 0;
+			virtual void EndCached() {}
 
 			virtual POINT GetCursorClientPos() = 0;
 		};
diff --git a/Engine/Source/Programs/UnrealBuildAccelerator/Visualizer/Private/Windows/UbaVisualizerWin.cpp b/Engine/Source/Programs/UnrealBuildAccelerator/Visualizer/Private/Windows/UbaVisualizerWin.cpp
--- a/Engine/Source/Programs/UnrealBuildAccelerator/Visualizer/Private/Windows/UbaVisualizerWin.cpp
+++ b/Engine/Source/Programs/UnrealBuildAccelerator/Visualizer/Private/Windows/UbaVisualizerWin.cpp
@@ -551,12 +551,6 @@ namespace uba
 		CloseClipboard();
 	}
 
-	void VisualizerWin::UnselectAndRedraw()
-	{
-		if (Unselect() || m_config.showCursorLine)
-			Redraw(false);
-	}
-
 	bool VisualizerWin::UpdateSelection()
 	{
 		if (!m_mouseOverWindow || m_dragToScrollCounter > 0)
diff --git a/Engine/Source/Programs/UnrealBuildAccelerator/Visualizer/Private/Windows/UbaVisualizerWin.h b/Engine/Source/Programs/UnrealBuildAccelerator/Visualizer/Private/Windows/UbaVisualizerWin.h
--- a/Engine/Source/Programs/UnrealBuildAccelerator/Visualizer/Private/Windows/UbaVisualizerWin.h
+++ b/Engine/Source/Programs/UnrealBuildAccelerator/Visualizer/Private/Windows/UbaVisualizerWin.h
@@ -31,7 +31,6 @@ namespace uba
 		virtual void FinishIncompatible() override;
 
 		void CopyTextToClipboard(const TString& str);
-		void UnselectAndRedraw();
 		virtual bool UpdateSelection() override;
 		virtual void UpdateScrollbars(bool redraw) override;
 		void ThreadLoop();
