[UBA]
* Added "-delete" which can be used to fast delete hierarchies of directories. Since windows/robocopy etc are so slow this is way faster when needing to delete millions of folders/files

diff --git a/Engine/Source/Programs/UnrealBuildAccelerator/Cli/Private/UbaCli.cpp b/Engine/Source/Programs/UnrealBuildAccelerator/Cli/Private/UbaCli.cpp
--- a/Engine/Source/Programs/UnrealBuildAccelerator/Cli/Private/UbaCli.cpp
+++ b/Engine/Source/Programs/UnrealBuildAccelerator/Cli/Private/UbaCli.cpp
@@ -8,6 +8,7 @@
 #include "UbaCompressedFileHeader.h"
 #include "UbaConfig.h"
 #include "UbaCoordinatorWrapper.h"
+#include "UbaDirectoryIterator.h"
 #include "UbaFileAccessor.h"
 #include "UbaNetworkBackendTcp.h"
 #include "UbaPathUtils.h"
@@ -144,6 +145,52 @@ namespace uba
 	
 	StringBuffer<> g_rootDir(DefaultRootDir);
 
+	struct DirToDelete
+	{
+		Atomic<u32> counter;
+		DirToDelete* parent = nullptr;
+		TString path;
+	};
+
+	void DeleteDir(Logger& logger, WorkManagerImpl& workManager, DirToDelete* dir, Atomic<u64>& deleteCount)
+	{
+		StringBuffer<> fullName;
+		fullName.Append(dir->path).Append(PathSeparator);
+		u32 dirLen = fullName.count;
+
+		TraverseDir(logger, dir->path, [&](const DirectoryEntry& e)
+			{
+				fullName.Resize(dirLen).Append(e.name);
+				if (IsDirectory(e.attributes))
+				{
+					++dir->counter;
+					auto child = new DirToDelete{1, dir, fullName.ToString()};
+					workManager.AddWork([&, child](const WorkContext&)
+						{
+							DeleteDir(logger, workManager, child, deleteCount);
+						}, 1, TC("D"));
+				}
+				else
+				{
+					DeleteFileW(fullName.data);
+					++deleteCount;
+					if ((deleteCount.load() % 100000) == 0)
+						logger.Info(TC("Deleted: %llu"), deleteCount.load());
+				}
+			});
+
+		for (auto it=dir;it;)
+		{
+			if (--it->counter)
+				return;
+			RemoveDirectoryW(it->path.c_str());
+			++deleteCount;
+			auto parent = it->parent;
+			delete it;
+			it = parent;
+		}
+	}
+
 	bool WrappedMain(int argc, tchar* argv[])
 	{
 		using namespace uba;
@@ -175,6 +222,7 @@ namespace uba
 		bool getCas = false;
 		bool listImports = false;
 		bool deleteCas = false;
+		bool fastDelete = false;
 		bool enableStdOut = true;
 		bool printSummary = false;
 		bool detailedTrace = false;
@@ -478,6 +526,10 @@ namespace uba
 					if ((g_rootDir.count = GetFullPathNameW(value.Replace('/', PathSeparator).data, g_rootDir.capacity, g_rootDir.data, nullptr)) == 0)
 						return PrintHelp(StringBuffer<>().Appendf(TC("-dir has invalid path %s"), g_rootDir.data).data);
 				}
+				else if (name.Equals(TCV("-delete")))
+				{
+					fastDelete = true;
+				}
 				else if (name.Equals(TCV("-?")))
 				{
 					return PrintHelp(TC(""));
@@ -611,6 +663,14 @@ namespace uba
 
 		bool exit = false;
 
+		if (fastDelete)
+		{
+			WorkManagerImpl workManager(maxProcessCount, TC("UbaWrk/DelCas"));
+			Atomic<u64> deleteCount = 0;
+			DeleteDir(logger, workManager, new DirToDelete{1, nullptr, g_rootDir.ToString()}, deleteCount);
+			return true;
+		}
+
 		if (!addCas.empty())
 		{
 			WorkManagerImpl workManager(maxProcessCount);
@@ -1445,6 +1505,7 @@ namespace uba
 			info.cacheClients = cacheClients;
 			info.cacheClientCount = cacheClient ? 1 : 0;
 			info.writeToCache = populateCache;
+			//info.writeTrackedInputs = true;
 			Scheduler scheduler(info);
 
 			//agentCount = 2;
