[UBA]
* Changed so weight calculations are not happening when reusing processes. Before this change we could end up restarting processes more than needed

diff --git a/Engine/Source/Programs/UnrealBuildAccelerator/Common/Private/UbaScheduler.cpp b/Engine/Source/Programs/UnrealBuildAccelerator/Common/Private/UbaScheduler.cpp
--- a/Engine/Source/Programs/UnrealBuildAccelerator/Common/Private/UbaScheduler.cpp
+++ b/Engine/Source/Programs/UnrealBuildAccelerator/Common/Private/UbaScheduler.cpp
@@ -747,7 +747,7 @@ namespace uba
 			}
 			*/
 
-			ExitProcess(*info, *handle.m_process, handle.m_process->GetExitCode(), false);
+			ExitProcess(*info, *handle.m_process, handle.m_process->GetExitCode(), false, false);
 			return;
 		}
 
@@ -756,7 +756,7 @@ namespace uba
 
 		m_finishedRemoteProcessThreads.emplace_back([this, handle = handle, info]()
 			{
-				ExitProcess(*info, *handle.m_process, handle.m_process->GetExitCode(), false);
+				ExitProcess(*info, *handle.m_process, handle.m_process->GetExitCode(), false, false);
 				UBA_ASSERT(info->refCount);
 				if (!--info->refCount)
 					delete info;
@@ -825,7 +825,7 @@ namespace uba
 		return entryToProcess;
 	}
 
-	u32 Scheduler::PopProcess(bool isLocal, RunStatus& outPrevStatus)
+	u32 Scheduler::PopProcess(bool isLocal, RunStatus& outPrevStatus, bool forReuse)
 	{
 		float activeWeight = m_activeLocalProcessWeight;
 		float maxWeight = float(m_maxLocalProcessors);
@@ -863,10 +863,13 @@ namespace uba
 			{
 				if (m_forceRemote && entry.canExecuteRemotely)
 					continue;
-				if (runStatus == RunStatus_QueuedForRun && atMaxLocalWeight)
-					continue;
-				if (activeWeight > 1.0f && activeWeight + entry.info->weight > maxWeight + 1.0f)
-					continue;
+				if (!forReuse)
+				{
+					if (runStatus == RunStatus_QueuedForRun && atMaxLocalWeight)
+						continue;
+					if (activeWeight > 1.0f && activeWeight + entry.info->weight > maxWeight + 1.0f)
+						continue;
+				}
 			}
 			else
 			{
@@ -904,7 +907,8 @@ namespace uba
 				if (!LocalProcessOkToSpawn())
 					return ~0u;
 
-				m_activeLocalProcessWeight += entry.info->weight;
+				if (!forReuse)
+					m_activeLocalProcessWeight += entry.info->weight;
 			}
 
 			outPrevStatus = entry.runStatus;
@@ -1002,7 +1006,7 @@ namespace uba
 						exitInfo.startInfo = info;
 						exitInfo.isLocal = true;
 						exitInfo.processIndex = entryToProcess;
-						ExitProcess(exitInfo, *process, 0, true);
+						ExitProcess(exitInfo, *process, 0, true, false);
 						++m_cacheHitCount;
 					}
 					else
@@ -1034,7 +1038,7 @@ namespace uba
 		{
 			RunStatus prevStatus;
 			SCOPED_FUTEX(m_processEntriesLock, lock);
-			u32 indexToRun = PopProcess(isLocal, prevStatus);
+			u32 indexToRun = PopProcess(isLocal, prevStatus, false);
 			if (indexToRun == ~0u)
 			{
 				lock.Leave();
@@ -1151,7 +1155,13 @@ namespace uba
 		if (!ei) // If null, process has already exited from some other thread
 			return false;
 
-		ExitProcess(*ei, process, prevExitCode, false);
+		bool isLocal = !process.IsRemote();
+
+		bool forReuse = true;
+		float weight = isLocal ? ei->startInfo->weight : 0.0f;
+		auto weightGuard = MakeGuard([&]() { m_activeLocalProcessWeight -= weight; });
+
+		ExitProcess(*ei, process, prevExitCode, false, forReuse);
 
 		ei->startInfo = nullptr;
 		ei->processIndex = ~0u;
@@ -1163,13 +1173,11 @@ namespace uba
 			ei->hasCheckedWasReturned = true;
 		}
 
-		bool isLocal = !process.IsRemote();
-
 		while (true)
 		{
 			RunStatus prevStatus;
 			SCOPED_FUTEX(m_processEntriesLock, lock);
-			u32 indexToRun = PopProcess(isLocal, prevStatus);
+			u32 indexToRun = PopProcess(isLocal, prevStatus, forReuse);
 			if (indexToRun == ~0u)
 				return false;
 			//UBA_ASSERT(prevStatus != RunStatus_QueuedForCache);
@@ -1221,11 +1229,13 @@ namespace uba
 			UBA_ASSERTF(temp1.Equals(temp2.data), TC("%s vs %s"), temp1.data, temp2.data);
 			#endif
 			
+			weightGuard.Cancel();
+
 			return true;
 		}
 	}
 
-	void Scheduler::ExitProcess(ExitProcessInfo& info, Process& process, u32 exitCode, bool fromCache)
+	void Scheduler::ExitProcess(ExitProcessInfo& info, Process& process, u32 exitCode, bool fromCache, bool forReuse)
 	{
 		auto si = info.startInfo;
 		if (!si)
@@ -1271,7 +1281,7 @@ namespace uba
 				entry.cacheStatus = CacheStatus_Success;
 				--m_activeCacheQueries;
 			}
-			else
+			else if (!forReuse)
 				m_activeLocalProcessWeight -= si->weight;
 		}
 
diff --git a/Engine/Source/Programs/UnrealBuildAccelerator/Common/Public/UbaScheduler.h b/Engine/Source/Programs/UnrealBuildAccelerator/Common/Public/UbaScheduler.h
--- a/Engine/Source/Programs/UnrealBuildAccelerator/Common/Public/UbaScheduler.h
+++ b/Engine/Source/Programs/UnrealBuildAccelerator/Common/Public/UbaScheduler.h
@@ -135,12 +135,12 @@ namespace uba
 		void ProcessExited(ExitProcessInfo* info, const ProcessHandle& handle);
 		bool CanRun();
 		u32 PopCache();
-		u32 PopProcess(bool isLocal, RunStatus& outPrevStatus);
+		u32 PopProcess(bool isLocal, RunStatus& outPrevStatus, bool forReuse);
 		bool RunCacheQuery();
 		bool RunQueuedProcess(bool isLocal);
 		bool RaceRemoteProcess();
 		bool HandleReuseMessage(Process& process, NextProcessInfo& outNextProcess, u32 prevExitCode);
-		void ExitProcess(ExitProcessInfo& info, Process& process, u32 exitCode, bool fromCache);
+		void ExitProcess(ExitProcessInfo& info, Process& process, u32 exitCode, bool fromCache, bool forReuse);
 		void SkipProcess(ProcessStartInfo2& info);
 		void UpdateStatusHitMissCount();
 		void UpdateQueueCounter(int offset);
