[UBA]
* Fixed so config is applied to storage proxy create info for UbaAgent

diff --git a/Engine/Source/Programs/UnrealBuildAccelerator/Agent/Private/UbaAgent.cpp b/Engine/Source/Programs/UnrealBuildAccelerator/Agent/Private/UbaAgent.cpp
--- a/Engine/Source/Programs/UnrealBuildAccelerator/Agent/Private/UbaAgent.cpp
+++ b/Engine/Source/Programs/UnrealBuildAccelerator/Agent/Private/UbaAgent.cpp
@@ -1236,6 +1236,7 @@ namespace uba
 
 					StorageProxyCreateInfo proxyInfo { *proxyServer, *client, storageServerUid, TC("Wooohoo"), storageClient };
 					proxyInfo.useLocalStorage = proxyUseLocalStorage;
+					proxyInfo.Apply(config);
 
 					storageProxy = new StorageProxy(proxyInfo);
 
