[UBA]
* Refactored code that handles deferred cas creation to handle situations where source can be deleted

diff --git a/Engine/Source/Programs/UnrealBuildAccelerator/Common/Private/UbaSession.cpp b/Engine/Source/Programs/UnrealBuildAccelerator/Common/Private/UbaSession.cpp
--- a/Engine/Source/Programs/UnrealBuildAccelerator/Common/Private/UbaSession.cpp
+++ b/Engine/Source/Programs/UnrealBuildAccelerator/Common/Private/UbaSession.cpp
@@ -2291,10 +2291,15 @@ namespace uba
 		if (registerToDir)
 			RegisterDeleteFile(key, fixedFilePath);
 
+		m_storage.DeleteExternalFileMapping(key, filePath);
+
 		SCOPED_FUTEX(m_virtualSourceFilesLock, virtualSourceFilesLock);
 		auto findIt = m_virtualSourceFiles.find(key); 
 		if (findIt == m_virtualSourceFiles.end())
+		{
+			//m_logger.Warning(TC("Failed to delete virtual file %s (%s)"), filePath, KeyToString(key).data);
 			return false;
+		}
 		VirtualSourceFile vsf = findIt->second;
 		m_virtualSourceFiles.erase(findIt);
 		virtualSourceFilesLock.Leave();
@@ -3866,6 +3871,7 @@ namespace uba
 		entry.createIndependentMapping = true; // We want independent mapping in casdb so they can be deleted
 		entry.isInvisible = transient;
 		entry.dropCasAfterUse = true;
+		entry.canBeFreed = false;
 
 		StringBuffer<> mappingName;
 		GetMappingString(mappingName, memoryHandle, mappingOffset, entry.canBeFreed);
diff --git a/Engine/Source/Programs/UnrealBuildAccelerator/Common/Private/UbaStorage.cpp b/Engine/Source/Programs/UnrealBuildAccelerator/Common/Private/UbaStorage.cpp
--- a/Engine/Source/Programs/UnrealBuildAccelerator/Common/Private/UbaStorage.cpp
+++ b/Engine/Source/Programs/UnrealBuildAccelerator/Common/Private/UbaStorage.cpp
@@ -9,6 +9,7 @@
 #include "UbaBinaryReaderWriter.h"
 #include "UbaDirectoryIterator.h"
 #include "UbaStorageUtils.h"
+#include "UbaTrace.h"
 #include "UbaWorkManager.h"
 
 #define UBA_CHECK_BUFFER_SLOTS 0
@@ -354,6 +355,45 @@ namespace uba
 		m_externalFileMappingsProvider = std::move(provider);
 	}
 
+	void StorageImpl::DeleteExternalFileMapping(StringKey fileNameKey, const tchar* fileName)
+	{
+		SCOPED_FUTEX(m_deferredCasCreationLookupLock, deferredLock);
+		auto it = m_deferredCasCreationLookupByName.find(fileNameKey);
+		if (it == m_deferredCasCreationLookupByName.end())
+			return;
+		CasKey key = it->second.key;
+		m_deferredCasCreationLookupByName.erase(it);
+		UBA_ASSERT(it->second.manualDelete);
+
+		auto it2 = m_deferredCasCreationLookup.find(key);
+		if (it2 == m_deferredCasCreationLookup.end())
+		{
+			UBA_ASSERT(false);
+			return;
+		}
+
+		DeferredCasCreation& dcc = it2->second;
+		++dcc.refCount;
+		deferredLock.Leave();
+
+		SCOPED_FUTEX(dcc.lock, dccLock);
+
+		for (auto i=dcc.sources.begin(), e=dcc.sources.end(); i!=e; ++i)
+		{
+			if (*i != it)
+				continue;
+			dcc.sources.erase(i);
+			break;
+		}
+
+		deferredLock.Enter();
+		if (!--dcc.refCount && dcc.sources.empty())
+		{
+			dccLock.Leave();
+			m_deferredCasCreationLookup.erase(it2);
+		}
+	}
+
 	bool StorageImpl::WriteCompressedFile(WriteResult& out, const tchar* from, FileHandle readHandle, u8* readMem, u64 fileSize, const tchar* to, const void* header, u64 headerSize, u64 lastWriteTime)
 	{
 		CompressWriter w(m_logger, m_sharedMemory, m_bufferSlots, m_workManager, Stats(), m_casCompressor, m_casCompressionLevel, m_asyncUnmapViewOfFile);
@@ -583,7 +623,7 @@ namespace uba
 	}
 
 
-	bool StorageImpl::AddCasFile(StringKey fileNameKey, const tchar* fileName, const CasKey& casKey, bool deferCreation, u64 fileSize, u64 lastWriteTime, CasEntry** outEntry)
+	bool StorageImpl::AddCasFile(StringKey fileNameKey, const tchar* fileName, const CasKey& casKey, bool deferCreation, u64 fileSize, u64 lastWriteTime, CasEntry** outEntry, bool dropCasAfterUse)
 	{
 		UBA_ASSERTF(IsCompressed(casKey) == m_storeCompressed, TC("CasKey compress mode must match storage compress mode. Compress mode Key: %i, Store: %i Cas: %s (%s)"), (IsCompressed(casKey) ? 1 : 0), (m_storeCompressed ? 1 : 0), CasKeyString(casKey).str, fileName);
 		CasEntry& casEntry = GetOrCreateCasEntry(casKey);
@@ -639,19 +679,31 @@ namespace uba
 		if (deferCreation)
 		{
 			SCOPED_FUTEX(m_deferredCasCreationLookupLock, deferredLock);
-			auto res = m_deferredCasCreationLookup.try_emplace(casKey);
-			if (res.second)
+
+			auto res2 = m_deferredCasCreationLookupByName.try_emplace(fileNameKey, casKey);
+			if (!res2.second)
 			{
-				DeferedCasCreation& dcc = res.first->second;
-				dcc.fileName = fileName;
-				dcc.fileSize = fileSize;
-				dcc.lastWriteTime = lastWriteTime;
+				if (res2.first->second.key != casKey)
+					return m_logger.Error(TC("Same file %s is registered with different caskey (%s vs %s)"), fileName, CasKeyString(res2.first->second.key).str, CasKeyString(casKey).str);
+				return true;
 			}
-			auto res2 = m_deferredCasCreationLookupByName.try_emplace(fileNameKey, casKey);
-			if (res2.second)
-				res.first->second.names.push_back(res2.first);
-			else if (res2.first->second != casKey)
-				return m_logger.Error(TC("Same file %s is registered with different caskey (%s vs %s)"), fileName, CasKeyString(res2.first->second).str, CasKeyString(casKey).str);
+			DeferredSource& source = res2.first->second;
+			source.manualDelete = dropCasAfterUse;
+			source.fileName = fileName;
+			source.fileSize = fileSize;
+			source.lastWriteTime = lastWriteTime;
+
+			auto res = m_deferredCasCreationLookup.try_emplace(casKey);
+
+			DeferredCasCreation& dcc = res.first->second;
+			++dcc.refCount;
+			deferredLock.Leave();
+
+			SCOPED_FUTEX(dcc.lock, dccLock);
+			dcc.sources.emplace_back(res2.first);
+
+			deferredLock.Enter();
+			--dcc.refCount;
 
 			if (outEntry)
 			{
@@ -669,21 +721,39 @@ namespace uba
 			SCOPED_FUTEX(m_deferredCasCreationLookupLock, deferredLock);
 			#if UBA_DEBUG
 			auto findIt = m_deferredCasCreationLookupByName.find(fileNameKey);
-			if (findIt != m_deferredCasCreationLookupByName.end() && findIt->second != casKey)
+			if (findIt != m_deferredCasCreationLookupByName.end() && findIt->second.key != casKey)
 				return m_logger.Error(TC("Found deferred entry for %s with different cas key while writing cas. A ReportFileWrite is missing somewhere"), fileName);
 			#endif
 
 			auto findIt2 = m_deferredCasCreationLookup.find(casKey);
 			if (findIt2 != m_deferredCasCreationLookup.end())
 			{
+				DeferredCasCreation& dcc = findIt2->second;
+				++dcc.refCount;
+				deferredLock.Leave();
+
 				//#if UBA_DEBUG
 				//return m_logger.Error(TC("This should not happen. (%s)"), fileName);
 				//#endif
 				m_logger.Info(TC("Investigate how we get here. (%s)"), fileName);
 
-				for (auto& it2 : findIt2->second.names)
-					m_deferredCasCreationLookupByName.erase(it2);
-				m_deferredCasCreationLookup.erase(findIt2);
+				SCOPED_FUTEX(dcc.lock, dccLock);
+				deferredLock.Enter();
+				for (auto i = dcc.sources.begin(); i!=dcc.sources.end();)
+				{
+					if ((*i)->second.manualDelete)
+					{
+						++i;
+						continue;
+					}
+					m_deferredCasCreationLookupByName.erase(*i);
+					i = dcc.sources.erase(i);
+				}
+				if (!--dcc.refCount && dcc.sources.empty())
+				{
+					dccLock.Leave();
+					m_deferredCasCreationLookup.erase(findIt2);
+				}
 			}
 		}
 
@@ -1743,6 +1813,8 @@ namespace uba
 			forKey.MakeLower();
 		StringKey fileNameKey = ToStringKey(forKey);
 
+		bool dropCasAfterUse = false;
+
 		FileEntry& fileEntry = GetOrCreateFileEntry(fileNameKey);
 		SCOPED_FUTEX(fileEntry.lock, entryLock);
 
@@ -1753,7 +1825,7 @@ namespace uba
 			if (fileEntry.casKey != CasKeyZero)
 			{
 				UBA_ASSERT(casKeyOverride == CasKeyZero || casKeyOverride == fileEntry.casKey);
- 				if (!AddCasFile(fileNameKey, fileName, fileEntry.casKey, deferCreation, fileEntry.size, fileEntry.lastWritten, outEntry))
+ 				if (!AddCasFile(fileNameKey, fileName, fileEntry.casKey, deferCreation, fileEntry.size, fileEntry.lastWritten, outEntry, dropCasAfterUse))
 					return false;
 			}
 			out = fileEntry.casKey;
@@ -1777,6 +1849,8 @@ namespace uba
 			fileSize = externalMapping.size;
 			storedSize = externalMapping.storedSize;
 			lastWriteTime = externalMapping.lastWriteTime;
+			dropCasAfterUse = externalMapping.dropCasAfterUse;
+			fileEntry.isTemp = true;
 		}
 		else
 		{
@@ -1812,14 +1886,14 @@ namespace uba
 			if (casKeyOverride != CasKeyZero && casKeyOverride != fileEntry.casKey)
 			{
 				fileEntry.casKey = casKeyOverride;
-				if (!AddCasFile(fileNameKey, fileName, fileEntry.casKey, deferCreation, storedSize, lastWriteTime, outEntry))
+				if (!AddCasFile(fileNameKey, fileName, fileEntry.casKey, deferCreation, storedSize, lastWriteTime, outEntry, dropCasAfterUse))
 					return false;
 				out = fileEntry.casKey;
 				return true;
 			}
 			if (fileSize == fileEntry.size && (lastWriteTime && lastWriteTime == fileEntry.lastWritten) && fileEntry.casKey != CasKeyInvalid)
 			{
-				if (!AddCasFile(fileNameKey, fileName, fileEntry.casKey, deferCreation, storedSize, lastWriteTime, outEntry))
+				if (!AddCasFile(fileNameKey, fileName, fileEntry.casKey, deferCreation, storedSize, lastWriteTime, outEntry, dropCasAfterUse))
 					return false;
 				out = fileEntry.casKey;
 				return true;
@@ -1871,7 +1945,7 @@ namespace uba
 		if (fileEntry.casKey == CasKeyZero)
 			return false;
 
-		if (!AddCasFile(fileNameKey, fileName, fileEntry.casKey, deferCreation, storedSize, lastWriteTime, outEntry))
+		if (!AddCasFile(fileNameKey, fileName, fileEntry.casKey, deferCreation, storedSize, lastWriteTime, outEntry, dropCasAfterUse))
 			return false;
 
 		out = fileEntry.casKey;
@@ -2068,12 +2142,39 @@ namespace uba
 		auto findIt = m_deferredCasCreationLookup.find(casKey);
 		if (findIt == m_deferredCasCreationLookup.end())
 			return false;
-		StringKey fileNameKey = findIt->second.fileNameKey;
-		StringBuffer<> deferredCreation(findIt->second.fileName);
+		DeferredCasCreation& dcc = findIt->second;
+		++dcc.refCount;
+		deferredLock.Leave();
+
+		SCOPED_FUTEX(dcc.lock, dccLock);
+
+		auto dccGuard = MakeGuard([&]()
+			{
+				deferredLock.Enter();
+				for (auto i = dcc.sources.begin(); i!=dcc.sources.end();)
+				{
+					if ((*i)->second.manualDelete)
+					{
+						++i;
+						continue;
+					}
+					m_deferredCasCreationLookupByName.erase(*i);
+					i = dcc.sources.erase(i);
+				}
+				if (!--dcc.refCount && dcc.sources.empty())
+				{
+					dccLock.Leave();
+					m_deferredCasCreationLookup.erase(findIt);
+				}
+			});
+
+		if (dcc.sources.empty())
+			return false;
+
+		auto& sourceIt = dcc.sources.front();
+		StringKey fileNameKey = sourceIt->first;
+		StringBuffer<> deferredCreation(sourceIt->second.fileName);
 		//u64 deferredFileSize = findIt->second.fileSize;
-		for (auto& it2 : findIt->second.names)
-			m_deferredCasCreationLookupByName.erase(it2);
-		m_deferredCasCreationLookup.erase(findIt);
 
 		if (casEntry.beingWritten)
 			m_logger.Warning(TC("Deferred cas %s is being written from network. This should never happen (%s)"), CasKeyString(casKey).str, deferredCreation.data);
@@ -2084,10 +2185,12 @@ namespace uba
 		// .. we need to add verification for this and throw an error
 
 		casEntry.verified = true;
-		deferredLock.Leave();
 		WriteCasResult res;
 		if (!WriteCasFile(res, deferredCreation.data, casKey, useMapping))
 			return m_logger.Error(TC("Failed to write deferred cas %s (%s)"), CasKeyString(casKey).str, deferredCreation.data);
+
+		dccGuard.Execute();
+
 		casEntry.dropped = res.dropAfterUse;
 
 		// TODO: THIS MISMATCH WITH compressed vs uncompressed
@@ -2670,7 +2773,7 @@ namespace uba
 
 	void StorageImpl::ReportFileWrite(StringKey fileNameKey, const tchar* fileName)
 	{
-		// If a defered cas creation is queued up while the source file is about to be modified we need to flush out the cas creation before modifying the file
+		// If a deferred cas creation is queued up while the source file is about to be modified we need to flush out the cas creation before modifying the file
 		SCOPED_FUTEX_READ(m_deferredCasCreationLookupLock, deferredLock);
 		auto findIt = m_deferredCasCreationLookupByName.find(fileNameKey);
 		if (findIt == m_deferredCasCreationLookupByName.end())
@@ -2678,7 +2781,7 @@ namespace uba
 
 		m_logger.Logf(IsDebug ? LogEntryType_Warning : LogEntryType_Info, TC("Flush cas creation of file %s (This should be rare)"), fileName);
 
-		CasKey key = findIt->second;
+		CasKey key = findIt->second.key;
 		deferredLock.Leave();
 		HasCasFile(key);
 	}
@@ -2723,7 +2826,20 @@ namespace uba
 
 	void StorageImpl::TraceUpdate(Trace& trace, u32 startRow)
 	{
+		#if 0
 		m_casDataBuffer.TraceStatus(trace, startRow, TC("Cas"), 1);
+		#endif
+
+		#if 0
+		SCOPED_FUTEX(m_deferredCasCreationLookupLock, lock);
+		u64 byCas = m_deferredCasCreationLookup.size();
+		u64 byName = m_deferredCasCreationLookupByName.size();
+		lock.Leave();
+
+		StringBuffer<> str;
+		str.Appendf(TC("DeferredCount  Cas: %llu Name: %llu"), byCas, byName);
+		trace.StatusUpdate(startRow, 7, str, LogEntryType_Info);
+		#endif
 	}
 
 	void StorageImpl::PrintSummary(Logger& logger)
diff --git a/Engine/Source/Programs/UnrealBuildAccelerator/Common/Private/UbaStorageServer.cpp b/Engine/Source/Programs/UnrealBuildAccelerator/Common/Private/UbaStorageServer.cpp
--- a/Engine/Source/Programs/UnrealBuildAccelerator/Common/Private/UbaStorageServer.cpp
+++ b/Engine/Source/Programs/UnrealBuildAccelerator/Common/Private/UbaStorageServer.cpp
@@ -663,7 +663,8 @@ namespace uba
 						}
 
 						bool deferCreation = false;
-						if (!AddCasFile(fileNameKey, hint.data, casKey, deferCreation, InvalidValue, 0, nullptr))
+						bool dropCasAfterUse = false;
+						if (!AddCasFile(fileNameKey, hint.data, casKey, deferCreation, InvalidValue, 0, nullptr, dropCasAfterUse))
 						{
 							m_logger.Error(TC("FetchBegin failed for cas file %s (%s) requested by client %u. Can't add cas file to database"), CasKeyString(casKey).str, hint.data, connectionInfo.GetId());
 							writer.WriteU16(0);
diff --git a/Engine/Source/Programs/UnrealBuildAccelerator/Common/Public/UbaStorage.h b/Engine/Source/Programs/UnrealBuildAccelerator/Common/Public/UbaStorage.h
--- a/Engine/Source/Programs/UnrealBuildAccelerator/Common/Public/UbaStorage.h
+++ b/Engine/Source/Programs/UnrealBuildAccelerator/Common/Public/UbaStorage.h
@@ -110,6 +110,7 @@ namespace uba
 		};
 		using ExternalFileMappingsProvider = Function<bool(ExternalFileMapping& out, StringKey fileNameKey, const tchar* fileName)>;
 		virtual void RegisterExternalFileMappingsProvider(ExternalFileMappingsProvider&& provider) {}
+		virtual void DeleteExternalFileMapping(StringKey fileNameKey, const tchar* fileName) {}
 	};
 
 	struct StorageCreateInfo
@@ -219,11 +220,12 @@ namespace uba
 		};
 
 		virtual void RegisterExternalFileMappingsProvider(ExternalFileMappingsProvider&& provider) override;
+		virtual void DeleteExternalFileMapping(StringKey fileNameKey, const tchar* fileName) override;
 		virtual bool WriteCompressedFile(WriteResult& out, const tchar* from, FileHandle readHandle, u8* readMem, u64 fileSize, const tchar* toFile, const void* header, u64 headerSize, u64 lastWriteTime = 0) final;
 		bool WriteCasFileNoCheck(WriteCasResult& out, const StringKey& fileNameKey, const tchar* fileName, const CasKey& casKey, const tchar* casFile, bool storeCompressed, bool useMapping);
 		bool WriteCasFile(WriteCasResult& out, const tchar* fileName, const CasKey& casKey, bool useMapping);
 		bool VerifyExisting(bool& outReturnValue, CasEntryWriteLock& casEntryLock, const CasKey& casKey, CasEntry& casEntry, StringBufferBase& casFile);
-		bool AddCasFile(StringKey fileNameKey, const tchar* fileName, const CasKey& casKey, bool deferCreation, u64 fileSize, u64 lastWriteTime, CasEntry** outEntry);
+		bool AddCasFile(StringKey fileNameKey, const tchar* fileName, const CasKey& casKey, bool deferCreation, u64 fileSize, u64 lastWriteTime, CasEntry** outEntry, bool dropCasAfterUse);
 		void CasEntryAccessed(const CasKey& casKey);
 		virtual bool IsDisallowedPath(const tchar* fileName);
 		virtual bool DecompressMemoryToMemory(const u8* compressedData, u64 compressedSize, u8* writeData, u64 decompressedSize, const tchar* readHint, const tchar* writeHint) override;
@@ -350,9 +352,10 @@ namespace uba
 		FileMappingBuffer m_casDataBuffer;
 
 		Futex m_deferredCasCreationLookupLock;
-		using NameToCas = UnorderedMap<StringKey, CasKey>;
-		struct DeferedCasCreation { StringKey fileNameKey; TString fileName; List<NameToCas::iterator> names; u64 fileSize = InvalidValue; u64 lastWriteTime = 0;};
-		UnorderedMap<CasKey, DeferedCasCreation> m_deferredCasCreationLookup;
+		struct DeferredSource { CasKey key; TString fileName; u64 fileSize = InvalidValue; u64 lastWriteTime = 0; bool manualDelete = false; };
+		using NameToCas = UnorderedMap<StringKey, DeferredSource>;
+		struct DeferredCasCreation { Futex lock; List<NameToCas::iterator> sources; u32 refCount = 0; };
+		UnorderedMap<CasKey, DeferredCasCreation> m_deferredCasCreationLookup;
 		NameToCas m_deferredCasCreationLookupByName;
 
 		DirectoryCache m_dirCache;
