[UBA]
* Improved assert description

diff --git a/Engine/Source/Programs/UnrealBuildAccelerator/Core/Private/UbaSharedMemoryView.cpp b/Engine/Source/Programs/UnrealBuildAccelerator/Core/Private/UbaSharedMemoryView.cpp
--- a/Engine/Source/Programs/UnrealBuildAccelerator/Core/Private/UbaSharedMemoryView.cpp
+++ b/Engine/Source/Programs/UnrealBuildAccelerator/Core/Private/UbaSharedMemoryView.cpp
@@ -208,7 +208,7 @@ namespace uba
 		void* viewA = mmap(targetAddress, size, protect, MAP_SHARED | MAP_FIXED, handle.shmFd, offset);
 		#endif
 
-		UBA_ASSERTF(viewA == targetAddress, TC("MapViewOfFile3 failed - %s (%s)"), LastErrorToText().data, hint);
+		UBA_ASSERTF(viewA == targetAddress, TC("MapViewOfFile3 failed extending memory with %llu bytes at offset %llu - %s (%s)"), size, offset, LastErrorToText().data, hint);
 		if (!viewA)
 			return false;
 
