[UBA]
* Fixed shared memory allocator temp file leak when not wanting to leak it

diff --git a/Engine/Source/Programs/UnrealBuildAccelerator/Common/Private/UbaSharedMemoryAllocator.cpp b/Engine/Source/Programs/UnrealBuildAccelerator/Common/Private/UbaSharedMemoryAllocator.cpp
--- a/Engine/Source/Programs/UnrealBuildAccelerator/Common/Private/UbaSharedMemoryAllocator.cpp
+++ b/Engine/Source/Programs/UnrealBuildAccelerator/Common/Private/UbaSharedMemoryAllocator.cpp
@@ -101,6 +101,10 @@ namespace uba
 		if (m_fileMapping.IsValid() && !keepFileMapping)
 			CloseFileMapping(m_logger, m_fileMapping, TC("SharedMemoryAllocator"));
 		m_fileMapping = {};
+
+		if (m_tempFile != InvalidFileHandle && !keepFileMapping)
+			CloseFile(TC(""), m_tempFile);
+		m_tempFile = InvalidFileHandle;
 	}
 
 	SharedMemoryAllocator::operator SharedMemoryAllocatorHandle() const
