[UBA]
* Fixed log message when un-dropping cas

diff --git a/Engine/Source/Programs/UnrealBuildAccelerator/Common/Private/UbaStorage.cpp b/Engine/Source/Programs/UnrealBuildAccelerator/Common/Private/UbaStorage.cpp
--- a/Engine/Source/Programs/UnrealBuildAccelerator/Common/Private/UbaStorage.cpp
+++ b/Engine/Source/Programs/UnrealBuildAccelerator/Common/Private/UbaStorage.cpp
@@ -651,7 +651,7 @@ namespace uba
 		// File is requested by server itself to provide to someone remote. We have to un-drop the cas in case it was marked as dropped
 		if (casEntry.dropped && m_allowUnDropOfCas && !outEntry)
 		{
-			m_logger.Info(TC("Un-dropped cas for %s (%s). Cas will leak in memory"), fileName, BytesToText(fileSize).str);
+			m_logger.Info(TC("Un-dropped cas for %s (%s)."), fileName, BytesToText(casEntry.size).str);
 			casEntry.dropped = false;
 		}
 
