[UBA]
* Fixed visualizer bug with cached text when zooming in

diff --git a/Engine/Source/Programs/UnrealBuildAccelerator/Visualizer/Private/Windows/UbaVisualizerBackendD2D.h b/Engine/Source/Programs/UnrealBuildAccelerator/Visualizer/Private/Windows/UbaVisualizerBackendD2D.h
--- a/Engine/Source/Programs/UnrealBuildAccelerator/Visualizer/Private/Windows/UbaVisualizerBackendD2D.h
+++ b/Engine/Source/Programs/UnrealBuildAccelerator/Visualizer/Private/Windows/UbaVisualizerBackendD2D.h
@@ -291,7 +291,8 @@ namespace uba
 			for (auto bm : m_textBitmaps)
 				bm->Release();
 			m_textBitmaps.clear();
-
+			m_lastBitmap = 0;
+			m_lastBitmapOffset = BitmapCacheHeight;
 		}
 
 		virtual void Resize(int width, int height) override
diff --git a/Engine/Source/Programs/UnrealBuildAccelerator/Visualizer/Private/Windows/UbaVisualizerWin.cpp b/Engine/Source/Programs/UnrealBuildAccelerator/Visualizer/Private/Windows/UbaVisualizerWin.cpp
--- a/Engine/Source/Programs/UnrealBuildAccelerator/Visualizer/Private/Windows/UbaVisualizerWin.cpp
+++ b/Engine/Source/Programs/UnrealBuildAccelerator/Visualizer/Private/Windows/UbaVisualizerWin.cpp
@@ -757,6 +757,7 @@ namespace uba
 
 					float horizontalScaleValue = m_config.horizontalScaleValue;
 					float newScaleValue = horizontalScaleValue;
+					int oldBoxHeight = m_config.boxHeight;
 					int newBoxHeight = m_config.boxHeight;
 
 					if (controlDown)
@@ -801,8 +802,9 @@ namespace uba
 					//if (!m_traceView.finished && m_scrollPosX <= minScroll)
 					//	m_autoScroll = true;
 
-
-					if (m_config.ShowReadWriteColors)
+					if (oldBoxHeight < newBoxHeight)
+						DirtyBitmaps(true);
+					else if (m_config.ShowReadWriteColors)
 						for (auto& session : m_traceView.sessions)
 							for (auto& processor : session.processors)
 								for (auto& process : processor.processes)
