[UBA]
* Changed so d2d backend is used by default in visualizer

diff --git a/Engine/Source/Programs/UnrealBuildAccelerator/Visualizer/Private/UbaVisualizerConfig.cpp b/Engine/Source/Programs/UnrealBuildAccelerator/Visualizer/Private/UbaVisualizerConfig.cpp
--- a/Engine/Source/Programs/UnrealBuildAccelerator/Visualizer/Private/UbaVisualizerConfig.cpp
+++ b/Engine/Source/Programs/UnrealBuildAccelerator/Visualizer/Private/UbaVisualizerConfig.cpp
@@ -45,7 +45,7 @@ namespace uba
 		UBA_VISUALIZER_FLAGS2
 		#undef UBA_VISUALIZER_FLAG
 
-		config.GetValueAsBool(useD2D, TC("UseD2D"));
+		config.GetValueAsBool(useGDI, TC("UseGDI"));
 
 		fontSize = Min(fontSize, 30u);
 		return true;
@@ -73,7 +73,7 @@ namespace uba
 		UBA_VISUALIZER_FLAGS2
 		#undef UBA_VISUALIZER_FLAG
 
-		config.AddValue(TC("UseD2D"), useD2D);
+		config.AddValue(TC("UseGDI"), useGDI);
 
 		return config.SaveToFile(logger, filename.c_str());
 	}
diff --git a/Engine/Source/Programs/UnrealBuildAccelerator/Visualizer/Private/UbaVisualizerConfig.h b/Engine/Source/Programs/UnrealBuildAccelerator/Visualizer/Private/UbaVisualizerConfig.h
--- a/Engine/Source/Programs/UnrealBuildAccelerator/Visualizer/Private/UbaVisualizerConfig.h
+++ b/Engine/Source/Programs/UnrealBuildAccelerator/Visualizer/Private/UbaVisualizerConfig.h
@@ -31,7 +31,7 @@ namespace uba
 		UBA_VISUALIZER_FLAG(ScaleHorizontalWithScrollWheel, false, "Use scroll wheel to scale horizontally") \
 		UBA_VISUALIZER_FLAG(DarkMode, false, "Use dark mode to draw visualizer") \
 		UBA_VISUALIZER_FLAG(AutoSaveSettings, true, "Auto save Position/Settings on close") \
-		UBA_VISUALIZER_FLAG(ShowAllTraces, true, "Show all traces started on channe") \
+		UBA_VISUALIZER_FLAG(ShowAllTraces, true, "Show all traces started on channel") \
 		UBA_VISUALIZER_FLAG(SortActiveRemoteSessions, true, "Sort active sessions on top") \
 		UBA_VISUALIZER_FLAG(AutoScaleHorizontal, true, "Automatically scale horizontally to fit processes") \
 		UBA_VISUALIZER_FLAG(LockTimelineToBottom, true, "Lock timeline to always paint at bottom") \
@@ -55,7 +55,7 @@ namespace uba
 		u32 maxActiveProcessHeight = 16;
 		u32 boxHeight = 12;
 		float horizontalScaleValue = 0.5f;
-		bool useD2D = false;
+		bool useGDI = false;
 
 		#define UBA_VISUALIZER_FLAG(name, defaultValue, desc) bool show##name = defaultValue;
 		UBA_VISUALIZER_FLAGS1
diff --git a/Engine/Source/Programs/UnrealBuildAccelerator/Visualizer/Private/Windows/UbaVisualizerApp.cpp b/Engine/Source/Programs/UnrealBuildAccelerator/Visualizer/Private/Windows/UbaVisualizerApp.cpp
--- a/Engine/Source/Programs/UnrealBuildAccelerator/Visualizer/Private/Windows/UbaVisualizerApp.cpp
+++ b/Engine/Source/Programs/UnrealBuildAccelerator/Visualizer/Private/Windows/UbaVisualizerApp.cpp
@@ -323,9 +323,9 @@ int WINAPI wWinMain(_In_ HINSTANCE hInstance, _In_opt_ HINSTANCE hPrevInstance,
 	visualizerConfig.ShowAllTraces = showAllTraces;
 	visualizerConfig.Load(logger);
 	if (useGDI)
-		visualizerConfig.useD2D = false;
+		visualizerConfig.useGDI = true;
 	if (useD2D)
-		visualizerConfig.useD2D = true;
+		visualizerConfig.useGDI = false;
 
 	NetworkBackendTcp networkBackend(logWriter);
 	VisualizerWin visualizer(visualizerConfig, logger);
diff --git a/Engine/Source/Programs/UnrealBuildAccelerator/Visualizer/Private/Windows/UbaVisualizerWin.cpp b/Engine/Source/Programs/UnrealBuildAccelerator/Visualizer/Private/Windows/UbaVisualizerWin.cpp
--- a/Engine/Source/Programs/UnrealBuildAccelerator/Visualizer/Private/Windows/UbaVisualizerWin.cpp
+++ b/Engine/Source/Programs/UnrealBuildAccelerator/Visualizer/Private/Windows/UbaVisualizerWin.cpp
@@ -58,7 +58,7 @@ namespace uba
 
 	VisualizerBackend& CreateBackend(VisualizerConfig& config, VisualizerWin& v)
 	{
-		if (config.useD2D)
+		if (!config.useGDI)
 		{
 			auto d2d = new VisualizerBackendD2D(v);
 			if (d2d->Create())
