[UBA]
* Fixed so when compiling UbaVisualizer files are outputted to Binaries/Mac/UnrealBuildAccelerator folder on mac

diff --git a/Engine/Source/Programs/UnrealBuildAccelerator/UbaAgent.Target.cs b/Engine/Source/Programs/UnrealBuildAccelerator/UbaAgent.Target.cs
--- a/Engine/Source/Programs/UnrealBuildAccelerator/UbaAgent.Target.cs
+++ b/Engine/Source/Programs/UnrealBuildAccelerator/UbaAgent.Target.cs
@@ -107,6 +107,11 @@ public class UbaAgentTarget : TargetRules
 				BinaryName = $"{BinaryPrefix}{Rules.LaunchModuleName}{BinaryExt}";
 				Rules.OutputFile = Path.Combine(BinariesFolder, BinaryName);
 			}
+			else
+			{
+				Rules.ExeBinariesSubFolder = "UnrealBuildAccelerator";
+				//Rules.OutputFile = Path.Combine(BinariesFolder, "UbaVisualizer");
+			}
 			Rules.GlobalDefinitions.AddRange(new string[] {
 				"UBA_DETOURS_LIBRARY=\"libUbaDetours.dylib\"",
 				"UBA_AGENT_EXECUTABLE=\"UbaAgent\"",
