[UBA]
* Fixed so Process::PollStdPipes poll again if it found data to read. Should fix slowness when processes are spewing out lots of text

diff --git a/Engine/Source/Programs/UnrealBuildAccelerator/Common/Private/UbaProcess.cpp b/Engine/Source/Programs/UnrealBuildAccelerator/Common/Private/UbaProcess.cpp
--- a/Engine/Source/Programs/UnrealBuildAccelerator/Common/Private/UbaProcess.cpp
+++ b/Engine/Source/Programs/UnrealBuildAccelerator/Common/Private/UbaProcess.cpp
@@ -2505,37 +2505,42 @@ namespace uba
 
 		PipeReader* pipeReaders[] = { &outReader, &errReader };
 
-		pollfd plist[] = { {m_stdOutPipe,POLLIN, 0}, {m_stdErrPipe,POLLIN, 0} };
-		int rval = poll(plist, sizeof_array(plist), timeoutMs);
-		if (rval < 0)
+		while (true)
 		{
-			#if PLATFORM_MAC
-			m_session.m_logger.Error(TC("pipe polling error with %i (%s)"), rval, strerror(errno));
-			#endif
-			return false;
-		}
+			pollfd plist[] = { {m_stdOutPipe,POLLIN, 0}, {m_stdErrPipe,POLLIN, 0} };
+			int rval = poll(plist, sizeof_array(plist), timeoutMs);
+			if (rval < 0)
+			{
+				#if PLATFORM_MAC
+				m_session.m_logger.Error(TC("pipe polling error with %i (%s)"), rval, strerror(errno));
+				#endif
+				return false;
+			}
 
-		bool hasRead = false;
-		for (int i=0;i!=2;++i)
-		{
-			if (plist[i].revents & POLLERR) // If there is an error on any of them we hang up
-				return m_session.m_logger.Error(TC("pipe polling error"));
+			bool hasRead = false;
+			for (int i=0;i!=2;++i)
+			{
+				if (plist[i].revents & POLLERR) // If there is an error on any of them we hang up
+					return m_session.m_logger.Error(TC("pipe polling error"));
 
-			if (!(plist[i].revents & POLLIN))
-				continue;
+				if (!(plist[i].revents & POLLIN))
+					continue;
 
-			char buffer[1024];
-			if (int bytesRead = read(plist[i].fd, buffer, sizeof_array(buffer) - 1))
-			{
-				hasRead = true;
-				buffer[bytesRead] = 0;
-				pipeReaders[i]->ReadData(buffer, bytesRead);
+				char buffer[1024];
+				if (int bytesRead = read(plist[i].fd, buffer, sizeof_array(buffer) - 1))
+				{
+					hasRead = true;
+					buffer[bytesRead] = 0;
+					pipeReaders[i]->ReadData(buffer, bytesRead);
+				}
 			}
-		}
 
-		if (!hasRead)
+			if (hasRead)
+				continue;
 			if (plist[0].revents & POLLHUP && plist[1].revents & POLLHUP)
 				return false;
+			break;
+		}
 		pipeGuard.Cancel();
 		return true;
 	}
