[UBA]
* Fixed crash in print callstacks on win32 when more than 512 modules are loaded

diff --git a/Engine/Source/Programs/UnrealBuildAccelerator/Core/Private/UbaPlatform.cpp b/Engine/Source/Programs/UnrealBuildAccelerator/Core/Private/UbaPlatform.cpp
--- a/Engine/Source/Programs/UnrealBuildAccelerator/Core/Private/UbaPlatform.cpp
+++ b/Engine/Source/Programs/UnrealBuildAccelerator/Core/Private/UbaPlatform.cpp
@@ -352,9 +352,17 @@ namespace uba
 		struct ModuleRec { u64 start; u64 size; HMODULE handle; u32 index; };
 		Map<u64, ModuleRec> moduleEndAddresses;
 
-		HMODULE loadedModules[512];
+		Vector<HMODULE> loadedModules(128);
 		DWORD needed = 0;
-		if (EnumProcessModules(GetCurrentProcess(), loadedModules, sizeof(loadedModules), &needed))
+		bool success = EnumProcessModules(GetCurrentProcess(), loadedModules.data(), u32(loadedModules.size()*sizeof(HMODULE)), &needed);
+		if (success && needed > loadedModules.size()*sizeof(HMODULE))
+		{
+			loadedModules.resize(needed/sizeof(HMODULE));
+			success = EnumProcessModules(GetCurrentProcess(), loadedModules.data(), u32(loadedModules.size()*sizeof(HMODULE)), &needed);
+			if (needed > loadedModules.size()*sizeof(HMODULE))
+				success = false;
+		}
+		if (success)
 		{
 			for (u32 i = 0, e = needed / sizeof(HMODULE); i != e; ++i)
 			{
