[UBA]
* Changed disconnect warning in network server to be an info instead. Can't figure out how this timeout can happen but with winsock and times it feels like anything can happen (especially on extremely saturated aws instances)

diff --git a/Engine/Source/Programs/UnrealBuildAccelerator/Common/Private/UbaNetworkServer.cpp b/Engine/Source/Programs/UnrealBuildAccelerator/Common/Private/UbaNetworkServer.cpp
--- a/Engine/Source/Programs/UnrealBuildAccelerator/Common/Private/UbaNetworkServer.cpp
+++ b/Engine/Source/Programs/UnrealBuildAccelerator/Common/Private/UbaNetworkServer.cpp
@@ -217,8 +217,12 @@ namespace uba
 
 			if (!m_disconnectCallbackCalled.IsSet(30*1000)) // This should never time out!
 			{
-				m_server.m_logger.Warning(TC("Disconnect callback event timed out. This should never happen!!"));
+				m_backend.SetDisconnectCallback(m_backendConnection, nullptr, [](void* context, const Guid& connectionUid, void* connection) {});
+				m_backend.SetRecvCallbacks(m_backendConnection, nullptr, 0, [](void*, const Guid&, u8*, void*&, u8*&, u32&) { return false; }, nullptr, TC("Stopped"));
+				m_server.m_logger.Info(TC("Disconnect callback event timed out. This should never happen!!"));
+				#if UBA_DEBUG
 				PrintAllCallstacks(m_server.m_logger);
+				#endif
 			}
 
 			return true;
