[UBA]
* Added thread for polling stdout/err on non-windows platforms

diff --git a/Engine/Source/Programs/UnrealBuildAccelerator/Common/Private/UbaProcess.cpp b/Engine/Source/Programs/UnrealBuildAccelerator/Common/Private/UbaProcess.cpp
--- a/Engine/Source/Programs/UnrealBuildAccelerator/Common/Private/UbaProcess.cpp
+++ b/Engine/Source/Programs/UnrealBuildAccelerator/Common/Private/UbaProcess.cpp
@@ -552,7 +552,7 @@ namespace uba
 		}
 	}
 
-	bool ProcessImpl::WaitForRead(PipeReader& outReader, PipeReader& errReader)
+	bool ProcessImpl::WaitForRead()
 	{
 		// This number could be bigger but
 		// 1. It would delay ctrl-c
@@ -564,10 +564,6 @@ namespace uba
 			if (m_readEvent.IsSet(readEventTimeout))
 				return !m_cancelled;
 
-			#if !PLATFORM_WINDOWS
-			PollStdPipes(outReader, errReader, 0);
-			#endif
-
 			if (!IsActive())
 				return m_readEvent.IsSet(0); // Do one more check
 
@@ -618,11 +614,20 @@ namespace uba
 			{
 				m_session.ProcessNativeCreated(*this);
 
+				#if !PLATFORM_WINDOWS
 				PipeReader outReader(*this, LogEntryType_Info);
 				PipeReader errReader(*this, LogEntryType_Error);
+				Atomic<bool> pollStdDone = false;
+				Thread pollStdThread([&]()
+				{
+					while (!pollStdDone)
+						PollStdPipes(outReader, errReader, 500);
+					return true;
+				}, TC("PollStd"));
+				#endif
 
 				bool loop = m_detourEnabled;
-				while (loop && WaitForRead(outReader, errReader))
+				while (loop && WaitForRead())
 				{
 					u64 startTime = GetTime();
 					BinaryReader reader(comMemory, 0, comMemorySize);
@@ -634,6 +639,8 @@ namespace uba
 				}
 
 				#if !PLATFORM_WINDOWS
+				pollStdDone = true;
+				pollStdThread.Wait();
 				while (PollStdPipes(outReader, errReader, 500) && !IsCancelled())
 					continue;
 				#endif
diff --git a/Engine/Source/Programs/UnrealBuildAccelerator/Common/Public/UbaProcess.h b/Engine/Source/Programs/UnrealBuildAccelerator/Common/Public/UbaProcess.h
--- a/Engine/Source/Programs/UnrealBuildAccelerator/Common/Public/UbaProcess.h
+++ b/Engine/Source/Programs/UnrealBuildAccelerator/Common/Public/UbaProcess.h
@@ -77,7 +77,7 @@ namespace uba
 		bool IsCancelled();
 		bool HasFailedMessage();
 
-		bool WaitForRead(PipeReader& outReader, PipeReader& errReader);
+		bool WaitForRead();
 		void SetWritten();
 
 		void ThreadRun(void* environment);
