[UBA]
* Fixed bug where non-known dlls in system32 were not copied if they existed on the remote. This caused remotes to use wrong version of vcruntime dlls causing new dxcompiler.dll to fail to load dll when remote helper has old window

diff --git a/Engine/Source/Programs/UnrealBuildAccelerator/Common/Private/UbaSession.cpp b/Engine/Source/Programs/UnrealBuildAccelerator/Common/Private/UbaSession.cpp
--- a/Engine/Source/Programs/UnrealBuildAccelerator/Common/Private/UbaSession.cpp
+++ b/Engine/Source/Programs/UnrealBuildAccelerator/Common/Private/UbaSession.cpp
@@ -1725,11 +1725,11 @@ namespace uba
 
 		FixPath(applicationName, nullptr, 0, temp3.Clear());
 
-		bool isSystem = StartsWith(applicationName, m_systemPath.data);
-		if (isSystem && IsKnownSystemFile(applicationName))
+		bool isInSystemFolder = StartsWith(applicationName, m_systemPath.data);
+		if (isInSystemFolder && IsKnownSystemFile(applicationName))
 			return true;
 
-		BinaryModule& binaryModule = out.emplace_back(library, temp3.ToString(), attr, isSystem);(void)binaryModule;
+		BinaryModule& binaryModule = out.emplace_back(library, temp3.ToString(), attr, isInSystemFolder);(void)binaryModule;
 
 		StringBuffer<> errorStr;
 		BinaryInfo info;
@@ -3732,12 +3732,15 @@ namespace uba
 		//m_debugLogger->Info(TC("Mapping check (%s)"), fileName.data);
 		#endif
 
+		auto wg = MakeGuard([&]() { writer.WriteU32(GetFileMappingSize()); });
+
 		MemoryMap out;
 		u64 alignment = GetMemoryMapAlignment(fileName);
 		bool canBeFreed = true;
+		if (process.IsCancelled())
+			return false;
 		if (!GetOrCreateMemoryMapFromFile(out, fileNameKey, fileName.data, false, alignment, TC("Remap"), &process, canBeFreed))
 			return m_logger.Error(TC("Failed to remap %s"), fileName.data);
-		writer.WriteU32(GetFileMappingSize());
 		return true;
 	}
 
diff --git a/Engine/Source/Programs/UnrealBuildAccelerator/Common/Private/UbaSessionClient.cpp b/Engine/Source/Programs/UnrealBuildAccelerator/Common/Private/UbaSessionClient.cpp
--- a/Engine/Source/Programs/UnrealBuildAccelerator/Common/Private/UbaSessionClient.cpp
+++ b/Engine/Source/Programs/UnrealBuildAccelerator/Common/Private/UbaSessionClient.cpp
@@ -402,7 +402,8 @@ namespace uba
 		}
 
 		List<ModuleInfo> modules;
-		if (!ReadModules(modules, 0, application))
+		StringBuffer<512> virtualApplication;
+		if (!ReadModules(modules, virtualApplication, process.GetId(), application))
 			return false;
 
 		StringBuffer<MaxPath> applicationDir;
@@ -442,7 +443,7 @@ namespace uba
 
 		if (!isAbsolute)
 		{
-			appEnv.virtualApplication = modules.front().name;
+			appEnv.virtualApplication = virtualApplication.ToString();
 			startInfo.applicationStr = appEnv.virtualApplication;
 			startInfo.application = startInfo.applicationStr.c_str();
 		}
@@ -450,7 +451,7 @@ namespace uba
 		return true;
 	}
 
-	bool SessionClient::ReadModules(List<ModuleInfo>& outModules, u32 processId, const tchar* application)
+	bool SessionClient::ReadModules(List<ModuleInfo>& outModules, StringBufferBase& outVirtualApplication, u32 processId, const tchar* application)
 	{
 		TrackWorkScope tws(m_client, AsView(TC("ReadModules")), ColorWork);
 
@@ -474,7 +475,7 @@ namespace uba
 			StringBuffer<> moduleFile;
 			reader.ReadString(moduleFile);
 			u32 fileAttributes = reader.ReadU32();
-			bool isSystem = reader.ReadBool();
+			bool isInSystemFolder = reader.ReadBool();
 
 			CasKey casKey = reader.ReadCasKey();
 			if (casKey == CasKeyZero)
@@ -486,14 +487,23 @@ namespace uba
 				return m_logger.Error(TC("%s has min os version %u but current os is %u"), moduleFile.data, minOsVersion, m_osVersion);
 			#endif
 
-			if (isSystem)
+			// These are files in system folder that is not known to system
+			if (isInSystemFolder)
 			{
 				StringBuffer<> localSystemModule;
 				localSystemModule.Append(m_systemPath).Append(moduleFile.data + serverSystemPathLen);
-				if (FileExists(m_logger, localSystemModule.data) && !localSystemModule.EndsWith(TCV(".exe")))
+
+				// For non-windows we always use the local system file if it exists.
+				#if !PLATFORM_WINDOWS
+				if (FileExists(m_logger, localSystemModule.data))
 					continue;
-				moduleFile.Clear().Append(localSystemModule);
+				#endif
+				if (outModules.empty())
+					outVirtualApplication.Append(localSystemModule);
 			}
+			else if (outModules.empty())
+				outVirtualApplication.Append(moduleFile);
+			
 			outModules.emplace_back(moduleFile.data, casKey, fileAttributes);
 		}
 
diff --git a/Engine/Source/Programs/UnrealBuildAccelerator/Common/Public/UbaSessionClient.h b/Engine/Source/Programs/UnrealBuildAccelerator/Common/Public/UbaSessionClient.h
--- a/Engine/Source/Programs/UnrealBuildAccelerator/Common/Public/UbaSessionClient.h
+++ b/Engine/Source/Programs/UnrealBuildAccelerator/Common/Public/UbaSessionClient.h
@@ -88,7 +88,7 @@ namespace uba
 		struct ModuleInfo;
 
 		bool GetCasKeyForFile(CasKey& out, u32 processId, const StringView& fileName, const StringKey& fileNameKey);
-		bool ReadModules(List<ModuleInfo>& outModules, u32 processId, const tchar* application);
+		bool ReadModules(List<ModuleInfo>& outModules, StringBufferBase& outVirtualApplication, u32 processId, const tchar* application);
 		bool EnsureBinaryFile(StringBufferBase& out, StringBufferBase& outVirtual, u32 processId, StringView fileName, const StringKey& fileNameKey, StringView applicationDir, StringView workingDir, const u8* loaderPaths, u32 loaderPathsSize);
 		bool WriteBinFile(StringBufferBase& out, const StringView& binaryName, const CasKey& casKey, const KeyToString& applicationDir, u32 fileAttributes);
 		bool SendFiles(ProcessImpl& process, Timer& sendFiles);
diff --git a/Engine/Source/Programs/UnrealBuildAccelerator/Core/Public/UbaWinBinDependencyParser.h b/Engine/Source/Programs/UnrealBuildAccelerator/Core/Public/UbaWinBinDependencyParser.h
--- a/Engine/Source/Programs/UnrealBuildAccelerator/Core/Public/UbaWinBinDependencyParser.h
+++ b/Engine/Source/Programs/UnrealBuildAccelerator/Core/Public/UbaWinBinDependencyParser.h
@@ -21,6 +21,8 @@ namespace uba
 	};
 
 	// These dlls should _always_ exist on all machines so we can filter them out from the list of imports to copy
+	// It should somewhat match what is found in Computer\HKEY_LOCAL_MACHINE\SYSTEM\CurrentControlSet\Control\Session Manager\KnownDLLs
+	// but also take info from wine source code which dlls that are custom implemented for wine.
 	static const wchar_t* g_knownSystemFiles[] = // Sorted lowercase
 	{
 		L"advapi32.dll",
@@ -67,7 +69,9 @@ namespace uba
 		L"sechost.dll",
 		L"setupapi.dll",
 		L"shell32.dll",
+		L"shlwapi.dll",
 		L"sspicli.dll",
+		L"tbs.dll",
 		L"ucrtbase.dll",
 		L"umpdc.dll",
 		L"umppc17706.dll",
