[UBA]
* Fixed tcp connection shutdown race when closing send thread. Could access memory that was deleted

diff --git a/Engine/Source/Programs/UnrealBuildAccelerator/Common/Private/UbaNetworkBackendTcp.cpp b/Engine/Source/Programs/UnrealBuildAccelerator/Common/Private/UbaNetworkBackendTcp.cpp
--- a/Engine/Source/Programs/UnrealBuildAccelerator/Common/Private/UbaNetworkBackendTcp.cpp
+++ b/Engine/Source/Programs/UnrealBuildAccelerator/Common/Private/UbaNetworkBackendTcp.cpp
@@ -985,15 +985,13 @@ namespace uba
 
 		auto closeGuard = MakeGuard([&]()
 			{
-				for (;first;first=first->next)
-					first->sent.Set();
+				while (first) { auto cur = first; first = first->next; cur->sent.Set(); }
 				SCOPED_FUTEX(context.blockQueueLock, queueLock);
 				first = context.firstBlock;
 				context.firstBlock = (SendBlock*)~0ull;
 				context.lastBlock = nullptr;
 				queueLock.Leave();
-				for (; first; first = first->next)
-					first->sent.Set();
+				while (first) { auto cur = first; first = first->next; cur->sent.Set(); }
 			});
 
 		while (connection.socket != InvalidSocket)
