[UBA]
* Changed recv data error from error to warning since we don't want server to throw errors if clients are sending bad data

diff --git a/Engine/Source/Programs/UnrealBuildAccelerator/Common/Private/UbaNetworkServer.cpp b/Engine/Source/Programs/UnrealBuildAccelerator/Common/Private/UbaNetworkServer.cpp
--- a/Engine/Source/Programs/UnrealBuildAccelerator/Common/Private/UbaNetworkServer.cpp
+++ b/Engine/Source/Programs/UnrealBuildAccelerator/Common/Private/UbaNetworkServer.cpp
@@ -439,9 +439,9 @@ namespace uba
 			u32 messageSize = *(u32*)(headerData + 2) & 0x00ffffff;
 
 			if (messageSize > SendMaxSize)
-				return conn.m_server.m_logger.Error(TC("Client %u Got message size %u which is larger than max %u. Protocol error? (serviceId %u, messageType %u, messageId %u)"), conn.m_client->id, messageSize, SendMaxSize, u32(serviceId), u32(messageType), u32(messageId));
+				return conn.m_server.m_logger.Warning(TC("Client %u Got message size %u which is larger than max %u. Protocol error? (serviceId %u, messageType %u, messageId %u)"), conn.m_client->id, messageSize, SendMaxSize, u32(serviceId), u32(messageType), u32(messageId));
 			if (serviceId >= sizeof(NetworkServer::m_workerFunctions))
-				return conn.m_server.m_logger.Error(TC("Client %u Got message with service id %u which is out of range. Protocol error?"), conn.m_client->id, serviceId);
+				return conn.m_server.m_logger.Warning(TC("Client %u Got message with service id %u which is out of range. Protocol error?"), conn.m_client->id, serviceId);
 
 			if (conn.m_recvTimeoutCount)
 			{
