[UBA]
* Added fileSize to unmap call when loading trace file in visualizer

diff --git a/Engine/Source/Programs/UnrealBuildAccelerator/Visualizer/Private/UbaTraceReader.cpp b/Engine/Source/Programs/UnrealBuildAccelerator/Visualizer/Private/UbaTraceReader.cpp
--- a/Engine/Source/Programs/UnrealBuildAccelerator/Visualizer/Private/UbaTraceReader.cpp
+++ b/Engine/Source/Programs/UnrealBuildAccelerator/Visualizer/Private/UbaTraceReader.cpp
@@ -119,7 +119,7 @@ namespace uba
 		m_memoryBegin = MapViewOfFile(m_logger, m_memoryHandle, FILE_MAP_READ, 0, fileSize);
 		if (!m_memoryBegin)
 			return false;
-		auto closeView = MakeGuard([&]() { UnmapViewOfFile(m_logger, m_memoryBegin, 0, TC("TraceReader")); m_memoryBegin = nullptr; m_memoryPos = nullptr; m_memoryEnd = nullptr; });
+		auto closeView = MakeGuard([&]() { UnmapViewOfFile(m_logger, m_memoryBegin, fileSize, TC("TraceReader")); m_memoryBegin = nullptr; m_memoryPos = nullptr; m_memoryEnd = nullptr; });
 
 		m_memoryPos = m_memoryBegin;
 		m_memoryEnd = m_memoryBegin + fileSize;
