[UBA]
* Move font metrics calculations to its own function since mac and windows draw things differently

diff --git a/Engine/Source/Programs/UnrealBuildAccelerator/Visualizer/Private/UbaVisualizer.cpp b/Engine/Source/Programs/UnrealBuildAccelerator/Visualizer/Private/UbaVisualizer.cpp
--- a/Engine/Source/Programs/UnrealBuildAccelerator/Visualizer/Private/UbaVisualizer.cpp
+++ b/Engine/Source/Programs/UnrealBuildAccelerator/Visualizer/Private/UbaVisualizer.cpp
@@ -2677,6 +2677,30 @@ namespace uba
 	{
 	}
 
+	void Visualizer::GetFontMetrics(int& outFh, int& outOffset, int height)
+	{
+		outFh = height;
+		outOffset = 0;
+		if (height <= 13)
+		{
+			++outFh;
+			--outOffset;
+		}
+		if (height <= 11)
+			++outFh;
+		if (height <= 9)
+			++outFh;
+		if (height <= 8)
+			++outFh;
+		if (height <= 6)
+		{
+			++outFh;
+			--outOffset;
+		}
+		if (height <= 4)
+			--outOffset;
+	}
+
 	bool Visualizer::NewTrace(u32 replay, bool paused)
 	{
 		m_replay = replay;
@@ -2830,27 +2854,8 @@ namespace uba
 	void Visualizer::UpdateFont(Font& font, int height, bool createUnderline)
 	{
 		font.height = height;
-
-		int fh = height;
-		font.offset = 0;
-		if (height <= 13)
-		{
-			++fh;
-			--font.offset;
-		}
-		if (height <= 11)
-			++fh;
-		if (height <= 9)
-			++fh;
-		if (height <= 8)
-			++fh;
-		if (height <= 6)
-		{
-			++fh;
-			--font.offset;
-		}
-		if (height <= 4)
-			--font.offset;
+		int fh;
+		GetFontMetrics(fh, font.offset, height);
 
 		if (font.handle)
 			m_backend.DeleteFont(font.handle);
diff --git a/Engine/Source/Programs/UnrealBuildAccelerator/Visualizer/Private/UbaVisualizer.h b/Engine/Source/Programs/UnrealBuildAccelerator/Visualizer/Private/UbaVisualizer.h
--- a/Engine/Source/Programs/UnrealBuildAccelerator/Visualizer/Private/UbaVisualizer.h
+++ b/Engine/Source/Programs/UnrealBuildAccelerator/Visualizer/Private/UbaVisualizer.h
@@ -46,6 +46,7 @@ namespace uba
 		virtual void PostNewTitle(const StringView& title);
 		virtual bool UpdateSelection();
 		virtual void StopUpdate();
+		virtual void GetFontMetrics(int& outFh, int& outOffset, int height);
 
 		bool NewTrace(u32 replay, bool paused);
 		bool UpdateTrace();
