[UBA]
* Fixed so custom assert handler can take user data

diff --git a/Engine/Source/Programs/UnrealBuildAccelerator/Common/Private/UbaLogger.cpp b/Engine/Source/Programs/UnrealBuildAccelerator/Common/Private/UbaLogger.cpp
--- a/Engine/Source/Programs/UnrealBuildAccelerator/Common/Private/UbaLogger.cpp
+++ b/Engine/Source/Programs/UnrealBuildAccelerator/Common/Private/UbaLogger.cpp
@@ -20,6 +20,7 @@
 namespace uba
 {
 	CustomAssertHandler* g_assertHandler;
+	void* g_assertHandlerUserData;
 
 	void ParseCallstack(StringBufferBase& out, BinaryReader& reader)
 	{
@@ -58,7 +59,7 @@ namespace uba
 
 		if (g_assertHandler)
 		{
-			g_assertHandler(sb.data);
+			g_assertHandler(sb.data, sb.count, g_assertHandlerUserData);
 			return;
 		}
 
@@ -93,9 +94,10 @@ namespace uba
 #endif
 	}
 
-	void SetCustomAssertHandler(CustomAssertHandler* handler)
+	void SetCustomAssertHandler(CustomAssertHandler* handler, void* userData)
 	{
 		g_assertHandler = handler;
+		g_assertHandlerUserData = userData;
 	}
 
 	ANALYSIS_NORETURN void FatalError(u32 code, const tchar* format, ...)
diff --git a/Engine/Source/Programs/UnrealBuildAccelerator/Core/Public/UbaPlatform.h b/Engine/Source/Programs/UnrealBuildAccelerator/Core/Public/UbaPlatform.h
--- a/Engine/Source/Programs/UnrealBuildAccelerator/Core/Public/UbaPlatform.h
+++ b/Engine/Source/Programs/UnrealBuildAccelerator/Core/Public/UbaPlatform.h
@@ -91,8 +91,8 @@ namespace uba
 	void WriteAssertInfo(StringBufferBase& out, const tchar* text, const char* file, u32 line, const char* expr, void* context = nullptr);
 	ANALYSIS_NORETURN void UbaAssert(const tchar* text, const char* file, u32 line, const char* expr, bool allowTerminate, u32 terminateCode, void* context, u32 skipCallstackCount);
 	ANALYSIS_NORETURN void FatalError(u32 code, const tchar* format, ...);
-	using CustomAssertHandler = void(const tchar* text);
-	void SetCustomAssertHandler(CustomAssertHandler* handler);
+	using CustomAssertHandler = void(const tchar* text, u32 textLen, void* userData);
+	void SetCustomAssertHandler(CustomAssertHandler* handler, void* userData);
 	bool CreateGuid(Guid& out);
 	const char* GetWineVersion();
 	bool IsRunningWine();
diff --git a/Engine/Source/Programs/UnrealBuildAccelerator/Host/Private/UbaExports.cpp b/Engine/Source/Programs/UnrealBuildAccelerator/Host/Private/UbaExports.cpp
--- a/Engine/Source/Programs/UnrealBuildAccelerator/Host/Private/UbaExports.cpp
+++ b/Engine/Source/Programs/UnrealBuildAccelerator/Host/Private/UbaExports.cpp
@@ -1132,7 +1132,9 @@ extern "C"
 
 	void Uba_SetCustomAssertHandler(Uba_CustomAssertHandler* handler)
 	{
-		uba::SetCustomAssertHandler(handler);
+		using namespace uba;
+		static auto f = [](const tchar* text, u32 textLen, void* userData) { ((Uba_CustomAssertHandler*)userData)(text); };
+		SetCustomAssertHandler(f, (void*)handler);
 	}
 
 	void Uba_FindImports(const uba::tchar* binary, ImportFunc* func, void* userData)
