[UBA]
* Fixed so clicking on .uba files on mac works if UbaVisualizer is setup as default application
* Fixed rendering resolution. Now text is sharp

diff --git a/Engine/Programs/UbaVisualizer/Config/DefaultEngine.ini b/Engine/Programs/UbaVisualizer/Config/DefaultEngine.ini
new file mode 100644
--- /dev/null
+++ b/Engine/Programs/UbaVisualizer/Config/DefaultEngine.ini
@@ -0,0 +1,2 @@
+[/Script/MacTargetPlatform.XcodeProjectSettings]
+TemplateMacPlist=(FilePath="/Engine/Source/Programs/UnrealBuildAccelerator/Visualizer/Resources/Mac/Info.Template.plist")
diff --git a/Engine/Source/Programs/UnrealBuildAccelerator/Visualizer/Private/Mac/UbaVisualizerApp.mm b/Engine/Source/Programs/UnrealBuildAccelerator/Visualizer/Private/Mac/UbaVisualizerApp.mm
--- a/Engine/Source/Programs/UnrealBuildAccelerator/Visualizer/Private/Mac/UbaVisualizerApp.mm
+++ b/Engine/Source/Programs/UnrealBuildAccelerator/Visualizer/Private/Mac/UbaVisualizerApp.mm
@@ -24,6 +24,76 @@ - (void)tick:(NSTimer *)timer;
 
 namespace uba
 {
+	struct MessageBoxLogWriter : public LogWriter
+	{
+		NSWindow* _window;
+
+		MessageBoxLogWriter(NSWindow* window) : _window(window) {}
+
+		virtual void BeginScope() override
+		{
+		}
+
+		virtual void EndScope() override
+		{
+		}
+
+		virtual void Log(LogEntryType type, const tchar* str, u32 strLen, const tchar* prefix = nullptr, u32 prefixLen = 0) override
+		{
+			if (type > LogEntryType_Warning)
+			{
+				#if UBA_DEBUG
+				StringBuffer<> buf;
+				buf.Append(str, strLen).Append(TCV("\r\n"));
+				//OutputDebugStringW(buf.data);
+				#endif
+				return;
+			}
+			NSString* text = [[NSString alloc] initWithBytes:str length:strLen encoding:NSUTF8StringEncoding];
+			if (!text)
+				text = @"Boo";
+
+			NSAlert* alert = [[NSAlert alloc] init];
+			alert.messageText = @"UbaVisualizer";
+			alert.informativeText = text;
+			[alert addButtonWithTitle:@"OK"];
+			alert.alertStyle = NSAlertStyleInformational;
+
+			if ([NSThread isMainThread])
+			{
+				[NSApp activateIgnoringOtherApps:YES];
+				[alert runModal];
+			}
+			else
+			{
+				dispatch_sync(dispatch_get_main_queue(), ^{
+					[alert beginSheetModalForWindow:_window completionHandler:nil];
+				});
+			}
+
+			if (type == LogEntryType_Error)
+				_Exit(-1);
+
+			/*
+			HWND hwnd = NULL;
+			if (m_visualizer)
+			{
+				hwnd = m_visualizer->GetHwnd();
+				m_visualizer->Lock(true);
+			}
+
+			UINT flags = type == LogEntryType_Error ? MB_ICONERROR : MB_ICONWARNING;
+			if (!hwnd)
+				flags |= MB_TOPMOST;
+			MessageBox(hwnd, str, TC("UbaVisualizer"), flags);
+			if (type == LogEntryType_Error)
+				ExitProcess(~0u);
+			if (m_visualizer)
+				m_visualizer->Lock(false);
+			*/
+		}
+	};
+
 	class VisualizerBackendMac final : public VisualizerBackend
 	{
 	public:
@@ -31,8 +101,6 @@ - (void)tick:(NSTimer *)timer;
 
 		struct Bitmap
 		{
-			size_t w, h;
-			size_t bytesPerRow;
 			CGContextRef ctx;
 			CGImageRef img;
 			void* pixels;
@@ -80,7 +148,7 @@ virtual void Text(int x, int y, StringView str, const RECT* clipRect) override
 				if (!cfStr)
 					return;
 
-				y -= 2; // TODO: This should not be here
+				y -= 3; // Look in GetFontMetrics
 
 				const void* keys[]   = { kCTFontAttributeName,			kCTForegroundColorAttributeName };
 				const void* values[] = { (CTFontRef)activeFont.handle,	activeTextColorRef };
@@ -229,14 +297,22 @@ virtual void DrawCached(TraceView::BitmapCache& process, RECT rect, int offsetY,
 							if (owner.m_lastBitmap)
 								owner.m_textBitmaps.push_back(owner.m_lastBitmap);
 
+							CGFloat scale = NSScreen.mainScreen.backingScaleFactor;
+							if (scale <= 0)
+								scale = 1.0;
+
 							auto& b = *new Bitmap;
-							b.w = 256;
-							b.h = BitmapCacheHeight;
-							b.bytesPerRow = b.w * 4;
-							b.pixels = calloc(b.h, b.bytesPerRow);
-							CGBitmapInfo info = kCGBitmapByteOrder32Little | kCGImageAlphaPremultipliedFirst;
-							b.ctx = CGBitmapContextCreate(b.pixels, b.w, b.h, 8, b.bytesPerRow, owner.colorSpaceRef, info);
+							size_t w = 256;
+							size_t h = BitmapCacheHeight;
 
+							w = (size_t)lround((double)w * (double)scale);
+							h = (size_t)lround((double)h * (double)scale);
+
+							size_t bytesPerRow = w * 4;
+							b.pixels = calloc(h, bytesPerRow);
+							CGBitmapInfo info = kCGBitmapByteOrder32Little | kCGImageAlphaPremultipliedFirst;
+							b.ctx = CGBitmapContextCreate(b.pixels, w, h, 8, bytesPerRow, owner.colorSpaceRef, info);
+							CGContextScaleCTM(b.ctx, scale, scale);
 							owner.m_lastBitmapOffset = 0;
 							owner.m_lastBitmap = &b;
 						}
@@ -261,7 +337,7 @@ virtual void DrawCached(TraceView::BitmapCache& process, RECT rect, int offsetY,
 
 				int processWidth = rect.right - rect.left;
 				int width = Min(processWidth, 256);
-				int bitmapOffsetY = process.bitmapOffset;
+				int bitmapOffsetY = process.bitmapOffset;// - processFont.height/4; // TODO: Revisit this.. something is going wrong with the bitmap offset
 				int bltOffsetY = offsetY;
 				if (bltOffsetY < 0)
 				{
@@ -289,8 +365,8 @@ virtual void DrawCached(TraceView::BitmapCache& process, RECT rect, int offsetY,
 				BlitInfo info
 				{
 					*lastSelectedBitmap,
-					CGRectMake(float(rect.left), float(rect.top + bltOffsetY), float(width), float(height)),
-					CGRectMake(float(bitmapOffsetX), float(bitmapOffsetY), float(width), float(height)),
+					CGRectMake(float(rect.left), float(rect.top + bltOffsetY), float(width), float(height)), // dest
+					CGRectMake(float(bitmapOffsetX), float(bitmapOffsetY), float(width), float(height)), // src
 				};
 
 				if (info.b.img)
@@ -311,9 +387,19 @@ void Blit(BlitInfo& info)
 				CGImageRef img = info.b.img;
 				CGContextSaveGState(ctx);
 				CGContextClipToRect(ctx, info.dest);
+				CGContextSetInterpolationQuality(ctx, kCGInterpolationNone);
+
 				CGFloat sx = info.dest.size.width  / info.source.size.width;
 				CGFloat sy = info.dest.size.height / info.source.size.height;
-				CGContextTranslateCTM(ctx, info.dest.origin.x - info.source.origin.x * sx, info.dest.origin.y - info.source.origin.y * sy);
+
+				// pixel-snap translation
+				CGFloat scale = NSScreen.mainScreen.backingScaleFactor;
+				CGFloat tx = info.dest.origin.x - info.source.origin.x * sx;
+				CGFloat ty = info.dest.origin.y - info.source.origin.y * sy;
+				tx = round(tx * scale) / scale;
+				ty = round(ty * scale) / scale;
+
+				CGContextTranslateCTM(ctx, tx, ty);
 				CGContextScaleCTM(ctx, sx, sy);
 				CGRect full = CGRectMake(0, 0, (CGFloat)256, (CGFloat)BitmapCacheHeight); // Draw whole image; clip ensures only desired region shows
 				CGContextDrawImage(ctx, full, img);
@@ -438,6 +524,10 @@ virtual void DeleteFont(FontHandle font) override
 		VisualizerMac(VisualizerConfig& config, Logger& logger, VisualizerBackend& backend, VisualizerCanvasView* view) : Visualizer(config, logger, backend)
 		{
 			m_view = view;
+		}
+
+		void Init()
+		{
 			//m_config.boxHeight = 20;
 
 			UpdateDefaultFont();
@@ -473,6 +563,7 @@ virtual void DeleteFont(FontHandle font) override
 			}
 			else
 			{
+				m_listenDone = false;
 				m_listenThread.Start([this]() { ThreadListen(); return 0; });
 			}
 		}
@@ -490,9 +581,29 @@ virtual void StopUpdate() override
 		{
 			[m_view stopTimer];
 		}
+		void GetFontMetrics(int& outFh, int& outOffset, int height) override
+		{
+			outFh = height;
+			outOffset = 3 - height/4;
+			if (height <= 13)
+				++outFh;
+			if (height <= 11)
+				++outFh;
+			if (height <= 9)
+				++outFh;
+			if (height <= 8)
+				++outFh;
+			if (height <= 6)
+				++outFh;
+		}
 		virtual void FinishIncompatible() override
 		{
 		}
+		virtual void PostNewTitle(const StringView& title)
+		{
+			if (NSString* nsTitle = [[NSString alloc] initWithBytes:title.data length:title.count encoding:NSUTF8StringEncoding])
+				dispatch_async(dispatch_get_main_queue(), ^{ [m_view.window setTitle:nsTitle]; });
+		}
 		bool UpdateSelection() override
 		{
 			return Visualizer::UpdateSelection(((VisualizerBackendMac&)m_backend).m_mousePos);
@@ -583,11 +694,14 @@ bool ThreadListen()
 				return false;
 			}
 
+			StringBuffer<> title;
+			PostNewTitle(GetTitlePrefix(title).Appendf("Listening for new sessions on channel '%s'", channelName));
+
 			m_listenTimeout.Create(EventResetType_Auto);
 			
 			TString newTrace;
 			StringBuffer<256> traceName;
-			while (true)
+			while (!m_listenDone)
 			{
 				traceName.Clear();
 				if (!channel.Read(traceName))
@@ -619,10 +733,13 @@ bool ThreadListen()
 
 		VisualizerCanvasView* m_view;
 		Event m_listenTimeout;
+		Atomic<bool> m_listenDone;
 		Thread m_listenThread;
 	};
 }
 
+uba::VisualizerMac* g_visualizerMac = nullptr;
+
 
 @implementation VisualizerCanvasView
 
@@ -631,14 +748,23 @@ - (instancetype)initWithFrame:(NSRect)frame {
 	self = [super initWithFrame:frame];
 	if (self)
 	{
-		uba::VisualizerConfig config("~/UbaVisualizerConfig.toml");
+		using namespace uba;
+		VisualizerConfig config("~/UbaVisualizerConfig.toml");
 		config.boxHeight = 14;
 		config.horizontalScaleValue = 1.0f;
 		config.DarkMode = true;
-		auto logger = new uba::LoggerWithWriter(uba::g_consoleLogWriter);
-		auto backend = new uba::VisualizerBackendMac;
+		auto writer = new MessageBoxLogWriter(self.window);
+
+		SetCustomAssertHandler([](const tchar* text, u32 textLen, void* userData)
+		{
+			((MessageBoxLogWriter*)userData)->Log(LogEntryType_Error, text, textLen);
+		}, writer);
 
-		_visualizer = new uba::VisualizerMac(config, *logger, *backend, self);
+		auto logger = new LoggerWithWriter(*writer);
+		auto backend = new VisualizerBackendMac;
+
+		_visualizer = new VisualizerMac(config, *logger, *backend, self);
+		g_visualizerMac = _visualizer;
 	}
 	return self;
 }
@@ -770,20 +896,42 @@ - (void)magnifyWithEvent:(NSEvent *)event
 }
 @end
 
-@interface AppDelegate : NSObject <NSWindowDelegate>
+@interface AppDelegate : NSObject <NSApplicationDelegate, NSWindowDelegate>
 @end
 
 @implementation AppDelegate
 - (void)windowWillClose:(NSNotification *)notification {
 	[NSApp terminate:nil];
 }
+
+- (BOOL)application:(NSApplication *)sender openFile:(NSString *)filename
+{
+	using namespace uba;
+	const char* pendingFile = [filename UTF8String];
+	auto& visualizer = *g_visualizerMac;
+	//m_newTraceName.Clear();//.Append(newTrace);
+
+	visualizer.m_listenDone = true;
+	visualizer.m_listenThread.Wait();
+
+	visualizer.m_usingNamed = false;
+	char fixedPath[512];
+	GetFullPathNameW([filename UTF8String], 512, fixedPath, nullptr);
+	visualizer.m_fileName.Append(fixedPath);
+	visualizer.NewTrace(0, false);
+    return YES;
+}
 @end
 
 int main(int argc, const char * argv[]) {
 	@autoreleasepool {
-		[NSApplication sharedApplication];
+		NSApplication* app = [NSApplication sharedApplication];
 		[NSApp setActivationPolicy:NSApplicationActivationPolicyRegular];
 
+		AppDelegate* delegate = [[AppDelegate alloc] init];
+		app.delegate = delegate; 
+
+
 		// --- Minimal menu bar with "Quit" (Cmd+Q) ---
 		NSMenu *menubar = [[NSMenu alloc] init];
 		NSMenuItem *appMenuItem = [[NSMenuItem alloc] init];
@@ -807,7 +955,6 @@ int main(int argc, const char * argv[]) {
 		NSWindowStyleMask style = NSWindowStyleMaskTitled | NSWindowStyleMaskClosable | NSWindowStyleMaskResizable | NSWindowStyleMaskMiniaturizable;
 		NSWindow* window = [[NSWindow alloc] initWithContentRect:frame styleMask:style backing:NSBackingStoreBuffered defer:NO];
 
-		AppDelegate* delegate = [[AppDelegate alloc] init];
 		[window setDelegate:delegate];
 		[window setTitle:@"UbaVisualizer"];
 
@@ -815,6 +962,7 @@ int main(int argc, const char * argv[]) {
 		[window setContentView:view];
 		[window makeKeyAndOrderFront:nil];
 		[NSApp activateIgnoringOtherApps:YES]; // Make sure app shows in Dock and gets focus
+		g_visualizerMac->Init();
 		[NSApp run];
 	}
 	return 0;
