[UBA]
* Added warning and disable racing if allowFilesInMemory is false

diff --git a/Engine/Source/Programs/UnrealBuildAccelerator/Common/Private/UbaScheduler.cpp b/Engine/Source/Programs/UnrealBuildAccelerator/Common/Private/UbaScheduler.cpp
--- a/Engine/Source/Programs/UnrealBuildAccelerator/Common/Private/UbaScheduler.cpp
+++ b/Engine/Source/Programs/UnrealBuildAccelerator/Common/Private/UbaScheduler.cpp
@@ -212,6 +212,12 @@ namespace uba
 			DirectoryCache().CreateDirectory(m_session.GetLogger(), StringBuffer<>(m_session.GetSessionDir()).Append(TCV("inputs")).data);
 
 		m_session.SetOuterScheduler(this);
+
+		if (m_maxRacingWeight != 0 && !m_session.AllowKeepFilesInMemory())
+		{
+			m_maxRacingWeight = 0;
+			m_session.GetLogger().Warning(TC("Session has keepFilesInMemoryDisabled. Racing will not work and has been disabled"));
+		}
 	}
 
 	Scheduler::~Scheduler()
diff --git a/Engine/Source/Programs/UnrealBuildAccelerator/Common/Public/UbaSession.h b/Engine/Source/Programs/UnrealBuildAccelerator/Common/Public/UbaSession.h
--- a/Engine/Source/Programs/UnrealBuildAccelerator/Common/Public/UbaSession.h
+++ b/Engine/Source/Programs/UnrealBuildAccelerator/Common/Public/UbaSession.h
@@ -99,6 +99,7 @@ namespace uba
 		bool ShouldStoreIntermediateFilesCompressed() const { return m_storeIntermediateFilesCompressed; }
 		bool WritePlaceholders() const { return m_writePlaceholders; }
 		bool HasDetailedtrace() const { return m_detailedTrace; }
+		bool AllowKeepFilesInMemory() const { return m_allowKeepFilesInMemory; }
 
 		// Will write current state of trace directly to disk
 		bool SaveSnapshotOfTrace();
