[UBA]
* Removed code line that always initializes view. This line must have been left in by mistake and cause tons of extra kernel calls

diff --git a/Engine/Source/Programs/UnrealBuildAccelerator/Common/Private/UbaSharedMemoryAllocator.cpp b/Engine/Source/Programs/UnrealBuildAccelerator/Common/Private/UbaSharedMemoryAllocator.cpp
--- a/Engine/Source/Programs/UnrealBuildAccelerator/Common/Private/UbaSharedMemoryAllocator.cpp
+++ b/Engine/Source/Programs/UnrealBuildAccelerator/Common/Private/UbaSharedMemoryAllocator.cpp
@@ -422,7 +422,6 @@ namespace uba
 			for (auto& slice : h.slices)
 				totalSize += slice.size;
 		}
-		view.Init(*this, totalSize);
 
 		bool initialized = false;
 
