[UBA]
* Added process reuse count in stats summary per session

diff --git a/Engine/Source/Programs/UnrealBuildAccelerator/Common/Private/UbaProcess.cpp b/Engine/Source/Programs/UnrealBuildAccelerator/Common/Private/UbaProcess.cpp
--- a/Engine/Source/Programs/UnrealBuildAccelerator/Common/Private/UbaProcess.cpp
+++ b/Engine/Source/Programs/UnrealBuildAccelerator/Common/Private/UbaProcess.cpp
@@ -1561,6 +1561,7 @@ namespace uba
 		m_kernelStats = {};
 
 		m_startTime = GetTime();
+		++m_session.m_reuseCount;
 
 		// Rewrite overlay.. and let reused process recreate dir table
 		if (!nextProcess.overlayFiles.empty())
diff --git a/Engine/Source/Programs/UnrealBuildAccelerator/Common/Private/UbaSession.cpp b/Engine/Source/Programs/UnrealBuildAccelerator/Common/Private/UbaSession.cpp
--- a/Engine/Source/Programs/UnrealBuildAccelerator/Common/Private/UbaSession.cpp
+++ b/Engine/Source/Programs/UnrealBuildAccelerator/Common/Private/UbaSession.cpp
@@ -3790,6 +3790,7 @@ namespace uba
 		logger.Info(TC("  DirectoryTable      %7u %9s"), u32(m_directoryTable.m_lookup.size()), BytesToText(GetDirectoryTableSize()).str);
 		logger.Info(TC("  MappingTable        %7u %9s"), u32(m_fileMappingTableLookup.size()), BytesToText(GetFileMappingSize()).str);
 		logger.Info(TC("  MappingBuffer       %7u %9s"), 1, BytesToText(m_fileMappingBuffer.GetUsed()).str);
+		logger.Info(TC("  ProcessReuseCount   %7u"), m_reuseCount.load());
 		m_stats.Print(logger);
 		logger.Info(TC(""));
 	}
diff --git a/Engine/Source/Programs/UnrealBuildAccelerator/Common/Public/UbaSession.h b/Engine/Source/Programs/UnrealBuildAccelerator/Common/Public/UbaSession.h
--- a/Engine/Source/Programs/UnrealBuildAccelerator/Common/Public/UbaSession.h
+++ b/Engine/Source/Programs/UnrealBuildAccelerator/Common/Public/UbaSession.h
@@ -276,8 +276,9 @@ namespace uba
 		FileMappingHandle m_fileMappingTableHandle;
 		FileMappingBuffer m_fileMappingBuffer;
 
-		Atomic<u64> m_independentMappingCreated;
-		Atomic<u64> m_independentMappingActive;
+		Atomic<u64> m_independentMappingCreated = 0;
+		Atomic<u64> m_independentMappingActive = 0;
+		Atomic<u64> m_reuseCount = 0;
 
 		ReaderWriterLock m_fileMappingTableMemLock;
 		u8* m_fileMappingTableMem;
diff --git a/Engine/Source/Programs/UnrealBuildAccelerator/Core/Private/UbaSharedMemoryView.cpp b/Engine/Source/Programs/UnrealBuildAccelerator/Core/Private/UbaSharedMemoryView.cpp
--- a/Engine/Source/Programs/UnrealBuildAccelerator/Core/Private/UbaSharedMemoryView.cpp
+++ b/Engine/Source/Programs/UnrealBuildAccelerator/Core/Private/UbaSharedMemoryView.cpp
@@ -113,6 +113,8 @@ namespace uba
 
 	bool SharedMemoryView::InternalInit(SharedMemoryAllocatorHandle handle, u64 capacity, void* baseAddress, const tchar* hint)
 	{
+		UBA_ASSERT(!m_memory);
+
 		m_capacity = AlignUp(capacity, PageSize);
 
 		TimerScope ts(KernelStats::GetCurrent().virtualAlloc2);
