[UBA]
* Removing zeroing of memory for virtual file since we memcpy in to it straight away

diff --git a/Engine/Source/Programs/UnrealBuildAccelerator/Common/Private/UbaSession.cpp b/Engine/Source/Programs/UnrealBuildAccelerator/Common/Private/UbaSession.cpp
--- a/Engine/Source/Programs/UnrealBuildAccelerator/Common/Private/UbaSession.cpp
+++ b/Engine/Source/Programs/UnrealBuildAccelerator/Common/Private/UbaSession.cpp
@@ -3843,7 +3843,7 @@ namespace uba
 		SharedMemoryHandle mapping = m_sharedMemory.CreateHandle(m_logger, filePath.data);
 		if (!mapping.IsValid())
 			return false;
-		m_sharedMemory.ExtendMemory(mapping, memorySize);
+		m_sharedMemory.ExtendMemory(mapping, memorySize, filePath.data, false);
 		u8* mem2 = m_sharedMemory.MapView(mapping, filePath.data, SharedMemoryMapType_ReadWrite);
 		if (!mem2)
 			return false;
