[UBA]
* Fixed bug where process log for reuse was not ending up in the right directory

diff --git a/Engine/Source/Programs/UnrealBuildAccelerator/Common/Private/UbaScheduler.cpp b/Engine/Source/Programs/UnrealBuildAccelerator/Common/Private/UbaScheduler.cpp
--- a/Engine/Source/Programs/UnrealBuildAccelerator/Common/Private/UbaScheduler.cpp
+++ b/Engine/Source/Programs/UnrealBuildAccelerator/Common/Private/UbaScheduler.cpp
@@ -1206,10 +1206,14 @@ namespace uba
 			ei->processIndex = indexToRun;
 
 			auto& si = *newInfo;
+
+			StringBuffer<> temp;
+			const tchar* logFile = m_session.FixLogFile(si.logFile, temp, si.arguments, process.GetId());
+
 			outNextProcess.arguments = si.arguments;
 			outNextProcess.workingDir = si.workingDir;
 			outNextProcess.description = si.description;
-			outNextProcess.logFile = si.logFile;
+			outNextProcess.logFile = logFile;
 			outNextProcess.breadcrumbs = si.breadcrumbs;
 
 			if (*si.overlayFiles)
diff --git a/Engine/Source/Programs/UnrealBuildAccelerator/Common/Private/UbaSession.cpp b/Engine/Source/Programs/UnrealBuildAccelerator/Common/Private/UbaSession.cpp
--- a/Engine/Source/Programs/UnrealBuildAccelerator/Common/Private/UbaSession.cpp
+++ b/Engine/Source/Programs/UnrealBuildAccelerator/Common/Private/UbaSession.cpp
@@ -2185,22 +2185,8 @@ namespace uba
 		
 		u32 processId = CreateProcessId();
 
-		StringBuffer<> logFile;
-		if (si.logFile && *si.logFile)
-		{
-			if (TStrchr(si.logFile, PathSeparator) == nullptr)
-			{
-				logFile.Append(m_sessionLogDir).Append(si.logFile);
-				si.logFile = logFile.data;
-			}
-		}
-		else if (m_logToFile)
-		{
-			logFile.Append(m_sessionLogDir);
-			GenerateNameForProcess(logFile, startInfo.arguments, processId);
-			logFile.Append(TCV(".log"));
-			si.logFile = logFile.data;
-		}
+		StringBuffer<> temp;
+		si.logFile = FixLogFile(si.logFile, temp, startInfo.arguments, processId);
 
 		if (!si.rules)
 			si.rules = GetRules(si);
@@ -2746,6 +2732,22 @@ namespace uba
 		lock.Leave();
 	}
 
+	const tchar* Session::FixLogFile(const tchar* logFile, StringBufferBase& temp, const tchar* processArguments, u32 processId)
+	{
+		if (logFile && *logFile)
+		{
+			if (TStrchr(logFile, PathSeparator) == nullptr)
+				return temp.Append(m_sessionLogDir).Append(logFile).data;
+		}
+		else if (m_logToFile)
+		{
+			temp.Append(m_sessionLogDir);
+			GenerateNameForProcess(temp, processArguments, processId);
+			return temp.Append(TCV(".log")).data;
+		}
+		return logFile;
+	}
+
 	bool Session::ProcessThreadStart(ProcessImpl& process)
 	{
 		return true;
diff --git a/Engine/Source/Programs/UnrealBuildAccelerator/Common/Public/UbaSession.h b/Engine/Source/Programs/UnrealBuildAccelerator/Common/Public/UbaSession.h
--- a/Engine/Source/Programs/UnrealBuildAccelerator/Common/Public/UbaSession.h
+++ b/Engine/Source/Programs/UnrealBuildAccelerator/Common/Public/UbaSession.h
@@ -121,6 +121,8 @@ namespace uba
 
 		void FlushDeadProcesses();
 
+		const tchar* FixLogFile(const tchar* logFile, StringBufferBase& temp, const tchar* processArguments, u32 processId);
+
 	protected:
 		Session(const SessionCreateInfo& info, const tchar* logPrefix, bool runningRemote, WorkManager& workManager);
 		bool Create(const SessionCreateInfo& info);
