[UBA]
* Fixed so create directory check directory table first in order to avoid kernel calls

diff --git a/Engine/Source/Programs/UnrealBuildAccelerator/Detours/Private/Windows/UbaDetoursFunctionsKernelBase.inl b/Engine/Source/Programs/UnrealBuildAccelerator/Detours/Private/Windows/UbaDetoursFunctionsKernelBase.inl
--- a/Engine/Source/Programs/UnrealBuildAccelerator/Detours/Private/Windows/UbaDetoursFunctionsKernelBase.inl
+++ b/Engine/Source/Programs/UnrealBuildAccelerator/Detours/Private/Windows/UbaDetoursFunctionsKernelBase.inl
@@ -228,12 +228,25 @@ BOOL Detoured_CreateDirectoryW(LPCWSTR lpPathName, LPSECURITY_ATTRIBUTES lpSecur
 	}
 
 	DevirtualizePath(pathName);
+	
+	StringBuffer<> lowerName(pathName);
+	lowerName.MakeLower();
+
+	StringKey pathNameKey = ToStringKey(lowerName);
+
+	bool checkIfDir = true;
+	u32 dirTableOffset = ~u32(0);
+	DirectoryTable::Exists exists = g_directoryTable.EntryExists(pathNameKey, lowerName, checkIfDir, &dirTableOffset);
+	if (exists == DirectoryTable::Exists_Yes)
+	{
+		SetLastError(ERROR_ALREADY_EXISTS);
+		DEBUG_LOG_DETOURED(L"CreateDirectoryW", L"(DIRTABLE) %ls -> %ls (%u)", pathName.data, ToString(false), ERROR_ALREADY_EXISTS);
+		return false;
+	}
 
 	u32 directoryTableSize;
 	BOOL res;
 	u32 errorCode = 0;
-	StringKey pathNameKey = ToStringKeyLower(pathName);
-
 	{
 		RPC_MESSAGE(CreateDirectory, createFile)
 		writer.WriteStringKey(pathNameKey);
diff --git a/Engine/Source/Programs/UnrealBuildAccelerator/Test/Private/UbaTestBasics.cpp b/Engine/Source/Programs/UnrealBuildAccelerator/Test/Private/UbaTestBasics.cpp
--- a/Engine/Source/Programs/UnrealBuildAccelerator/Test/Private/UbaTestBasics.cpp
+++ b/Engine/Source/Programs/UnrealBuildAccelerator/Test/Private/UbaTestBasics.cpp
@@ -276,6 +276,9 @@ namespace uba
 		if (!CreateDirectoryW(TC("TestDir")))
 			return logger.Error(TC("Failed to create dir"));
 
+		if (CreateDirectoryW(TC("TestDir")))
+			return logger.Error(TC("Did not fail to create existing dir"));
+
 		FileHandle fileHandle4;
 		if (OpenFileSequentialRead(logger, TC("TestDir"), fileHandle4))
 			return logger.Error(TC("This should return fail"));
diff --git a/Engine/Source/Programs/UnrealBuildAccelerator/TestApp/Private/Windows/UbaTestAppWin.cpp b/Engine/Source/Programs/UnrealBuildAccelerator/TestApp/Private/Windows/UbaTestAppWin.cpp
--- a/Engine/Source/Programs/UnrealBuildAccelerator/TestApp/Private/Windows/UbaTestAppWin.cpp
+++ b/Engine/Source/Programs/UnrealBuildAccelerator/TestApp/Private/Windows/UbaTestAppWin.cpp
@@ -166,6 +166,8 @@ int wmain(int argc, wchar_t* argv[])
 		{
 			if (!CreateDirectoryW(L"DirA", NULL))
 				return LogError(L"Failed to create directory");
+			if (CreateDirectoryW(L"DirA", NULL))
+				return LogError(L"Should not succeed creating directory that exists");
 
 			if (GetFileAttributesW(L"DirA") == 0)
 				return LogError(L"Failed to get attributes of directory");
@@ -177,7 +179,7 @@ int wmain(int argc, wchar_t* argv[])
 				return LogError(L"Found attributes of deleted directory");
 
 			if (CreateDirectoryW(L"Dir2\\Dir3", NULL))
-				return LogError(L"Should not succeed creation directory that exists");
+				return LogError(L"Should not succeed creating directory inside directory that does not exist");
 			if (GetLastError() != ERROR_ALREADY_EXISTS)
 				return LogError(L"Did not get correct error when failing to create existing directory");
 			if (GetFileAttributesW(L"Dir2\\Dir3\\Dir4\\Dir5") == INVALID_FILE_ATTRIBUTES)
