UBA: load UbaWine.dll.so on the host when running under Wine

diff --git a/Engine/Source/Programs/UnrealBuildAccelerator/Common/Private/UbaNetworkBackendTcp.cpp b/Engine/Source/Programs/UnrealBuildAccelerator/Common/Private/UbaNetworkBackendTcp.cpp
--- a/Engine/Source/Programs/UnrealBuildAccelerator/Common/Private/UbaNetworkBackendTcp.cpp
+++ b/Engine/Source/Programs/UnrealBuildAccelerator/Common/Private/UbaNetworkBackendTcp.cpp
@@ -240,7 +240,12 @@ namespace uba
 
 		if (IsRunningWine())
 		{
-			g_wineDll = LoadLibrary(TC("UbaWine.dll.so"));
+			StringBuffer<> dllPath;
+			if (GetDirectoryOfCurrentModule(logger, dllPath))
+				dllPath.EnsureEndsWithSlash();
+			dllPath.Append(TCV("UbaWine.dll.so"));
+
+			g_wineDll = LoadLibrary(dllPath.data);
 			if (g_wineDll)
 				SocketInitWine(logger, g_wineDll);
 			else// if (FileExists(logger, TC("UbaWine.dll.so")))
