[UBA]
* Fixed bug where visualizer d2d backend would crash if showing jobs and resizing window
* Fixed so brushes are not recreated every render frame
* Added support for special case "IDLE" which renders processor graphs differently

diff --git a/Engine/Source/Programs/UnrealBuildAccelerator/Visualizer/Private/UbaTraceReader.cpp b/Engine/Source/Programs/UnrealBuildAccelerator/Visualizer/Private/UbaTraceReader.cpp
--- a/Engine/Source/Programs/UnrealBuildAccelerator/Visualizer/Private/UbaTraceReader.cpp
+++ b/Engine/Source/Programs/UnrealBuildAccelerator/Visualizer/Private/UbaTraceReader.cpp
@@ -863,13 +863,17 @@ namespace uba
 			auto& process2 = processes.back();
 			active.processIndex = u32(processes.size() - 1);
 			process2.id = processId;
-			process2.description = reason.data;
 			process2.breadcrumbs = breadcrumbs;
 			process2.start = time;
 			process2.stop = ~u64(0);
 			process2.exitCode = ~0u;
 			process2.isRemote = isRemote;
 
+			if (reason.Equals(TCV("IDLE")))
+				process2.isIdle = true;
+			else
+				process2.description = reason.data;
+
 			++session.processExitedCount;
 			++out.totalProcessExitedCount;
 			break;
diff --git a/Engine/Source/Programs/UnrealBuildAccelerator/Visualizer/Private/UbaTraceReader.h b/Engine/Source/Programs/UnrealBuildAccelerator/Visualizer/Private/UbaTraceReader.h
--- a/Engine/Source/Programs/UnrealBuildAccelerator/Visualizer/Private/UbaTraceReader.h
+++ b/Engine/Source/Programs/UnrealBuildAccelerator/Visualizer/Private/UbaTraceReader.h
@@ -39,6 +39,7 @@ namespace uba
 			bool isRemote = false;
 			bool isReuse = false;
 			bool isRacing = false;
+			bool isIdle = false;
 			u64 createFilesTime = 0;
 			u64 writeFilesTime = 0;
 			Vector<u8> stats;
diff --git a/Engine/Source/Programs/UnrealBuildAccelerator/Visualizer/Private/UbaVisualizer.cpp b/Engine/Source/Programs/UnrealBuildAccelerator/Visualizer/Private/UbaVisualizer.cpp
--- a/Engine/Source/Programs/UnrealBuildAccelerator/Visualizer/Private/UbaVisualizer.cpp
+++ b/Engine/Source/Programs/UnrealBuildAccelerator/Visualizer/Private/UbaVisualizer.cpp
@@ -210,6 +210,8 @@ namespace uba
 			m_processBrushes[1].cacheMiss = CreateBrush(130, 130, 130);
 			m_processBrushes[0].warning = CreateBrush(100, 100, 0);
 			m_processBrushes[1].warning = CreateBrush(100, 100, 0);
+			m_processBrushes[0].idle = CreateBrush(45, 47, 47);
+			m_processBrushes[1].idle = CreateBrush(55, 57, 57);
 
 			m_backgroundBrush = CreateBrush(38, 37, 37);
 			m_separatorPen = m_backend.CreatePen(RGB(50, 50, 50));
@@ -263,6 +265,8 @@ namespace uba
 			m_processBrushes[1].cacheMiss = CreateBrush(180, 180, 180);
 			m_processBrushes[0].warning = CreateBrush(170, 170, 0);
 			m_processBrushes[1].warning = CreateBrush(170, 170, 0);
+			m_processBrushes[0].idle = CreateBrush(70, 70, 160);
+			m_processBrushes[1].idle = CreateBrush(80, 80, 170);
 
 			m_separatorPen = m_backend.CreatePen(RGB(140, 140, 140));
 			m_backgroundBrush = m_backend.CreateSolidBrush(backgroundColor);
@@ -284,6 +288,8 @@ namespace uba
 		m_memPen = m_backend.CreatePen(m_memColor);
 		m_drivePen = m_backend.CreatePen(m_driveColor);
 		m_activeProcPen = m_backend.CreatePen(m_activeProcColor);
+
+		m_coloredBrushes.clear(); // Needs to be recreated
 	}
 
 	u64 Visualizer::GetPlayTime()
@@ -826,7 +832,7 @@ namespace uba
 				RECT rect;
 				rect.left = posX;
 				rect.right = endX;
-				rect.top = posY + activeFont.offset;
+				rect.top = posY;// + activeFont.offset;
 				rect.bottom = posY + activeFont.height + 2;
 				context.SetTextColor(type == LogEntryType_Info ? m_textColor : (type == LogEntryType_Error ? m_textErrorColor : m_textWarningColor));
 
@@ -1044,18 +1050,16 @@ namespace uba
 
 					RECT boxRect = drawBox(process.start, process.stop, boxHeight, selected, true);
 
-					u32 v = boxHeight - 1;
+					u32 v = boxHeight;// - 1;
 
 					if (v > 4)
 					{
 						u32 fontIndex = Min(v, u32(sizeof_array(m_activeProcessFont) - 1));
-						if (!m_activeProcessFont[fontIndex].handle)
-						{
-							UpdateFont(m_activeProcessFont[fontIndex], fontIndex - 1, false);
-							m_activeProcessFont[fontIndex].offset += 1;
-						}
+						Font& font = m_activeProcessFont[fontIndex];
+						if (!font.handle)
+							UpdateFont(font, fontIndex - 1, false);
 						if (firstWithHeight)
-							context.SetFont(m_activeProcessFont[fontIndex]);
+							context.SetFont(font);
 
 						StringBuffer<> str;
 						auto firstPar = -1;//process.description.find_first_of('(');
@@ -1073,7 +1077,6 @@ namespace uba
 							str.Append(TCV(" [cache]"));
 						if (boxRect.left < 0)
 							str.Appendf(TC("   %s"), TimeToText(playTime - process.start, true).str);
-						//drawIndentedText(str, LogEntryType_Info, 1, true);
 						drawStatusText(str, LogEntryType_Info, Max(int(boxRect.left+1), 1), boxRect.right, false);
 					}
 				});
@@ -1798,7 +1801,8 @@ namespace uba
 			logger.AddWidth(3);
 
 			logger.AddSpace(2);
-			logger.Info(TC("  %s"), process.description.c_str());
+			const tchar* desc = process.isIdle ? TC("Waiting for process") : process.description.c_str();
+			logger.Info(TC("  %s"), desc);
 			if (process.type != TraceView::ProcessType_Task)
 			{
 				logger.Info(TC("  Host:        %s"), m_processSelectedLocation.sessionIndex == 0 ? TC("local") : m_traceView.GetSession(m_processSelectedLocation).name.c_str());
@@ -2119,6 +2123,8 @@ namespace uba
 		BrushHandle brush = m_processBrushes[selected].success;
 		if (!process.returnedReason.empty())
 			brush = process.type == TraceView::ProcessType_CacheFetchTest ? m_processBrushes[selected].cacheMiss : m_processBrushes[selected].returned;
+		else if (process.isIdle)
+			brush = m_processBrushes[selected].idle;
 		else if (!done)
 			brush = process.isRacing ? m_processBrushes[selected].inProgressRace : m_processBrushes[selected].inProgress;
 		else if (process.type == TraceView::ProcessType_CacheFetchTest)
diff --git a/Engine/Source/Programs/UnrealBuildAccelerator/Visualizer/Private/UbaVisualizer.h b/Engine/Source/Programs/UnrealBuildAccelerator/Visualizer/Private/UbaVisualizer.h
--- a/Engine/Source/Programs/UnrealBuildAccelerator/Visualizer/Private/UbaVisualizer.h
+++ b/Engine/Source/Programs/UnrealBuildAccelerator/Visualizer/Private/UbaVisualizer.h
@@ -123,6 +123,7 @@ namespace uba
 			BrushHandle cacheFetchDownload = 0;
 			BrushHandle cacheMiss = 0;
 			BrushHandle warning = 0;
+			BrushHandle idle = 0;
 		};
 
 		Logger& m_logger;
diff --git a/Engine/Source/Programs/UnrealBuildAccelerator/Visualizer/Private/Windows/UbaVisualizerBackendD2D.h b/Engine/Source/Programs/UnrealBuildAccelerator/Visualizer/Private/Windows/UbaVisualizerBackendD2D.h
--- a/Engine/Source/Programs/UnrealBuildAccelerator/Visualizer/Private/Windows/UbaVisualizerBackendD2D.h
+++ b/Engine/Source/Programs/UnrealBuildAccelerator/Visualizer/Private/Windows/UbaVisualizerBackendD2D.h
@@ -63,9 +63,6 @@ namespace uba
 			}
 			virtual void Text(int x, int y, StringView str, const RECT* clipRect) override
 			{
-				x += clientOffsetX;
-				y += clientOffsetY;
-
 				D2D1_RECT_F layoutRect = D2D1::RectF(float(x), float(y), float(x + 10000), float(y + 100));
 
 				if (clipRect)
@@ -73,8 +70,8 @@ namespace uba
 					//UBA_ASSERT(clipRect->left <= x);
 					//UBA_ASSERT(clipRect->top <= y);
 					D2D1_RECT_F cr = layoutRect;
-					cr.right = clipRect->right + float(clientOffsetX);
-					cr.bottom = clipRect->bottom + float(clientOffsetX);
+					cr.right = float(clipRect->right);
+					cr.bottom = float(clipRect->bottom);
 					if (cr.bottom < 0)
 						return;
 					renderTarget.PushAxisAlignedClip(cr, D2D1_ANTIALIAS_MODE_ALIASED);
@@ -247,9 +244,6 @@ namespace uba
 				return p;
 			}
 
-			int clientOffsetX = 0;
-			int clientOffsetY = 0;
-
 			VisualizerBackendD2D& owner;
 			ID2D1RenderTarget& renderTarget;
 			SIZE clientSize;
@@ -317,6 +311,7 @@ namespace uba
 			{
 				m_owner.InitBrushes();
 				m_d2dRenderTarget->CreateSolidColorBrush(ToRgbD2D(m_owner.m_textColor), &m_textBrush);
+				m_dirty = false;
 			}
 			PAINTSTRUCT ps;
 			BeginPaint(m_owner.m_hwnd, &ps);
