[UBA]
* Added support for calling process re-use multiple times from the same process
* Added new reuse callback function that takes in timeout and outShouldExit. This can (should) be used for processes to wait in the most optimal place where it can trigger quickly on new processes

diff --git a/Engine/Source/Programs/UnrealBuildAccelerator/Common/Private/UbaProcess.cpp b/Engine/Source/Programs/UnrealBuildAccelerator/Common/Private/UbaProcess.cpp
--- a/Engine/Source/Programs/UnrealBuildAccelerator/Common/Private/UbaProcess.cpp
+++ b/Engine/Source/Programs/UnrealBuildAccelerator/Common/Private/UbaProcess.cpp
@@ -1506,13 +1506,14 @@ namespace uba
 	bool ProcessImpl::HandleGetNextProcess(BinaryReader& reader, BinaryWriter& writer)
 	{
 		u32 prevExitCode = reader.ReadU32();
+		u32 timeoutMs = reader.ReadU32();
 		NextProcessInfo nextProcess;
 
 		StackBinaryWriter<16 * 1024> statsWriter;
 					
 		ProcessStats processStats;
 		processStats.Read(reader, TraceVersion);
-		processStats.Add(m_processStats);
+
 		processStats.startupTime = m_processStats.startupTime;
 		processStats.wallTime = GetTime() - m_startTime;
 		processStats.cpuTime = 0;
@@ -1520,7 +1521,10 @@ namespace uba
 		
 		KernelStats kernelStats;
 		kernelStats.Read(reader, TraceVersion);
-		kernelStats.Add(m_kernelStats);
+		
+		// Transfer stats from detoured process to local
+		m_processStats.Add(processStats);
+		m_kernelStats.Add(kernelStats);
 
 		processStats.Write(statsWriter);
 		if (m_runningRemote)
@@ -1532,14 +1536,24 @@ namespace uba
 		WriteFilesToDisk();
 
 		bool newProcess = false;
+		bool shouldExit = false;
 		m_exitCode = prevExitCode;
-		CallSession(MessageType_GetNextProcess, m_session.GetNextProcess(*this, newProcess, nextProcess, prevExitCode, statsReader));
+		CallSession(MessageType_GetNextProcess, m_session.GetNextProcess(*this, newProcess, shouldExit, nextProcess, prevExitCode, statsReader, timeoutMs));
 		writer.WriteBool(newProcess);
+		writer.WriteBool(shouldExit);
+
 		m_exitCode = ~0u;
+		CleanupWrittenFiles();
+
+		m_startInfo.arguments = nullptr;
+
 		if (!newProcess)
 			return true;
 
-		CleanupWrittenFiles();
+		m_processStats = {};
+		m_sessionStats = {};
+		m_storageStats = {};
+		m_kernelStats = {};
 
 		m_startInfo.argumentsStr = nextProcess.arguments;
 		m_startInfo.descriptionStr = nextProcess.description;
@@ -1555,11 +1569,6 @@ namespace uba
 		m_trackedOutputs.clear();
 		m_shared.writtenFiles.clear();
 
-		m_processStats = {};
-		m_sessionStats = {};
-		m_storageStats = {};
-		m_kernelStats = {};
-
 		m_startTime = GetTime();
 		++m_session.m_reuseCount;
 
diff --git a/Engine/Source/Programs/UnrealBuildAccelerator/Common/Private/UbaScheduler.cpp b/Engine/Source/Programs/UnrealBuildAccelerator/Common/Private/UbaScheduler.cpp
--- a/Engine/Source/Programs/UnrealBuildAccelerator/Common/Private/UbaScheduler.cpp
+++ b/Engine/Source/Programs/UnrealBuildAccelerator/Common/Private/UbaScheduler.cpp
@@ -87,6 +87,7 @@ namespace uba
 		Atomic<u32> refCount = 1;
 		u32 processIndex = ~0u;
 		u32 processCount = 1;
+		float weight = 1.0f;
 		bool wasReturned = false;
 		bool hasCheckedWasReturned = false;
 		bool isLocal = true;
@@ -194,9 +195,12 @@ namespace uba
 	,	m_writeToCache(info.writeToCache && info.cacheClientCount)
 	{
 		m_cacheClients.insert(m_cacheClients.end(), info.cacheClients, info.cacheClients + info.cacheClientCount);
-		m_session.RegisterGetNextProcess([this](Process& process, NextProcessInfo& outNextProcess, u32 prevExitCode)
+		m_session.RegisterGetNextProcess([this](Process& process, NextProcessInfo& outNextProcess, u32 prevExitCode, BinaryReader& prevStats, u32 timeoutMs, bool& outShouldExit)
 			{
-				return HandleReuseMessage(process, outNextProcess, prevExitCode);
+				if (HandleGetNextProcess(process, outNextProcess, prevExitCode, prevStats, timeoutMs))
+					return true;
+				outShouldExit = m_cancelled;
+				return false;
 			});
 
 		if (info.memWatchdog && IsWindows && !IsRunningWine()) // Only implemented for normal windows for now
@@ -549,7 +553,7 @@ namespace uba
 
 		if (isLocal)
 		{
-			processWeight = ei.startInfo->weight;
+			processWeight = ei.weight;
 			processIndex = ei.processIndex; // Must be fetched before cancel
 
 			SCOPED_CRITICAL_SECTION(m_wasReturnedLock, lock);
@@ -684,7 +688,8 @@ namespace uba
 		{
 			UBA_ASSERT(handle.GetExitCode() == ProcessCancelExitCode);
 
-			float weight = info->startInfo->weight;
+			UBA_ASSERT(si);
+			float weight = info->weight;
 
 			if (handle.IsRemote()) // Racing process won. Let's make it a "real" local process
 			{
@@ -723,9 +728,18 @@ namespace uba
 			info->hasCheckedWasReturned = true;
 		}
 
-		if (!si) // Can be a process that was reused but didn't get a new process
+		if (!si) // Can be a remote process or a failed reuse
 		{
 			UBA_ASSERT(info->processIndex == ~0u);
+			
+			// Have to return the weight here if local. This means that this is a failed reuse and process has already been "exited" but we've kept process weight alive
+			if (info->isLocal)
+			{
+				m_activeLocalProcessWeight -= info->weight;
+				UBA_ASSERTF(m_activeLocalProcessWeight >= 0, TC("Active process weight is below zero (%f)"), m_activeLocalProcessWeight.load());
+				info->weight = 0;
+				m_updateThreadWakeup.Set();
+			}
 			return;
 		}
 
@@ -1077,6 +1091,7 @@ namespace uba
 			exitInfo->startInfo = info;
 			exitInfo->isLocal = isLocal;
 			exitInfo->processIndex = indexToRun;
+			exitInfo->weight = info->weight;
 
 			ProcessStartInfo si = *info;
 			si.userData = exitInfo;
@@ -1114,7 +1129,7 @@ namespace uba
 		ProcessHandle handle = m_session.RunProcessRacing([&](ProcessHandle& remoteProcess, ProcessStartInfo& psi)
 			{
 				auto ei = (ExitProcessInfo*)remoteProcess.GetStartInfo().userData;
-				float weight = ei->startInfo->weight;
+				float weight = ei->weight;
 				if (totalWeight + weight > m_maxRacingWeight)
 					return false;
 				m_activeRacingProcessWeight += weight;
@@ -1151,7 +1166,7 @@ namespace uba
 		return true;
 	}
 
-	bool Scheduler::HandleReuseMessage(Process& process, NextProcessInfo& outNextProcess, u32 prevExitCode)
+	bool Scheduler::HandleGetNextProcess(Process& process, NextProcessInfo& outNextProcess, u32 prevExitCode, BinaryReader& prevStats, u32 timeoutMs)
 	{
 		if (!m_enableProcessReuse)
 			return false;
@@ -1164,8 +1179,6 @@ namespace uba
 		bool isLocal = !process.IsRemote();
 
 		bool forReuse = true;
-		float weight = isLocal ? ei->startInfo->weight : 0.0f;
-		auto weightGuard = MakeGuard([&]() { m_activeLocalProcessWeight -= weight; });
 
 		ExitProcess(*ei, process, prevExitCode, false, forReuse);
 
@@ -1179,13 +1192,36 @@ namespace uba
 			ei->hasCheckedWasReturned = true;
 		}
 
+		// TODO: handle timeout.. when local process we can sit on this thread until a new process pops up
+		// or scheduler wants to exit.
+		u64 startTime = GetTime();
+		bool firstWait = true;
+
 		while (true)
 		{
 			RunStatus prevStatus;
 			SCOPED_FUTEX(m_processEntriesLock, lock);
 			u32 indexToRun = PopProcess(isLocal, prevStatus, forReuse);
 			if (indexToRun == ~0u)
-				return false;
+			{
+				u64 usedTimeMs = TimeToMs(GetTime() - startTime);
+				if (usedTimeMs >= timeoutMs || !isLocal || m_cancelled)
+					return false;
+				lock.Leave();
+				if (firstWait)
+				{
+					m_session.SetProcessAsIdle((ProcessImpl&)process, prevStats);
+					//m_session.GetTrace().ProcessEnvironmentUpdated(process.GetId(), TCV("IDLE"), prevStats.GetPositionData(), prevStats.GetLeft(), TCV(""));
+					firstWait = false;
+				}
+
+				// Piggyback on event that is triggered when queue changes.
+				// And set it again if we got it since there might be more things happening that the update thread should handle
+				u32 timeToWaitMs = u32(timeoutMs - usedTimeMs);
+				if (m_updateThreadWakeup.IsSet(timeToWaitMs))
+					m_updateThreadWakeup.Set();
+				continue;
+			}
 			//UBA_ASSERT(prevStatus != RunStatus_QueuedForCache);
 			auto& processEntry = m_processEntries[indexToRun];
 			auto newInfo = processEntry.info;
@@ -1202,10 +1238,13 @@ namespace uba
 
 			UpdateActiveProcessCounter(isLocal, 1);
 
+			auto& si = *newInfo;
+
+			UBA_ASSERT(ei->weight == si.weight); // TOOD: Support different weights?
+
 			ei->startInfo = newInfo;
 			ei->processIndex = indexToRun;
-
-			auto& si = *newInfo;
+			ei->weight = si.weight;
 
 			StringBuffer<> temp;
 			const tchar* logFile = m_session.FixLogFile(si.logFile, temp, si.arguments, process.GetId());
@@ -1238,9 +1277,6 @@ namespace uba
 			PrepPath(temp2, si);
 			UBA_ASSERTF(temp1.Equals(temp2.data), TC("%s vs %s"), temp1.data, temp2.data);
 			#endif
-			
-			weightGuard.Cancel();
-
 			return true;
 		}
 	}
@@ -1292,7 +1328,10 @@ namespace uba
 				--m_activeCacheQueries;
 			}
 			else if (!forReuse)
-				m_activeLocalProcessWeight -= si->weight;
+			{
+				m_activeLocalProcessWeight -= info.weight;
+				info.weight = 0;
+			}
 		}
 
 		lock.Leave();
@@ -2235,7 +2274,7 @@ namespace uba
 		if (userData)
 		{
 			auto& ei = *(ExitProcessInfo*)userData;
-			weight = ei.startInfo->weight;
+			weight = ei.weight;
 		}
 
 		// Return weight while running these tasks
diff --git a/Engine/Source/Programs/UnrealBuildAccelerator/Common/Private/UbaSession.cpp b/Engine/Source/Programs/UnrealBuildAccelerator/Common/Private/UbaSession.cpp
--- a/Engine/Source/Programs/UnrealBuildAccelerator/Common/Private/UbaSession.cpp
+++ b/Engine/Source/Programs/UnrealBuildAccelerator/Common/Private/UbaSession.cpp
@@ -2748,6 +2748,15 @@ namespace uba
 		return logFile;
 	}
 
+	void Session::SetProcessAsIdle(ProcessImpl& process, BinaryReader& prevStats)
+	{
+		m_trace.ProcessEnvironmentUpdated(process.GetId(), TCV("IDLE"), prevStats.GetPositionData(), prevStats.GetLeft(), TCV(""));
+		process.m_processStats = {};
+		process.m_sessionStats = {};
+		process.m_storageStats = {};
+		process.m_kernelStats = {};
+	}
+
 	bool Session::ProcessThreadStart(ProcessImpl& process)
 	{
 		return true;
@@ -3482,7 +3491,7 @@ namespace uba
 
 				UBA_ASSERT(msg.toKey != StringKeyZero);
 				auto insres = fs.writtenFiles.try_emplace(msg.toKey);
-				UBA_ASSERT(insres.second);
+				UBA_ASSERTF(insres.second, TC("Failed to add new file %s when moving from %s (destination already exists)"), msg.fromName.data, msg.toName.data);
 				WrittenFile& newFile = insres.first->second;
 				newFile = oldFile;
 				newFile.key = msg.toKey;
@@ -3625,20 +3634,19 @@ namespace uba
 		return true;
 	}
 
-	bool Session::GetNextProcess(Process& process, bool& outNewProcess, NextProcessInfo& outNextProcess, u32 prevExitCode, BinaryReader& statsReader)
+	bool Session::GetNextProcess(Process& process, bool& outNewProcess, bool& outShouldExit, NextProcessInfo& outNextProcess, u32 prevExitCode, BinaryReader& prevStats, u32 timeoutMs)
 	{
+		outShouldExit = false;
+		outNewProcess = false;
+
 		if (!m_getNextProcessFunction)
-		{
-			outNewProcess = false;
 			return true;
-		}
 
-		outNewProcess = m_getNextProcessFunction(process, outNextProcess, prevExitCode);
+		outNewProcess = m_getNextProcessFunction(process, outNextProcess, prevExitCode, prevStats, timeoutMs, outShouldExit);
 		if (!outNewProcess)
 			return true;
 
-		m_trace.ProcessEnvironmentUpdated(process.GetId(), outNextProcess.description, statsReader.GetPositionData(), statsReader.GetLeft(), outNextProcess.breadcrumbs);
-
+		m_trace.ProcessEnvironmentUpdated(process.GetId(), outNextProcess.description, prevStats.GetPositionData(), prevStats.GetLeft(), outNextProcess.breadcrumbs);
 		return true;
 	}
 
diff --git a/Engine/Source/Programs/UnrealBuildAccelerator/Common/Private/UbaSessionClient.cpp b/Engine/Source/Programs/UnrealBuildAccelerator/Common/Private/UbaSessionClient.cpp
--- a/Engine/Source/Programs/UnrealBuildAccelerator/Common/Private/UbaSessionClient.cpp
+++ b/Engine/Source/Programs/UnrealBuildAccelerator/Common/Private/UbaSessionClient.cpp
@@ -2610,8 +2610,9 @@ namespace uba
 		Session::PrintSessionStats(logger);
 	}
 
-	bool SessionClient::GetNextProcess(Process& process, bool& outNewProcess, NextProcessInfo& outNextProcess, u32 prevExitCode, BinaryReader& statsReader)
+	bool SessionClient::GetNextProcess(Process& process, bool& outNewProcess, bool& outShouldExit, NextProcessInfo& outNextProcess, u32 prevExitCode, BinaryReader& statsReader, u32 timeoutMs)
 	{
+		outShouldExit = false;
 		outNewProcess = false;
 
 		if (!m_remoteExecutionEnabled)
@@ -2622,46 +2623,75 @@ namespace uba
 		if (!FlushWrittenFiles(pi))
 			return false;
 
+		// Replace process stats i reader
 		ProcessStats processStats;
 		processStats.Read(statsReader, TraceVersion);
 		processStats.sendFiles = pi.m_processStats.sendFiles;
 
+		u64 startTime = GetTime();
+
 		StackBinaryReader<SendMaxSize> reader;
 		StackBinaryWriter<16 * 1024> writer;
-		NetworkMessage msg(m_client, ServiceId, SessionMessageType_GetNextProcess, writer);
-		writer.WriteU32(pi.m_id);
-		writer.WriteU32(prevExitCode);
-		processStats.Write(writer);
-		writer.WriteBytes(statsReader.GetPositionData(), statsReader.GetLeft());
-
-		if (!msg.Send(reader, m_stats.customMsg))
-			return false;
 
-		outNewProcess = reader.ReadBool();
-		if (outNewProcess)
+		for (u32 timeoutIndex=0;; ++timeoutIndex)
 		{
-			if (m_shouldSendLogToServer)
-				SendLogFileToServer(pi);
+			NetworkMessage msg(m_client, ServiceId, SessionMessageType_GetNextProcess, writer.Reset());
+			writer.WriteU32(pi.m_id);
+			writer.WriteU32(prevExitCode);
+			if (timeoutIndex == 0)
+			{
+				processStats.Write(writer);
+				writer.WriteBytes(statsReader.GetPositionData(), statsReader.GetLeft());
+			}
+
+			if (!msg.Send(reader.Reset(), m_stats.customMsg))
+				return false;
+
+			outNewProcess = reader.ReadBool();
+			outShouldExit = reader.ReadBool();
+			if (outNewProcess)
+			{
+				if (m_shouldSendLogToServer)
+					SendLogFileToServer(pi);
+
+				pi.m_exitCode = prevExitCode;
+				if (m_processFinished)
+					m_processFinished(&process);
+
+				outNextProcess.arguments = reader.ReadString();
+				outNextProcess.workingDir = reader.ReadString();
+				outNextProcess.description = reader.ReadString();
+				outNextProcess.logFile = reader.ReadString();
+				outNextProcess.overlayFiles.resize(reader.Read7BitEncoded());
+				reader.ReadBytes(outNextProcess.overlayFiles.data(), outNextProcess.overlayFiles.size());
 
-			pi.m_exitCode = prevExitCode;
-			if (m_processFinished)
-				m_processFinished(&process);
+				if (m_logToFile || (!outNextProcess.logFile.empty() && m_shouldSendLogToServer))
+				{
+					StringBuffer<512> logFile;
+					GetLogFileName(logFile, outNextProcess.logFile.c_str(), outNextProcess.arguments.c_str(), process.GetId());
+					outNextProcess.logFile = logFile.data;
+				}
+				break;
+			}
+			
+			if (!m_remoteExecutionEnabled)
+			{
+				outShouldExit = true;
+				return true;
+			}
 
-			outNextProcess.arguments = reader.ReadString();
-			outNextProcess.workingDir = reader.ReadString();
-			outNextProcess.description = reader.ReadString();
-			outNextProcess.logFile = reader.ReadString();
-			outNextProcess.overlayFiles.resize(reader.Read7BitEncoded());
-			reader.ReadBytes(outNextProcess.overlayFiles.data(), outNextProcess.overlayFiles.size());
+			u64 usedTimeMs = TimeToMs(GetTime() - startTime);
+			if (usedTimeMs >= timeoutMs)
+				return true;
 
-			if (m_logToFile || (!outNextProcess.logFile.empty() && m_shouldSendLogToServer))
+			if (timeoutIndex == 0)
 			{
-				StringBuffer<512> logFile;
-				GetLogFileName(logFile, outNextProcess.logFile.c_str(), outNextProcess.arguments.c_str(), process.GetId());
-				outNextProcess.logFile = logFile.data;
+				m_trace.ProcessEnvironmentUpdated(process.GetId(), TCV("IDLE"), statsReader.GetPositionData(), statsReader.GetLeft(), TCV(""));
+				pi.m_sessionStats = {};
 			}
-		}
 
+			Sleep(500); // TODO: BAAAAD.. never want to sleep.. this should incorporate with update thread request for processes
+		}
 		return SendUpdateDirectoryTable(reader.Reset());
 	}
 
diff --git a/Engine/Source/Programs/UnrealBuildAccelerator/Common/Private/UbaSessionServer.cpp b/Engine/Source/Programs/UnrealBuildAccelerator/Common/Private/UbaSessionServer.cpp
--- a/Engine/Source/Programs/UnrealBuildAccelerator/Common/Private/UbaSessionServer.cpp
+++ b/Engine/Source/Programs/UnrealBuildAccelerator/Common/Private/UbaSessionServer.cpp
@@ -121,6 +121,7 @@ namespace uba
 		Vector<u8> m_trackedOutputs;
 		bool m_cancelled = false;
 		bool m_allowCrossArchitecture = false;
+		bool m_isIdle = false;
 		u32 m_clientId = ~0u;
 		u32 m_sessionId = 0;
 		u64 m_startTime = 0;
@@ -1391,6 +1392,8 @@ namespace uba
 
 	bool SessionServer::WriteReceivedFiles(Vector<ReceivedFile>& receivedFiles, RemoteProcess& process, Timer& writeFilesTimer, const ConnectionInfo& connectionInfo)
 	{
+		if (receivedFiles.empty())
+			return true;
 		TimerScope ts(writeFilesTimer);
 		Atomic<bool> writeSuccess = true;
 		m_server.ParallelFor(u32(receivedFiles.size()), receivedFiles, [&](const WorkContext&, auto& it)
@@ -2018,12 +2021,26 @@ namespace uba
 			SkipReceivedFiles(receivedFiles, connectionInfo);
 		//processStats.writeFiles = writeFilesTimer
 
+		// Trace stream does not handle empty
+		StackBinaryWriter<1024> statsWriter;
+		if (!reader.GetLeft())
+		{
+			ProcessStats().Write(statsWriter);
+			SessionStats().Write(statsWriter);
+			StorageStats().Write(statsWriter);
+			KernelStats().Write(statsWriter);
+		}
+		BinaryReader prevStats = reader.GetLeft() ? reader : BinaryReader(statsWriter.GetData(), 0, statsWriter.GetPosition());
+
+
 		SCOPED_FUTEX(remoteProcess.m_exitedLock, exitedLock);
 		NextProcessInfo nextProcess;
 		bool newProcess;
+		bool shouldExit;
+		u32 timeoutMs = 0;
 		remoteProcess.m_exitCode = prevExitCode;
 		remoteProcess.m_done.Set();
-		bool success = GetNextProcess(remoteProcess, newProcess, nextProcess, prevExitCode, reader);
+		bool success = GetNextProcess(remoteProcess, newProcess, shouldExit, nextProcess, prevExitCode, prevStats, timeoutMs);
 		remoteProcess.m_exitCode = ~0u;
 		remoteProcess.m_done.Reset();
 		exitedLock.Leave();
@@ -2034,6 +2051,7 @@ namespace uba
 			return false;
 
 		writer.WriteBool(newProcess);
+		writer.WriteBool(shouldExit);
 		if (newProcess)
 		{
 			writer.WriteString(nextProcess.arguments);
@@ -2042,6 +2060,13 @@ namespace uba
 			writer.WriteString(nextProcess.logFile);
 			writer.Write7BitEncoded(nextProcess.overlayFiles.size());
 			writer.WriteBytes(nextProcess.overlayFiles.data(), nextProcess.overlayFiles.size());
+			remoteProcess.m_isIdle = false;
+		}
+		else if (!remoteProcess.m_isIdle)
+		{
+			// Stats are the previous call's stats
+			m_trace.ProcessEnvironmentUpdated(processId, TCV("IDLE"), prevStats.GetPositionData(), prevStats.GetLeft(), TCV(""));
+			remoteProcess.m_isIdle = true;
 		}
 		return true;
 	}
diff --git a/Engine/Source/Programs/UnrealBuildAccelerator/Common/Public/UbaNetwork.h b/Engine/Source/Programs/UnrealBuildAccelerator/Common/Public/UbaNetwork.h
--- a/Engine/Source/Programs/UnrealBuildAccelerator/Common/Public/UbaNetwork.h
+++ b/Engine/Source/Programs/UnrealBuildAccelerator/Common/Public/UbaNetwork.h
@@ -16,7 +16,7 @@ namespace uba
 	
 	static constexpr u32 SystemNetworkVersion = 1339;
 	static constexpr u32 StorageNetworkVersion = 4;
-	static constexpr u32 SessionNetworkVersion = 47;
+	static constexpr u32 SessionNetworkVersion = 48;
 	static constexpr u32 CacheNetworkVersion = 5;
 	static constexpr u32 CacheBucketVersion = 3;
 
diff --git a/Engine/Source/Programs/UnrealBuildAccelerator/Common/Public/UbaScheduler.h b/Engine/Source/Programs/UnrealBuildAccelerator/Common/Public/UbaScheduler.h
--- a/Engine/Source/Programs/UnrealBuildAccelerator/Common/Public/UbaScheduler.h
+++ b/Engine/Source/Programs/UnrealBuildAccelerator/Common/Public/UbaScheduler.h
@@ -139,7 +139,7 @@ namespace uba
 		bool RunCacheQuery();
 		bool RunQueuedProcess(bool isLocal);
 		bool RaceRemoteProcess();
-		bool HandleReuseMessage(Process& process, NextProcessInfo& outNextProcess, u32 prevExitCode);
+		bool HandleGetNextProcess(Process& process, NextProcessInfo& outNextProcess, u32 prevExitCode, BinaryReader& prevStats, u32 timeoutMs);
 		void ExitProcess(ExitProcessInfo& info, Process& process, u32 exitCode, bool fromCache, bool forReuse);
 		void SkipProcess(ProcessStartInfo2& info);
 		void UpdateStatusHitMissCount();
diff --git a/Engine/Source/Programs/UnrealBuildAccelerator/Common/Public/UbaSession.h b/Engine/Source/Programs/UnrealBuildAccelerator/Common/Public/UbaSession.h
--- a/Engine/Source/Programs/UnrealBuildAccelerator/Common/Public/UbaSession.h
+++ b/Engine/Source/Programs/UnrealBuildAccelerator/Common/Public/UbaSession.h
@@ -78,7 +78,7 @@ namespace uba
 		using CustomServiceFunction = Function<u32(Process& handle, const void* recv, u32 recvSize, void* send, u32 sendCapacity)>;
 		void RegisterCustomService(CustomServiceFunction&& function); // Register a custom service (that can be communicated with from the remote agents)
 
-		using GetNextProcessFunction = Function<bool(Process& handle, NextProcessInfo& outNextProcess, u32 prevExitCode)>;
+		using GetNextProcessFunction = Function<bool(Process& handle, NextProcessInfo& outNextProcess, u32 prevExitCode, BinaryReader& prevStats, u32 timeoutMs, bool& outShouldExit)>;
 		void RegisterGetNextProcess(GetNextProcessFunction&& function); // Register a custom service (that can be communicated with from the remote agents)
 
 		bool GetOutputFileSize(u64& outSize, const tchar* filePath); // Available if ShouldWriteToDisk is false
@@ -122,6 +122,7 @@ namespace uba
 		void FlushDeadProcesses();
 
 		const tchar* FixLogFile(const tchar* logFile, StringBufferBase& temp, const tchar* processArguments, u32 processId);
+		void SetProcessAsIdle(ProcessImpl& process, BinaryReader& prevStats);
 
 	protected:
 		Session(const SessionCreateInfo& info, const tchar* logPrefix, bool runningRemote, WorkManager& workManager);
@@ -187,7 +188,7 @@ namespace uba
 		virtual bool GetListDirectoryInfo(ListDirectoryResponse& out, const StringView& dirName, const StringKey& dirKey);
 		virtual bool WriteFilesToDisk(ProcessImpl& process, WrittenFile** files, u32 fileCount);
 		virtual bool AllocFailed(Process& process, const tchar* allocType, u32 error);
-		virtual bool GetNextProcess(Process& process, bool& outNewProcess, NextProcessInfo& outNextProcess, u32 prevExitCode, BinaryReader& statsReader);
+		virtual bool GetNextProcess(Process& process, bool& outNewProcess, bool& outShouldExit, NextProcessInfo& outNextProcess, u32 prevExitCode, BinaryReader& prevStats, u32 timeoutMs);
 		virtual bool CustomMessage(Process& process, BinaryReader& reader, BinaryWriter& writer);
 		virtual bool SHGetKnownFolderPath(Process& process, BinaryReader& reader, BinaryWriter& writer);
 		virtual bool HostRun(BinaryReader& reader, BinaryWriter& writer);
diff --git a/Engine/Source/Programs/UnrealBuildAccelerator/Common/Public/UbaSessionClient.h b/Engine/Source/Programs/UnrealBuildAccelerator/Common/Public/UbaSessionClient.h
--- a/Engine/Source/Programs/UnrealBuildAccelerator/Common/Public/UbaSessionClient.h
+++ b/Engine/Source/Programs/UnrealBuildAccelerator/Common/Public/UbaSessionClient.h
@@ -73,7 +73,7 @@ namespace uba
 		virtual bool WriteFilesToDisk(ProcessImpl& process, WrittenFile** files, u32 fileCount) override;
 		virtual bool AllocFailed(Process& process, const tchar* allocType, u32 error) override;
 		virtual void PrintSessionStats(Logger& logger) override;
-		virtual bool GetNextProcess(Process& process, bool& outNewProcess, NextProcessInfo& outNextProcess, u32 prevExitCode, BinaryReader& statsReader) override;
+		virtual bool GetNextProcess(Process& process, bool& outNewProcess, bool& outShouldExit, NextProcessInfo& outNextProcess, u32 prevExitCode, BinaryReader& statsReader, u32 timeoutMs) override;
 		virtual bool CustomMessage(Process& process, BinaryReader& reader, BinaryWriter& writer) override;
 		virtual bool SHGetKnownFolderPath(Process& process, BinaryReader& reader, BinaryWriter& writer) override;
 		virtual bool HostRun(BinaryReader& reader, BinaryWriter& writer) override;
diff --git a/Engine/Source/Programs/UnrealBuildAccelerator/Detours/Private/Windows/UbaDetoursFunctionsWin.cpp b/Engine/Source/Programs/UnrealBuildAccelerator/Detours/Private/Windows/UbaDetoursFunctionsWin.cpp
--- a/Engine/Source/Programs/UnrealBuildAccelerator/Detours/Private/Windows/UbaDetoursFunctionsWin.cpp
+++ b/Engine/Source/Programs/UnrealBuildAccelerator/Detours/Private/Windows/UbaDetoursFunctionsWin.cpp
@@ -1148,7 +1148,7 @@ extern "C"
 		return g_runningRemote;
 	}
 
-	UBA_DETOURED_API bool UbaRequestNextProcess(u32 prevExitCode, wchar_t* outArguments, u32 outArgumentsCapacity)
+	UBA_DETOURED_API bool UbaRequestNextProcess2(u32 prevExitCode, wchar_t* outArguments, u32 outArgumentsCapacity, u32 timeOutMs, bool* outShouldExit)
 	{
 		#if UBA_DEBUG_LOG_ENABLED
 		FlushDebugLog();
@@ -1169,12 +1169,14 @@ extern "C"
 		{
 			RPC_MESSAGE(GetNextProcess, log)
 			writer.WriteU32(prevExitCode);
+			writer.WriteU32(timeOutMs);
 			g_stats.Write(writer);
 			g_kernelStats.Write(writer);
 
 			writer.Flush();
 			BinaryReader reader;
 			newProcess = reader.ReadBool();
+			*outShouldExit = reader.ReadBool();
 			if (newProcess)
 			{
 				reader.ReadString(outArguments, outArgumentsCapacity);
@@ -1186,11 +1188,11 @@ extern "C"
 			}
 		}
 
+		g_stats = {};
+		g_kernelStats = {};
+
 		if (newProcess)
 		{
-			g_kernelStats = {};
-			g_stats = {};
-
 			#if UBA_DEBUG_LOG_ENABLED
 			SuppressCreateFileDetourScope scope;
 			HANDLE debugFile = (HANDLE)g_debugFile;
@@ -1214,4 +1216,11 @@ extern "C"
 		Rpc_UpdateTables();
 		return newProcess;
 	}
+
+	UBA_DETOURED_API bool UbaRequestNextProcess(u32 prevExitCode, wchar_t* outArguments, u32 outArgumentsCapacity)
+	{
+		bool shouldExit;
+		return UbaRequestNextProcess2(prevExitCode, outArguments, outArgumentsCapacity, 0, &shouldExit);
+	}
+
 }
diff --git a/Engine/Source/Programs/UnrealBuildAccelerator/Detours/Public/UbaDetoursApi.h b/Engine/Source/Programs/UnrealBuildAccelerator/Detours/Public/UbaDetoursApi.h
--- a/Engine/Source/Programs/UnrealBuildAccelerator/Detours/Public/UbaDetoursApi.h
+++ b/Engine/Source/Programs/UnrealBuildAccelerator/Detours/Public/UbaDetoursApi.h
@@ -26,6 +26,9 @@ extern "C"
 	// prevExitCode should contain exit code of "finished" process.
 	// Will write arguments into outArguments. outArgumentsCapacity should be in characters (not bytes)
 	UBA_DETOURED_API bool UbaRequestNextProcess(unsigned int prevExitCode, wchar_t* outArguments, unsigned int outArgumentsCapacity);
+
+	// Same as above but with a timeout how long it can wait before giving up to get next process and a "should exit" variable telling the calling process to exit
+	UBA_DETOURED_API bool UbaRequestNextProcess2(unsigned int prevExitCode, wchar_t* outArguments, unsigned int outArgumentsCapacity, unsigned int timeOutMs, bool* outShouldExit);
 }
 
 
diff --git a/Engine/Source/Programs/UnrealBuildAccelerator/Test/Private/UbaTestScheduler.cpp b/Engine/Source/Programs/UnrealBuildAccelerator/Test/Private/UbaTestScheduler.cpp
--- a/Engine/Source/Programs/UnrealBuildAccelerator/Test/Private/UbaTestScheduler.cpp
+++ b/Engine/Source/Programs/UnrealBuildAccelerator/Test/Private/UbaTestScheduler.cpp
@@ -32,31 +32,51 @@ namespace uba
 			});
 	}
 
-	bool TestLocalScheduleReuse(LoggerWithWriter& logger, const StringBufferBase& testRootDir)
+	bool RunSchedulerReuseTest(LoggerWithWriter& logger, SessionServer& session, const tchar* workingDir, bool isLocal)
 	{
-		return RunLocal(logger, testRootDir, [](LoggerWithWriter& logger, SessionServer& session, const tchar* workingDir, const RunProcessFunction& runProcess, void* extraData)
+		SchedulerCreateInfo info(session);
+		info.enableProcessReuse = true;
+		info.maxLocalProcessors = isLocal ? 1 : 0;
+		Scheduler scheduler(info);
+
+		StringBuffer<> testApp;
+		GetTestAppPath(logger, testApp);
+
+		ProcessStartInfo processInfo;
+		processInfo.application = testApp.data;
+		processInfo.workingDir = workingDir;
+		processInfo.arguments = TC("-reuse");
+
+		EnqueueProcessInfo epi(processInfo);
+		scheduler.EnqueueProcess(epi);
+		scheduler.Start();
+
+		u32 queued, activeLocal, activeRemote, finished;
+		bool addSecond = true;
+		do
+		{
+			scheduler.GetStats(queued, activeLocal, activeRemote, finished);
+			if (finished == 1 && addSecond)
 			{
-				SchedulerCreateInfo info(session);
-				info.enableProcessReuse = true;
-				Scheduler scheduler(info);
+				Sleep(100);
+				scheduler.EnqueueProcess(epi);
+				addSecond = false;
+			}
+		} while (finished != 2);
 
-				StringBuffer<> testApp;
-				GetTestAppPath(logger, testApp);
+		if (!isLocal)
+			session.GetServer().DisconnectClients();
 
-				ProcessStartInfo processInfo;
-				processInfo.application = testApp.data;
-				processInfo.workingDir = workingDir;
-				processInfo.arguments = TC("-reuse");
+		scheduler.Stop();
 
-				EnqueueProcessInfo epi(processInfo);
-				scheduler.EnqueueProcess(epi);
-				scheduler.Start();
-
-				u32 queued, activeLocal, activeRemote, finished;
-				do { scheduler.GetStats(queued, activeLocal, activeRemote, finished); } while (finished != 1);
+		return true;
+	}
 
-				scheduler.Stop();
-				return true;
+	bool TestLocalScheduleReuse(LoggerWithWriter& logger, const StringBufferBase& testRootDir)
+	{
+		return RunLocal(logger, testRootDir, [](LoggerWithWriter& logger, SessionServer& session, const tchar* workingDir, const RunProcessFunction& runProcess, void* extraData)
+			{
+				return RunSchedulerReuseTest(logger, session, workingDir, true);
 			});
 	}
 
@@ -64,30 +84,7 @@ namespace uba
 	{
 		return RunRemote(logger, testRootDir, [](LoggerWithWriter& logger, SessionServer& session, const tchar* workingDir, const RunProcessFunction& runProcess, void* extraData)
 			{
-				SchedulerCreateInfo info(session);
-				info.enableProcessReuse = true;
-				info.maxLocalProcessors = 0;
-				Scheduler scheduler(info);
-
-				StringBuffer<> testApp;
-				GetTestAppPath(logger, testApp);
-
-				ProcessStartInfo processInfo;
-				processInfo.application = testApp.data;
-				processInfo.workingDir = workingDir;
-				processInfo.arguments = TC("-reuse");
-
-				EnqueueProcessInfo epi(processInfo);
-				scheduler.EnqueueProcess(epi);
-				scheduler.Start();
-
-				u32 queued, activeLocal, activeRemote, finished;
-				do { scheduler.GetStats(queued, activeLocal, activeRemote, finished); } while (finished != 1);
-
-				session.GetServer().DisconnectClients(); // Must make sure all is disconnected since scheduler goes out of scope
-
-				scheduler.Stop();
-				return true;
+				return RunSchedulerReuseTest(logger, session, workingDir, false);
 			});
 	}
 }
diff --git a/Engine/Source/Programs/UnrealBuildAccelerator/TestApp/Private/Windows/UbaTestAppWin.cpp b/Engine/Source/Programs/UnrealBuildAccelerator/TestApp/Private/Windows/UbaTestAppWin.cpp
--- a/Engine/Source/Programs/UnrealBuildAccelerator/TestApp/Private/Windows/UbaTestAppWin.cpp
+++ b/Engine/Source/Programs/UnrealBuildAccelerator/TestApp/Private/Windows/UbaTestAppWin.cpp
@@ -258,6 +258,9 @@ int wmain(int argc, wchar_t* argv[])
 		wchar_t arguments[1024];
 		if (requestNextProcess(0, arguments, sizeof(arguments)))
 			return LogError(L"Didn't expect another process");
+		while (true)
+			if (requestNextProcess(0, arguments, sizeof(arguments)))
+				break;
 	}
 	else if (wcsncmp(argv[1], L"-file=", 6) == 0)
 	{
