[UBA]
* Fixed bug where process stats was not properly reset on remote side

diff --git a/Engine/Source/Programs/UnrealBuildAccelerator/Common/Private/UbaSessionClient.cpp b/Engine/Source/Programs/UnrealBuildAccelerator/Common/Private/UbaSessionClient.cpp
--- a/Engine/Source/Programs/UnrealBuildAccelerator/Common/Private/UbaSessionClient.cpp
+++ b/Engine/Source/Programs/UnrealBuildAccelerator/Common/Private/UbaSessionClient.cpp
@@ -2685,10 +2685,7 @@ namespace uba
 				return true;
 
 			if (timeoutIndex == 0)
-			{
-				m_trace.ProcessEnvironmentUpdated(process.GetId(), TCV("IDLE"), statsReader.GetPositionData(), statsReader.GetLeft(), TCV(""));
-				pi.m_sessionStats = {};
-			}
+				SetProcessAsIdle(pi, statsReader);
 
 			Sleep(500); // TODO: BAAAAD.. never want to sleep.. this should incorporate with update thread request for processes
 		}
