[UBA]
* Fixed so when ThreadLoop takes m_updateThreadWakeup event and run logic to spawn process checks if any thread is sitting in GetNextProcess function before spawning process. If there is one, it sets the m_updateThreadWakeup again and loop around. This means that new processes are always spawned by GetNextProcess if available.

diff --git a/Engine/Source/Programs/UnrealBuildAccelerator/Common/Private/UbaScheduler.cpp b/Engine/Source/Programs/UnrealBuildAccelerator/Common/Private/UbaScheduler.cpp
--- a/Engine/Source/Programs/UnrealBuildAccelerator/Common/Private/UbaScheduler.cpp
+++ b/Engine/Source/Programs/UnrealBuildAccelerator/Common/Private/UbaScheduler.cpp
@@ -931,6 +931,16 @@ namespace uba
 					m_activeLocalProcessWeight += entry.info->weight;
 			}
 
+
+			// If there are queued getnext, then we try to hand over running queued processes
+			// Setting the event will wake up the getnext thread.. and If the getnext thread is timing out we will retry
+			if (isLocal && !forReuse && m_getNextQueueCount)
+			{
+				m_updateThreadWakeup.Set();
+				Sleep(1);
+				return ~0u;
+			}
+
 			outPrevStatus = entry.runStatus;
 			entry.runStatus = RunStatus_Running;
 			return i;
@@ -1197,6 +1207,11 @@ namespace uba
 		u64 startTime = GetTime();
 		bool firstWait = true;
 
+		// We always want to use the getnext processes before spawning new ones.
+		// This logic makes sure that the update thread does not pick up actions if getNextQueueCount > 0
+		++m_getNextQueueCount;
+		auto queueGuard = MakeGuard([&]() { --m_getNextQueueCount; m_updateThreadWakeup.Set(); });
+
 		while (true)
 		{
 			RunStatus prevStatus;
@@ -1222,6 +1237,9 @@ namespace uba
 					m_updateThreadWakeup.Set();
 				continue;
 			}
+
+			queueGuard.Execute();
+
 			//UBA_ASSERT(prevStatus != RunStatus_QueuedForCache);
 			auto& processEntry = m_processEntries[indexToRun];
 			auto newInfo = processEntry.info;
diff --git a/Engine/Source/Programs/UnrealBuildAccelerator/Common/Public/UbaScheduler.h b/Engine/Source/Programs/UnrealBuildAccelerator/Common/Public/UbaScheduler.h
--- a/Engine/Source/Programs/UnrealBuildAccelerator/Common/Public/UbaScheduler.h
+++ b/Engine/Source/Programs/UnrealBuildAccelerator/Common/Public/UbaScheduler.h
@@ -198,6 +198,7 @@ namespace uba
 		Atomic<float> m_activeRacingProcessWeight = 0.0f;
 		u32 m_activeCacheQueries = 0;
 		Atomic<u32> m_activeCacheDownloads = 0;
+		Atomic<u32> m_getNextQueueCount = 0;
 
 		Atomic<u32> m_totalProcesses = 0;
 		Atomic<u32> m_queuedProcesses = 0;
