[UBA]
* Changed so GetNextProcess callback on remote helpers are now only sending one message to try to get new process. And if it doesn't it adds itself to a linked list if processes waiting in GetNextProcess. The update thread is now using that linked list to send one message per loop iteration. Each message then contains all the processes sitting in GetNextProcess. This will reduce message spamming by a lot

diff --git a/Engine/Source/Programs/UnrealBuildAccelerator/Common/Private/UbaSessionClient.cpp b/Engine/Source/Programs/UnrealBuildAccelerator/Common/Private/UbaSessionClient.cpp
--- a/Engine/Source/Programs/UnrealBuildAccelerator/Common/Private/UbaSessionClient.cpp
+++ b/Engine/Source/Programs/UnrealBuildAccelerator/Common/Private/UbaSessionClient.cpp
@@ -65,6 +65,16 @@ namespace uba
 		u32 attributes;
 	};
 
+	struct SessionClient::GetNextEntry
+	{
+		u32 processId;
+		bool& outNewProcess;
+		bool& outShouldExit;
+		NextProcessInfo& outNextProcess;
+		GetNextEntry* next;
+		Event done;
+	};
+
 	SessionClient::SessionClient(const SessionClientCreateInfo& info)
 	: Session(info, TC("UbaSessionClient"), true, info.client)
 	,	m_client(info.client)
@@ -1859,6 +1869,7 @@ namespace uba
 		auto logFile = pi.m_startInfo.logFile;
 		if (!logFile || !*logFile)
 			return;
+		pi.m_startInfo.logFile = nullptr;
 		WrittenFile f;
 		f.backedName = logFile;
 		f.attributes = DefaultAttributes();
@@ -2016,6 +2027,12 @@ namespace uba
 				lastWaitTime = GetTime();
 			}
 
+			if (!SendGetNextMessage())
+			{
+				m_loop = false;
+				break;
+			}
+
 			bool canSpawn = TimeToMs(GetTime() - lastWaitTime) > waitTimeToSpawnAfterKillMs && m_allowSpawn;
 			if (!canSpawn)
 				waitTimeoutMs = 500;
@@ -2297,6 +2314,7 @@ namespace uba
 			}
 		}
 
+		CancelGetNextMessageQueue();
 		CancelAllProcessesAndWait(); // If we got the exit from server there is no point sending anything more back.. cancel everything
 
 		u32 retry = 0;
@@ -2610,86 +2628,161 @@ namespace uba
 		Session::PrintSessionStats(logger);
 	}
 
-	bool SessionClient::GetNextProcess(Process& process, bool& outNewProcess, bool& outShouldExit, NextProcessInfo& outNextProcess, u32 prevExitCode, BinaryReader& statsReader, u32 timeoutMs)
+	bool SessionClient::SendGetNextMessage()
 	{
-		outShouldExit = false;
-		outNewProcess = false;
+		SCOPED_FUTEX(m_getNextEntryLock, lock);
+		if (!m_firstGetNextEntry)
+			return true;
 
 		if (!m_remoteExecutionEnabled)
+		{
+			lock.Leave();
+			CancelGetNextMessageQueue();
 			return true;
+		}
+
+		StackBinaryReader<SendMaxSize> reader;
+		StackBinaryWriter<16 * 1024> writer;
+		NetworkMessage msg(m_client, ServiceId, SessionMessageType_GetNextProcess, writer);
+		writer.WriteBool(false);
+		for (auto i=m_firstGetNextEntry; i; i=i->next)
+			writer.WriteU32(i->processId);
+
+		if (!msg.Send(reader, m_stats.customMsg))
+			return false;
+
+		for (auto i=m_firstGetNextEntry; i; i=i->next)
+			if (HandleGetNextResponse(reader, i->processId, i->outNewProcess, i->outShouldExit, i->outNextProcess))
+				i->done.Set();
+
+		return true;
+	}
+
+	void SessionClient::CancelGetNextMessageQueue()
+	{
+		SCOPED_FUTEX(m_getNextEntryLock, lock);
+		for (auto i=m_firstGetNextEntry; i; i=i->next)
+		{
+			i->outNewProcess = false;
+			i->outShouldExit = true;
+			i->done.Set();
+		}
+	}
+
+	bool SessionClient::HandleGetNextResponse(BinaryReader& reader, u32 processId, bool& outNewProcess, bool& outShouldExit, NextProcessInfo& outNextProcess)
+	{
+		outNewProcess = reader.ReadBool();
+		outShouldExit = reader.ReadBool();
+		if (!outNewProcess)
+			return false;
+
+		outNextProcess.arguments = reader.ReadString();
+		outNextProcess.workingDir = reader.ReadString();
+		outNextProcess.description = reader.ReadString();
+		outNextProcess.logFile = reader.ReadString();
+		outNextProcess.overlayFiles.resize(reader.Read7BitEncoded());
+		reader.ReadBytes(outNextProcess.overlayFiles.data(), outNextProcess.overlayFiles.size());
+
+		if (m_logToFile || (!outNextProcess.logFile.empty() && m_shouldSendLogToServer))
+		{
+			StringBuffer<512> logFile;
+			GetLogFileName(logFile, outNextProcess.logFile.c_str(), outNextProcess.arguments.c_str(), processId);
+			outNextProcess.logFile = logFile.data;
+		}
+		return true;
+	}
+
+	bool SessionClient::GetNextProcess(Process& process, bool& outNewProcess, bool& outShouldExit, NextProcessInfo& outNextProcess, u32 prevExitCode, BinaryReader& statsReader, u32 timeoutMs)
+	{
+		outShouldExit = true;
+		outNewProcess = false;
 
 		auto& pi = (ProcessImpl&)process;
 
+		if (!m_remoteExecutionEnabled)
+			return true;
+
 		if (!FlushWrittenFiles(pi))
 			return false;
 
+		if (m_shouldSendLogToServer)
+			SendLogFileToServer(pi);
+
 		// Replace process stats i reader
 		ProcessStats processStats;
 		processStats.Read(statsReader, TraceVersion);
 		processStats.sendFiles = pi.m_processStats.sendFiles;
 
-		u64 startTime = GetTime();
-
 		StackBinaryReader<SendMaxSize> reader;
 		StackBinaryWriter<16 * 1024> writer;
 
-		for (u32 timeoutIndex=0;; ++timeoutIndex)
-		{
-			NetworkMessage msg(m_client, ServiceId, SessionMessageType_GetNextProcess, writer.Reset());
-			writer.WriteU32(pi.m_id);
-			writer.WriteU32(prevExitCode);
-			if (timeoutIndex == 0)
-			{
-				processStats.Write(writer);
-				writer.WriteBytes(statsReader.GetPositionData(), statsReader.GetLeft());
-			}
+		// Try getting process from this thread
+		NetworkMessage msg(m_client, ServiceId, SessionMessageType_GetNextProcess, writer);
+		writer.WriteBool(true);
+		writer.WriteU32(pi.m_id);
+		writer.WriteU32(prevExitCode);
+		processStats.Write(writer);
+		writer.WriteBytes(statsReader.GetPositionData(), statsReader.GetLeft());
 
-			if (!msg.Send(reader.Reset(), m_stats.customMsg))
-				return false;
+		pi.m_exitCode = prevExitCode;
+		if (m_processFinished)
+			m_processFinished(&process);
 
-			outNewProcess = reader.ReadBool();
-			outShouldExit = reader.ReadBool();
-			if (outNewProcess)
-			{
-				if (m_shouldSendLogToServer)
-					SendLogFileToServer(pi);
+		if (!msg.Send(reader, m_stats.customMsg))
+			return false;
+		if (HandleGetNextResponse(reader, pi.m_id, outNewProcess, outShouldExit, outNextProcess))
+			return SendUpdateDirectoryTable(reader.Reset());
 
-				pi.m_exitCode = prevExitCode;
-				if (m_processFinished)
-					m_processFinished(&process);
 
-				outNextProcess.arguments = reader.ReadString();
-				outNextProcess.workingDir = reader.ReadString();
-				outNextProcess.description = reader.ReadString();
-				outNextProcess.logFile = reader.ReadString();
-				outNextProcess.overlayFiles.resize(reader.Read7BitEncoded());
-				reader.ReadBytes(outNextProcess.overlayFiles.data(), outNextProcess.overlayFiles.size());
+		if (!timeoutMs)
+			return false;
 
-				if (m_logToFile || (!outNextProcess.logFile.empty() && m_shouldSendLogToServer))
-				{
-					StringBuffer<512> logFile;
-					GetLogFileName(logFile, outNextProcess.logFile.c_str(), outNextProcess.arguments.c_str(), process.GetId());
-					outNextProcess.logFile = logFile.data;
-				}
-				break;
-			}
-			
-			if (!m_remoteExecutionEnabled)
-			{
-				outShouldExit = true;
-				return true;
-			}
+		SetProcessAsIdle(pi, statsReader);
 
-			u64 usedTimeMs = TimeToMs(GetTime() - startTime);
-			if (usedTimeMs >= timeoutMs)
-				return true;
+		// No process found, let main loop query free processes
+
+		SCOPED_FUTEX(m_getNextEntryLock, lock);
+		GetNextEntry entry { pi.m_id, outNewProcess, outShouldExit, outNextProcess };
+		entry.next = m_firstGetNextEntry;
+		m_firstGetNextEntry = &entry;
+		entry.done.Create(EventResetType_Manual);
+		lock.Leave();
 
-			if (timeoutIndex == 0)
-				SetProcessAsIdle(pi, statsReader);
+		// Wait for main loop to set event (or timeout)
+		entry.done.IsSet(timeoutMs);
 
-			Sleep(500); // TODO: BAAAAD.. never want to sleep.. this should incorporate with update thread request for processes
+		// Remove from linked list
+		lock.Enter();
+		for (GetNextEntry* i=m_firstGetNextEntry, *p=nullptr; i; p=i, i=i->next)
+		{
+			if (i != &entry)
+				continue;
+			if (!p)
+				m_firstGetNextEntry = entry.next;
+			else
+				p->next = entry.next;
+			break;
 		}
-		return SendUpdateDirectoryTable(reader.Reset());
+		lock.Leave();
+
+		if (outNewProcess)
+			return SendUpdateDirectoryTable(reader.Reset());
+
+		if (!m_remoteExecutionEnabled)
+			outShouldExit = true;
+
+		return true;
+
+		/*
+		u64 usedTimeMs = TimeToMs(GetTime() - startTime);
+		if (usedTimeMs >= timeoutMs)
+			return true;
+
+		if (timeoutIndex == 0)
+			SetProcessAsIdle(pi, statsReader);
+
+		Sleep(500); // TODO: BAAAAD.. never want to sleep.. this should incorporate with update thread request for processes
+		*/
 	}
 
 	bool SessionClient::CustomMessage(Process& process, BinaryReader& reader, BinaryWriter& writer)
diff --git a/Engine/Source/Programs/UnrealBuildAccelerator/Common/Private/UbaSessionServer.cpp b/Engine/Source/Programs/UnrealBuildAccelerator/Common/Private/UbaSessionServer.cpp
--- a/Engine/Source/Programs/UnrealBuildAccelerator/Common/Private/UbaSessionServer.cpp
+++ b/Engine/Source/Programs/UnrealBuildAccelerator/Common/Private/UbaSessionServer.cpp
@@ -1997,77 +1997,88 @@ namespace uba
 
 	bool SessionServer::HandleGetNextProcess(const ConnectionInfo& connectionInfo, const WorkContext& workContext, BinaryReader& reader, BinaryWriter& writer)
 	{
-		u32 processId = reader.ReadU32();
-		u32 prevExitCode = reader.ReadU32();
-		ProcessHandle h(GetProcess(processId));
-		if (!h.IsValid())
-		{
-			m_logger.Info(TC("Failed to find process for id %u when receiving GetNextProcess message"), processId);
-			return false;
-		}
-
-		auto& remoteProcess = *(RemoteProcess*)h.m_process;
-		auto receivedFiles = std::move(remoteProcess.m_receivedFiles);
-
-		Timer writeFilesTimer; // TODO: This should be added somewhere
-		
-		auto& startInfo = remoteProcess.m_startInfo; // This is kind of wrong with processor reuse but we know the rules is the same
-		if (startInfo.writeOutputFilesOnFail || startInfo.rules->IsExitCodeSuccess(prevExitCode))
-		{
-			if (!WriteReceivedFiles(receivedFiles, remoteProcess, writeFilesTimer, connectionInfo))
-				prevExitCode = UBA_EXIT_CODE(20);
-		}
-		else
-			SkipReceivedFiles(receivedFiles, connectionInfo);
-		//processStats.writeFiles = writeFilesTimer
+		bool isFirstGet = reader.ReadBool();
 
 		// Trace stream does not handle empty
 		StackBinaryWriter<1024> statsWriter;
-		if (!reader.GetLeft())
+		if (!isFirstGet)
 		{
 			ProcessStats().Write(statsWriter);
 			SessionStats().Write(statsWriter);
 			StorageStats().Write(statsWriter);
 			KernelStats().Write(statsWriter);
 		}
-		BinaryReader prevStats = reader.GetLeft() ? reader : BinaryReader(statsWriter.GetData(), 0, statsWriter.GetPosition());
 
+		do
+		{
+			u32 processId = reader.ReadU32();
+			u32 prevExitCode = ~0u;
+			if (isFirstGet)
+				prevExitCode = reader.ReadU32();
 
-		SCOPED_FUTEX(remoteProcess.m_exitedLock, exitedLock);
-		NextProcessInfo nextProcess;
-		bool newProcess;
-		bool shouldExit;
-		u32 timeoutMs = 0;
-		remoteProcess.m_exitCode = prevExitCode;
-		remoteProcess.m_done.Set();
-		bool success = GetNextProcess(remoteProcess, newProcess, shouldExit, nextProcess, prevExitCode, prevStats, timeoutMs);
-		remoteProcess.m_exitCode = ~0u;
-		remoteProcess.m_done.Reset();
-		exitedLock.Leave();
+			ProcessHandle h(GetProcess(processId));
+			if (!h.IsValid())
+			{
+				m_logger.Info(TC("Failed to find process for id %u when receiving GetNextProcess message"), processId);
+				return false;
+			}
 
-		ReleaseReceivedFiles(receivedFiles);
+			auto& remoteProcess = *(RemoteProcess*)h.m_process;
+			auto receivedFiles = std::move(remoteProcess.m_receivedFiles);
 
-		if (!success)
-			return false;
+			Timer writeFilesTimer; // TODO: This should be added somewhere
+		
+			auto& startInfo = remoteProcess.m_startInfo; // This is kind of wrong with processor reuse but we know the rules is the same
+			if (startInfo.writeOutputFilesOnFail || startInfo.rules->IsExitCodeSuccess(prevExitCode))
+			{
+				if (!WriteReceivedFiles(receivedFiles, remoteProcess, writeFilesTimer, connectionInfo))
+					prevExitCode = UBA_EXIT_CODE(20);
+			}
+			else
+				SkipReceivedFiles(receivedFiles, connectionInfo);
+			//processStats.writeFiles = writeFilesTimer
 
-		writer.WriteBool(newProcess);
-		writer.WriteBool(shouldExit);
-		if (newProcess)
-		{
-			writer.WriteString(nextProcess.arguments);
-			writer.WriteString(nextProcess.workingDir);
-			writer.WriteString(nextProcess.description);
-			writer.WriteString(nextProcess.logFile);
-			writer.Write7BitEncoded(nextProcess.overlayFiles.size());
-			writer.WriteBytes(nextProcess.overlayFiles.data(), nextProcess.overlayFiles.size());
-			remoteProcess.m_isIdle = false;
-		}
-		else if (!remoteProcess.m_isIdle)
-		{
-			// Stats are the previous call's stats
-			m_trace.ProcessEnvironmentUpdated(processId, TCV("IDLE"), prevStats.GetPositionData(), prevStats.GetLeft(), TCV(""));
-			remoteProcess.m_isIdle = true;
+			BinaryReader prevStats = isFirstGet ? reader : BinaryReader(statsWriter.GetData(), 0, statsWriter.GetPosition());
+
+
+			SCOPED_FUTEX(remoteProcess.m_exitedLock, exitedLock);
+			NextProcessInfo nextProcess;
+			bool newProcess;
+			bool shouldExit;
+			u32 timeoutMs = 0;
+			remoteProcess.m_exitCode = prevExitCode;
+			remoteProcess.m_done.Set();
+			bool success = GetNextProcess(remoteProcess, newProcess, shouldExit, nextProcess, prevExitCode, prevStats, timeoutMs);
+			remoteProcess.m_exitCode = ~0u;
+			remoteProcess.m_done.Reset();
+			exitedLock.Leave();
+
+			ReleaseReceivedFiles(receivedFiles);
+
+			if (!success)
+				return false;
+
+			writer.WriteBool(newProcess);
+			writer.WriteBool(shouldExit);
+			if (newProcess)
+			{
+				writer.WriteString(nextProcess.arguments);
+				writer.WriteString(nextProcess.workingDir);
+				writer.WriteString(nextProcess.description);
+				writer.WriteString(nextProcess.logFile);
+				writer.Write7BitEncoded(nextProcess.overlayFiles.size());
+				writer.WriteBytes(nextProcess.overlayFiles.data(), nextProcess.overlayFiles.size());
+				remoteProcess.m_isIdle = false;
+			}
+			else if (!remoteProcess.m_isIdle)
+			{
+				// Stats are the previous call's stats
+				m_trace.ProcessEnvironmentUpdated(processId, TCV("IDLE"), prevStats.GetPositionData(), prevStats.GetLeft(), TCV(""));
+				remoteProcess.m_isIdle = true;
+			}
 		}
+		while (!isFirstGet && reader.GetLeft());
+
 		return true;
 	}
 
diff --git a/Engine/Source/Programs/UnrealBuildAccelerator/Common/Public/UbaSessionClient.h b/Engine/Source/Programs/UnrealBuildAccelerator/Common/Public/UbaSessionClient.h
--- a/Engine/Source/Programs/UnrealBuildAccelerator/Common/Public/UbaSessionClient.h
+++ b/Engine/Source/Programs/UnrealBuildAccelerator/Common/Public/UbaSessionClient.h
@@ -100,6 +100,9 @@ namespace uba
 		void Connect();
 		void BuildEnvironmentVariables(BinaryReader& reader);
 		bool SendProcessAvailable(Vector<InternalProcessStartInfo>& outStartInfos, Vector<u32>& outCancelled, float availableWeight);
+		bool SendGetNextMessage();
+		void CancelGetNextMessageQueue();
+		bool HandleGetNextResponse(BinaryReader& reader, u32 processId, bool& outNewProcess, bool& outShouldExit, NextProcessInfo& outNextProcess);
 		void SendReturnProcess(u32 processId, const tchar* reason);
 		bool SendProcessInputs(ProcessImpl& process);
 		bool SendProcessFinished(ProcessImpl& process, u32 exitCode);
@@ -179,6 +182,10 @@ namespace uba
 
 		SessionSummaryStats m_stats;
 
+		struct GetNextEntry;
+		GetNextEntry* m_firstGetNextEntry = nullptr;
+		Futex m_getNextEntryLock;
+
 		Atomic<u64> m_bestPing = 0;
 		Atomic<u64> m_lastPing = 0;
 		u64 m_lastPingSendTime = 0;
