[UBA]
* Changed so visualizer is not dirtying process bitmaps if only changing horizontal scale and no read/write colors are visible

diff --git a/Engine/Source/Programs/UnrealBuildAccelerator/Visualizer/Private/Windows/UbaVisualizerWin.cpp b/Engine/Source/Programs/UnrealBuildAccelerator/Visualizer/Private/Windows/UbaVisualizerWin.cpp
--- a/Engine/Source/Programs/UnrealBuildAccelerator/Visualizer/Private/Windows/UbaVisualizerWin.cpp
+++ b/Engine/Source/Programs/UnrealBuildAccelerator/Visualizer/Private/Windows/UbaVisualizerWin.cpp
@@ -808,8 +808,8 @@ namespace uba
 						for (auto& session : m_traceView.sessions)
 							for (auto& processor : session.processors)
 								for (auto& process : processor.processes)
-									//if (TimeToMs(process.writeFilesTime, m_traceView.frequency) >= 300 || TimeToMs(process.createFilesTime, m_traceView.frequency) >= 300)
-									process.bitmapDirty = true;
+									if (process.type == TraceView::ProcessType_Normal && (process.createFilesTime || process.writeFilesTime))
+										process.bitmapDirty = true;
 				}
 				else if (m_activeSection == 1) // active processes
 				{
