[UBA]
* Fixed so idle exit code is always 0

diff --git a/Engine/Source/Programs/UnrealBuildAccelerator/Common/Private/UbaSessionServer.cpp b/Engine/Source/Programs/UnrealBuildAccelerator/Common/Private/UbaSessionServer.cpp
--- a/Engine/Source/Programs/UnrealBuildAccelerator/Common/Private/UbaSessionServer.cpp
+++ b/Engine/Source/Programs/UnrealBuildAccelerator/Common/Private/UbaSessionServer.cpp
@@ -2012,7 +2012,7 @@ namespace uba
 		do
 		{
 			u32 processId = reader.ReadU32();
-			u32 prevExitCode = ~0u;
+			u32 prevExitCode = 0; // Idle exit code is always success
 			if (isFirstGet)
 				prevExitCode = reader.ReadU32();
 
