[UBA]
* Re-added support for copy-on-write mappings in detoured process. It broke when refactoring to use shared memory. This fixes so radlink works with uba again

diff --git a/Engine/Source/Programs/UnrealBuildAccelerator/Core/Private/UbaSharedMemoryView.cpp b/Engine/Source/Programs/UnrealBuildAccelerator/Core/Private/UbaSharedMemoryView.cpp
--- a/Engine/Source/Programs/UnrealBuildAccelerator/Core/Private/UbaSharedMemoryView.cpp
+++ b/Engine/Source/Programs/UnrealBuildAccelerator/Core/Private/UbaSharedMemoryView.cpp
@@ -21,7 +21,7 @@ namespace uba
 		return InternalInit(handle, capacity, baseAddress, hint);
 	}
 
-	bool SharedMemoryView::Init(SharedMemoryAllocatorHandle handle, u8* sharedMemory, BinaryReader& reader, u64 capacity, const tchar* hint)
+	bool SharedMemoryView::Init(SharedMemoryAllocatorHandle handle, u8* sharedMemory, BinaryReader& reader, u64 capacity, const tchar* hint, SharedMemoryMapType type)
 	{
 		u32 sliceCount = reader.ReadU32();
 		if (sliceCount == 1 && sharedMemory)
@@ -34,7 +34,7 @@ namespace uba
 		}
 		if (!InternalInit(handle, capacity, nullptr, hint))
 			return false;
-		return AddRequestedMemory(handle, reader, sliceCount, SharedMemoryMapType_ReadOnly, 0, hint, 0);
+		return AddRequestedMemory(handle, reader, sliceCount, type, 0, hint, 0);
 	}
 
 	bool SharedMemoryView::AddRequestedMemory(SharedMemoryAllocatorHandle handle, BinaryReader& reader, SharedMemoryMapType mapType, u64 offset, const tchar* hint, u8 init)
diff --git a/Engine/Source/Programs/UnrealBuildAccelerator/Core/Public/UbaSharedMemoryView.h b/Engine/Source/Programs/UnrealBuildAccelerator/Core/Public/UbaSharedMemoryView.h
--- a/Engine/Source/Programs/UnrealBuildAccelerator/Core/Public/UbaSharedMemoryView.h
+++ b/Engine/Source/Programs/UnrealBuildAccelerator/Core/Public/UbaSharedMemoryView.h
@@ -33,7 +33,7 @@ namespace uba
 		SharedMemoryView(SharedMemoryView&&) = default;
 
 		bool Init(SharedMemoryAllocatorHandle handle, u64 capacity, void* baseAddress = nullptr, const tchar* hint = TC(""));
-		bool Init(SharedMemoryAllocatorHandle handle, u8* sharedMemory, BinaryReader& reader, u64 capacity, const tchar* hint = TC(""));
+		bool Init(SharedMemoryAllocatorHandle handle, u8* sharedMemory, BinaryReader& reader, u64 capacity, const tchar* hint = TC(""), SharedMemoryMapType type = SharedMemoryMapType_ReadOnly);
 
 		bool AddRequestedMemory(SharedMemoryAllocatorHandle handle, BinaryReader& reader, SharedMemoryMapType mapType = SharedMemoryMapType_ReadWrite, u64 offset = 0, const tchar* hint = TC(""), u8 init = 0);
 		void Reset();
diff --git a/Engine/Source/Programs/UnrealBuildAccelerator/Detours/Private/UbaDetoursFileMappingTable.h b/Engine/Source/Programs/UnrealBuildAccelerator/Detours/Private/UbaDetoursFileMappingTable.h
--- a/Engine/Source/Programs/UnrealBuildAccelerator/Detours/Private/UbaDetoursFileMappingTable.h
+++ b/Engine/Source/Programs/UnrealBuildAccelerator/Detours/Private/UbaDetoursFileMappingTable.h
@@ -104,7 +104,7 @@ namespace uba
 
 	SharedMemoryHandle Rpc_CreateSharedMemory(SharedMemoryView& out, SharedMemoryAllocatorHandle allocatorHandle, u64 capacity, u64 initialSize, const tchar* hint);
 	void Rpc_CommitSharedMemory(SharedMemoryView& view, SharedMemoryAllocatorHandle allocatorHandle, SharedMemoryHandle handle, u64 size);
-	void Rpc_GetSharedMemory(SharedMemoryView& view, SharedMemoryAllocatorHandle allocatorHandle, u8* allocatorMem, SharedMemoryHandle handle, u64 offset, u64 size, const tchar* hint);
+	void Rpc_GetSharedMemory(SharedMemoryView& view, SharedMemoryAllocatorHandle allocatorHandle, u8* allocatorMem, SharedMemoryHandle handle, u64 offset, u64 size, const tchar* hint, SharedMemoryMapType type = SharedMemoryMapType_ReadOnly);
 	void Rpc_UpdateSharedMemory(SharedMemoryView& view, SharedMemoryAllocatorHandle allocatorHandle, SharedMemoryHandle handle, u64 capacity);
 
 	struct DirHash
diff --git a/Engine/Source/Programs/UnrealBuildAccelerator/Detours/Private/Windows/UbaDetoursFunctionsKernelBase.inl b/Engine/Source/Programs/UnrealBuildAccelerator/Detours/Private/Windows/UbaDetoursFunctionsKernelBase.inl
--- a/Engine/Source/Programs/UnrealBuildAccelerator/Detours/Private/Windows/UbaDetoursFunctionsKernelBase.inl
+++ b/Engine/Source/Programs/UnrealBuildAccelerator/Detours/Private/Windows/UbaDetoursFunctionsKernelBase.inl
@@ -881,8 +881,9 @@ DWORD Detoured_GetFileSizeEx(HANDLE hFile, PLARGE_INTEGER lpFileSize)
 			UBA_ASSERTF(false, L"GetFileSizeEx (%s)", HandleToName(hFile));
 		}
 	}
-	DEBUG_LOG_TRUE(L"GetFileSizeEx", L"%llu (%ls)", uintptr_t(hFile), HandleToName(hFile));
-	return True_GetFileSizeEx(trueHandle, lpFileSize); // This ends up in Detoured_NtQueryInformationFile
+	DWORD res = True_GetFileSizeEx(trueHandle, lpFileSize); // This ends up in Detoured_NtQueryInformationFile
+	DEBUG_LOG_TRUE(L"GetFileSizeEx", L"%llu (%ls) -> %u (%llu)", uintptr_t(hFile), HandleToName(hFile), res, lpFileSize->QuadPart);
+	return res;
 }
 
 DWORD Detoured_SetFilePointer(HANDLE hFile, LONG lDistanceToMove, PLONG lpDistanceToMoveHigh, DWORD dwMoveMethod)
@@ -1448,8 +1449,10 @@ DWORD Detoured_GetModuleFileNameExA(HANDLE hProcess, HMODULE hModule, LPSTR lpFi
 
 BOOL Detoured_GetModuleHandleExW(DWORD dwFlags, LPCWSTR lpModuleName, HMODULE *phModule)
 {
+	bool isAddress = dwFlags & GET_MODULE_HANDLE_EX_FLAG_FROM_ADDRESS;
+
 	StringBuffer<> path;
-	if (lpModuleName && IsAbsolutePath(lpModuleName))
+	if (!isAddress && lpModuleName && IsAbsolutePath(lpModuleName))
 	{
 		FixPath(path, lpModuleName);
 		DevirtualizePath(path);
@@ -1461,9 +1464,10 @@ BOOL Detoured_GetModuleHandleExW(DWORD dwFlags, LPCWSTR lpModuleName, HMODULE *p
 	if (g_runningRemote)
 		t_ntOpenFileDisallowed = true;
 
-	DEBUG_LOG_TRUE(L"GetModuleHandleExW", L"%s", lpModuleName);
 	BOOL res = True_GetModuleHandleExW(dwFlags, lpModuleName, phModule);
 
+	DEBUG_LOG_TRUE(L"GetModuleHandleExW", L"%s -> %s (%llu)", isAddress ? TC("GET_MODULE_HANDLE_EX_FLAG_FROM_ADDRESS") : lpModuleName, ToString(res), u64(*phModule));
+
 	if (g_runningRemote)
 		t_ntOpenFileDisallowed = false;
 	return res;
@@ -2571,7 +2575,12 @@ LPVOID Detoured_MapViewOfFileEx(HANDLE hFileMappingObject, DWORD dwDesiredAccess
 			else if (!fi.fileMapMem)
 				fi.fileMapViewDesiredAccess = dwDesiredAccess;
 
-			if (!EnsureMapped(dh, dwFileOffsetHigh, dwFileOffsetLow, dwNumberOfBytesToMap, lpBaseAddress))
+			SharedMemoryMapType type = SharedMemoryMapType_ReadOnly;
+			if (dwDesiredAccess & FILE_MAP_COPY)
+				type = SharedMemoryMapType_CopyOnWrite;
+
+
+			if (!EnsureMapped(dh, dwFileOffsetHigh, dwFileOffsetLow, dwNumberOfBytesToMap, lpBaseAddress, type))
 				return nullptr;
 
 			SetLastError(ERROR_SUCCESS);
diff --git a/Engine/Source/Programs/UnrealBuildAccelerator/Detours/Private/Windows/UbaDetoursFunctionsWin.cpp b/Engine/Source/Programs/UnrealBuildAccelerator/Detours/Private/Windows/UbaDetoursFunctionsWin.cpp
--- a/Engine/Source/Programs/UnrealBuildAccelerator/Detours/Private/Windows/UbaDetoursFunctionsWin.cpp
+++ b/Engine/Source/Programs/UnrealBuildAccelerator/Detours/Private/Windows/UbaDetoursFunctionsWin.cpp
@@ -450,11 +450,13 @@ const wchar_t* HandleToName(HANDLE handle)
 u64 FileTypeMaxSize(StringView file, bool isSystemOrTempFile) { return g_rules->FileTypeMaxSize(file, isSystemOrTempFile); }
 u32 FileTypeCommitSize(StringView file) { return g_rules->FileTypeCommitSize(file); }
 
-bool EnsureMapped(DetouredHandle& handle, DWORD dwFileOffsetHigh = 0, DWORD dwFileOffsetLow = 0, SIZE_T numberOfBytesToMap = 0, void* baseAddress = nullptr)
+bool EnsureMapped(DetouredHandle& handle, DWORD dwFileOffsetHigh = 0, DWORD dwFileOffsetLow = 0, SIZE_T numberOfBytesToMap = 0, void* baseAddress = nullptr, SharedMemoryMapType type = SharedMemoryMapType_ReadOnly)
 {
 	FileInfo& info = *handle.fileObject->fileInfo;
 	
-	if (info.memoryFile || info.fileMapMem)
+	if (info.memoryFile)
+		return true;
+	if (info.fileMapMem && type == SharedMemoryMapType_ReadOnly)
 		return true;
 
 	u64 offset = ToLargeInteger(dwFileOffsetHigh, dwFileOffsetLow).QuadPart;
@@ -469,7 +471,7 @@ bool EnsureMapped(DetouredHandle& handle, DWORD dwFileOffsetHigh = 0, DWORD dwFi
 	u64 mapSize = (offset - alignedOffsetStart) + numberOfBytesToMap;
 
 	SharedMemoryView view;
-	Rpc_GetSharedMemory(view, g_sharedMemoryAllocator, g_sharedMemory, info.trueFileMapHandle, alignedOffsetStart, mapSize, info.originalName);
+	Rpc_GetSharedMemory(view, g_sharedMemoryAllocator, g_sharedMemory, info.trueFileMapHandle, alignedOffsetStart, mapSize, info.originalName, type);
 
 	info.fileMapMem = view.DetachMemory() + (offset - alignedOffsetStart);
 	info.fileMapMemSize = info.size;
