[UBA]
* Fixed so ubacli prints a warning if failing to load config file set with -config=<file>

diff --git a/Engine/Source/Programs/UnrealBuildAccelerator/Cli/Private/UbaCli.cpp b/Engine/Source/Programs/UnrealBuildAccelerator/Cli/Private/UbaCli.cpp
--- a/Engine/Source/Programs/UnrealBuildAccelerator/Cli/Private/UbaCli.cpp
+++ b/Engine/Source/Programs/UnrealBuildAccelerator/Cli/Private/UbaCli.cpp
@@ -655,7 +655,8 @@ namespace uba
 
 		Config config;
 		if (!configFile.empty())
-			config.LoadFromFile(logger, configFile.c_str());
+			if (!config.LoadFromFile(logger, configFile.c_str()))
+				logger.Warning(TC(" Failed to read config file %s"), configFile.c_str());
 
 		if constexpr (!IsArmBinary)
 			if (IsRunningArm())
