[UBA]
* Fixed up some code to have prevExitCode be transferred to trace file (not written yet because don't want to bump trace file version)

diff --git a/Engine/Source/Programs/UnrealBuildAccelerator/Common/Private/UbaScheduler.cpp b/Engine/Source/Programs/UnrealBuildAccelerator/Common/Private/UbaScheduler.cpp
--- a/Engine/Source/Programs/UnrealBuildAccelerator/Common/Private/UbaScheduler.cpp
+++ b/Engine/Source/Programs/UnrealBuildAccelerator/Common/Private/UbaScheduler.cpp
@@ -1225,7 +1225,7 @@ namespace uba
 				lock.Leave();
 				if (firstWait)
 				{
-					m_session.SetProcessAsIdle((ProcessImpl&)process, prevStats);
+					m_session.SetProcessAsIdle((ProcessImpl&)process, prevStats, prevExitCode);
 					//m_session.GetTrace().ProcessEnvironmentUpdated(process.GetId(), TCV("IDLE"), prevStats.GetPositionData(), prevStats.GetLeft(), TCV(""));
 					firstWait = false;
 				}
diff --git a/Engine/Source/Programs/UnrealBuildAccelerator/Common/Private/UbaSession.cpp b/Engine/Source/Programs/UnrealBuildAccelerator/Common/Private/UbaSession.cpp
--- a/Engine/Source/Programs/UnrealBuildAccelerator/Common/Private/UbaSession.cpp
+++ b/Engine/Source/Programs/UnrealBuildAccelerator/Common/Private/UbaSession.cpp
@@ -2748,9 +2748,9 @@ namespace uba
 		return logFile;
 	}
 
-	void Session::SetProcessAsIdle(ProcessImpl& process, BinaryReader& prevStats)
+	void Session::SetProcessAsIdle(ProcessImpl& process, BinaryReader& prevStats, u32 prevExitCode)
 	{
-		m_trace.ProcessEnvironmentUpdated(process.GetId(), TCV("IDLE"), prevStats.GetPositionData(), prevStats.GetLeft(), TCV(""));
+		m_trace.ProcessEnvironmentUpdated(process.GetId(), TCV("IDLE"), prevExitCode, prevStats.GetPositionData(), prevStats.GetLeft(), TCV(""));
 		process.m_processStats = {};
 		process.m_sessionStats = {};
 		process.m_storageStats = {};
@@ -3646,7 +3646,7 @@ namespace uba
 		if (!outNewProcess)
 			return true;
 
-		m_trace.ProcessEnvironmentUpdated(process.GetId(), outNextProcess.description, prevStats.GetPositionData(), prevStats.GetLeft(), outNextProcess.breadcrumbs);
+		m_trace.ProcessEnvironmentUpdated(process.GetId(), outNextProcess.description, prevExitCode, prevStats.GetPositionData(), prevStats.GetLeft(), outNextProcess.breadcrumbs);
 		return true;
 	}
 
@@ -3778,7 +3778,7 @@ namespace uba
 		process.m_processStats.Write(writer);
 		process.m_storageStats.Write(writer);
 		process.m_kernelStats.Write(writer);
-		m_trace.ProcessEnvironmentUpdated(process.GetId(), reason, writer.GetData(), writer.GetPosition(), ToView(process.GetStartInfo().breadcrumbs));
+		m_trace.ProcessEnvironmentUpdated(process.GetId(), reason, 0, writer.GetData(), writer.GetPosition(), ToView(process.GetStartInfo().breadcrumbs));
 		process.m_processStats = {};
 		process.m_storageStats = {};
 		process.m_kernelStats = {};
diff --git a/Engine/Source/Programs/UnrealBuildAccelerator/Common/Private/UbaSessionClient.cpp b/Engine/Source/Programs/UnrealBuildAccelerator/Common/Private/UbaSessionClient.cpp
--- a/Engine/Source/Programs/UnrealBuildAccelerator/Common/Private/UbaSessionClient.cpp
+++ b/Engine/Source/Programs/UnrealBuildAccelerator/Common/Private/UbaSessionClient.cpp
@@ -2737,7 +2737,7 @@ namespace uba
 		if (!timeoutMs)
 			return false;
 
-		SetProcessAsIdle(pi, statsReader);
+		SetProcessAsIdle(pi, statsReader, prevExitCode);
 
 		// No process found, let main loop query free processes
 
diff --git a/Engine/Source/Programs/UnrealBuildAccelerator/Common/Private/UbaSessionServer.cpp b/Engine/Source/Programs/UnrealBuildAccelerator/Common/Private/UbaSessionServer.cpp
--- a/Engine/Source/Programs/UnrealBuildAccelerator/Common/Private/UbaSessionServer.cpp
+++ b/Engine/Source/Programs/UnrealBuildAccelerator/Common/Private/UbaSessionServer.cpp
@@ -2073,7 +2073,7 @@ namespace uba
 			else if (!remoteProcess.m_isIdle)
 			{
 				// Stats are the previous call's stats
-				m_trace.ProcessEnvironmentUpdated(processId, TCV("IDLE"), prevStats.GetPositionData(), prevStats.GetLeft(), TCV(""));
+				m_trace.ProcessEnvironmentUpdated(processId, TCV("IDLE"), prevExitCode, prevStats.GetPositionData(), prevStats.GetLeft(), TCV(""));
 				remoteProcess.m_isIdle = true;
 			}
 		}
@@ -2108,7 +2108,7 @@ namespace uba
 		}
 		StringBuffer<> reason;
 		reader.ReadString(reason);
-		m_trace.ProcessEnvironmentUpdated(processId, reason, reader.GetPositionData(), reader.GetLeft(), ToView(h.GetStartInfo().breadcrumbs));
+		m_trace.ProcessEnvironmentUpdated(processId, reason, 0, reader.GetPositionData(), reader.GetLeft(), ToView(h.GetStartInfo().breadcrumbs));
 		return true;
 	}
 
diff --git a/Engine/Source/Programs/UnrealBuildAccelerator/Common/Private/UbaTrace.cpp b/Engine/Source/Programs/UnrealBuildAccelerator/Common/Private/UbaTrace.cpp
--- a/Engine/Source/Programs/UnrealBuildAccelerator/Common/Private/UbaTrace.cpp
+++ b/Engine/Source/Programs/UnrealBuildAccelerator/Common/Private/UbaTrace.cpp
@@ -435,12 +435,13 @@ namespace uba
 		writer.WriteLongString(breadcrumbs, 20);
 	}
 
-	void Trace::ProcessEnvironmentUpdated(u32 processId, StringView reason, const u8* data, u64 dataSize, StringView breadcrumbs)
+	void Trace::ProcessEnvironmentUpdated(u32 processId, StringView reason, u32 prevExitCode, const u8* prevStats, u64 prevStatsSize, StringView breadcrumbs)
 	{
 		BEGIN_TRACE_ENTRY(ProcessEnvironmentUpdated);
 		writer.WriteU32(processId);
 		writer.WriteString(reason);
-		writer.WriteBytes(data, dataSize);
+		//writer.WriteU32(prevExitCode); // TODO: Add this on next trace version update
+		writer.WriteBytes(prevStats, prevStatsSize);
 		writer.WriteLongString(breadcrumbs, 20);
 	}
 
diff --git a/Engine/Source/Programs/UnrealBuildAccelerator/Common/Public/UbaSession.h b/Engine/Source/Programs/UnrealBuildAccelerator/Common/Public/UbaSession.h
--- a/Engine/Source/Programs/UnrealBuildAccelerator/Common/Public/UbaSession.h
+++ b/Engine/Source/Programs/UnrealBuildAccelerator/Common/Public/UbaSession.h
@@ -122,7 +122,7 @@ namespace uba
 		void FlushDeadProcesses();
 
 		const tchar* FixLogFile(const tchar* logFile, StringBufferBase& temp, const tchar* processArguments, u32 processId);
-		void SetProcessAsIdle(ProcessImpl& process, BinaryReader& prevStats);
+		void SetProcessAsIdle(ProcessImpl& process, BinaryReader& prevStats, u32 prevExitCode);
 
 	protected:
 		Session(const SessionCreateInfo& info, const tchar* logPrefix, bool runningRemote, WorkManager& workManager);
diff --git a/Engine/Source/Programs/UnrealBuildAccelerator/Common/Public/UbaTrace.h b/Engine/Source/Programs/UnrealBuildAccelerator/Common/Public/UbaTrace.h
--- a/Engine/Source/Programs/UnrealBuildAccelerator/Common/Public/UbaTrace.h
+++ b/Engine/Source/Programs/UnrealBuildAccelerator/Common/Public/UbaTrace.h
@@ -105,7 +105,7 @@ namespace uba
 		void SessionSummary(u32 sessionId, const u8* data, u64 dataSize);
 		void SessionDisconnect(u32 sessionId);
 		void ProcessAdded(u32 sessionId, u32 processId, StringView description, StringView breadcrumbs);
-		void ProcessEnvironmentUpdated(u32 processId, StringView reason, const u8* data, u64 dataSize, StringView breadcrumbs);
+		void ProcessEnvironmentUpdated(u32 processId, StringView reason, u32 prevExitCode, const u8* prevStats, u64 prevStatsSize, StringView breadcrumbs);
 		void ProcessExited(u32 processId, u32 exitCode, const u8* data, u64 dataSize, const Vector<ProcessLogLine>& logLines);
 		void ProcessExited(u32 processId, u32 exitCode);
 		void ProcessReturned(u32 processId, StringView reason);
