Fix /ALTERNATENAME interpreted by UBA as a file name for implib RSPs

#rnx

#Synced-CL 49762101

diff --git a/Engine/Source/Programs/UnrealBuildAccelerator/ObjTool/Private/UbaObjTool.cpp b/Engine/Source/Programs/UnrealBuildAccelerator/ObjTool/Private/UbaObjTool.cpp
--- a/Engine/Source/Programs/UnrealBuildAccelerator/ObjTool/Private/UbaObjTool.cpp
+++ b/Engine/Source/Programs/UnrealBuildAccelerator/ObjTool/Private/UbaObjTool.cpp
@@ -143,6 +143,9 @@ namespace uba
 					else if (name.Equals(TCV("/NODEFAULTLIB")))
 					{
 					}
+					else if (name.StartsWith(TCV("/ALTERNATENAME")))
+					{
+					}
 					else if (name.Equals(TCV("/LTCG")))
 					{
 					}
