[UBA]
* Missing file for changelist that fixed radlink support

diff --git a/Engine/Source/Programs/UnrealBuildAccelerator/Detours/Private/UbaDetoursFileMappingTable.cpp b/Engine/Source/Programs/UnrealBuildAccelerator/Detours/Private/UbaDetoursFileMappingTable.cpp
--- a/Engine/Source/Programs/UnrealBuildAccelerator/Detours/Private/UbaDetoursFileMappingTable.cpp
+++ b/Engine/Source/Programs/UnrealBuildAccelerator/Detours/Private/UbaDetoursFileMappingTable.cpp
@@ -209,7 +209,8 @@ namespace uba
 					DEBUG_LOG(TC("Unmapping old memoryfile for: %s"), info.originalName);
 
 #if UBA_DEBUG
-					UBA_ASSERTF(false, TC("REPORT THIS TO UBA DEVS"));
+					// Happens when linking server exe in merged modules on win64
+					//UBA_ASSERTF(false, TC("REPORT THIS TO UBA DEVS"));
 					//UnmapViewOfFile(mf->baseAddress);
 					//CloseHandle(mf->memoryHandle.mh);
 					//CloseHandle(mf->memoryHandle.fh);
@@ -295,7 +296,7 @@ namespace uba
 		view.AddRequestedMemory(allocatorHandle, reader);
 	}
 
-	void Rpc_GetSharedMemory(SharedMemoryView& view, SharedMemoryAllocatorHandle allocatorHandle, u8* allocatorMem, SharedMemoryHandle handle, u64 offset, u64 size, const tchar* hint)
+	void Rpc_GetSharedMemory(SharedMemoryView& view, SharedMemoryAllocatorHandle allocatorHandle, u8* allocatorMem, SharedMemoryHandle handle, u64 offset, u64 size, const tchar* hint, SharedMemoryMapType type)
 	{
 		UBA_ASSERTF(AlignUp(offset, PageSize) == offset, TC("Offset %llu unaligned. Can't get unaligned shared memory (%s)"), offset, hint);
 		//UBA_ASSERTF(AlignUp(size, PageSize) == size, TC("Size %llu unaligned. Can't get unaligned shared memory (%s)"), size, hint);
@@ -309,7 +310,7 @@ namespace uba
 		u64 committed = size;
 		if (!committed)
 			committed = reader.Read7BitEncoded();
-		view.Init(allocatorHandle, allocatorMem, reader, committed, hint);
+		view.Init(allocatorHandle, allocatorMem, reader, committed, hint, type);
 		//UBA_ASSERTF(u64(view.GetMemory()) != 0xFFFFFFFFFFFFFFFFllu, TC("%llu %llu %llu"), offset, size, committed);
 		//view.AddRequestedMemory(allocatorHandle, reader, SharedMemoryMapType_ReadOnly);
 	}
