[UBA]
* Fixed active process graph visualizing wrong when pan

diff --git a/Engine/Source/Programs/UnrealBuildAccelerator/Visualizer/Private/UbaVisualizer.cpp b/Engine/Source/Programs/UnrealBuildAccelerator/Visualizer/Private/UbaVisualizer.cpp
--- a/Engine/Source/Programs/UnrealBuildAccelerator/Visualizer/Private/UbaVisualizer.cpp
+++ b/Engine/Source/Programs/UnrealBuildAccelerator/Visualizer/Private/UbaVisualizer.cpp
@@ -947,14 +947,13 @@ namespace uba
 				int prevY = 0;
 
 				float timeScale = (m_config.horizontalScaleValue*m_zoomValue)*50.0f;
-				float startOffset = ((m_scrollPosX/timeScale) - float(int(m_scrollPosX/timeScale))) * timeScale;
-				float time = -startOffset/timeScale;
+				float time = 0;
 				float playTimeS = TimeToS(playTime);
 				u32 activeProcessCountIndex = 0;
 				u16 count = 0;
 				while (time < playTimeS && activeProcessCountIndex < m_traceView.activeProcessCounts.size())
 				{
-					int posX = progressRect.left + int(startOffset + time*timeScale);
+					int posX = progressRect.left + int(m_scrollPosX + time*timeScale);
 					if (posX >= clientRect.right)
 						break;
 
