[UBA]
* Fixed bug where process bitmaps were not invalidated when zooming out during certain conditions

diff --git a/Engine/Source/Programs/UnrealBuildAccelerator/Visualizer/Private/Windows/UbaVisualizerWin.cpp b/Engine/Source/Programs/UnrealBuildAccelerator/Visualizer/Private/Windows/UbaVisualizerWin.cpp
--- a/Engine/Source/Programs/UnrealBuildAccelerator/Visualizer/Private/Windows/UbaVisualizerWin.cpp
+++ b/Engine/Source/Programs/UnrealBuildAccelerator/Visualizer/Private/Windows/UbaVisualizerWin.cpp
@@ -802,7 +802,7 @@ namespace uba
 					//if (!m_traceView.finished && m_scrollPosX <= minScroll)
 					//	m_autoScroll = true;
 
-					if (oldBoxHeight < newBoxHeight)
+					if (oldBoxHeight != newBoxHeight)
 						DirtyBitmaps(true);
 					else if (m_config.ShowReadWriteColors)
 						for (auto& session : m_traceView.sessions)
