[UBA]
* Fixed -mergemodules -vfs for native linux... UbaObjTool had bugs in command line parsing causing it to fail
* Added more commandline validation

diff --git a/Engine/Source/Programs/UnrealBuildAccelerator/Common/Private/UbaObjectFile.cpp b/Engine/Source/Programs/UnrealBuildAccelerator/Common/Private/UbaObjectFile.cpp
--- a/Engine/Source/Programs/UnrealBuildAccelerator/Common/Private/UbaObjectFile.cpp
+++ b/Engine/Source/Programs/UnrealBuildAccelerator/Common/Private/UbaObjectFile.cpp
@@ -17,6 +17,7 @@ namespace uba
 
 	ObjectFile* ObjectFile::OpenAndParse(Logger& logger, ObjectFileParseMode parseMode, const tchar* filename)
 	{
+		UBA_ASSERT(*filename);
 		auto file = new FileAccessor(logger, filename);
 		auto fileGuard = MakeGuard([&]() { delete file; });
 
@@ -280,6 +281,7 @@ namespace uba
 
 	bool SymbolFile::ParseFile(Logger& logger, const tchar* filename)
 	{
+		UBA_ASSERT(*filename);
 		FileAccessor symFile(logger, filename);
 		if (!symFile.OpenMemoryRead())
 			return false;
diff --git a/Engine/Source/Programs/UnrealBuildAccelerator/ObjTool/Private/UbaObjTool.cpp b/Engine/Source/Programs/UnrealBuildAccelerator/ObjTool/Private/UbaObjTool.cpp
--- a/Engine/Source/Programs/UnrealBuildAccelerator/ObjTool/Private/UbaObjTool.cpp
+++ b/Engine/Source/Programs/UnrealBuildAccelerator/ObjTool/Private/UbaObjTool.cpp
@@ -87,7 +87,7 @@ namespace uba
 				else
 				{
 					const tchar* colon = TStrchr(arg,':');
-					if (colon && colon[1] != '\\' && colon[1] != '/')
+					if (colon && (!IsWindows || (colon[1] != '\\' && colon[1] != '/')))
 					{
 						name.Append(arg, colon - arg);
 						const tchar* valueStart = colon+1;
@@ -104,6 +104,14 @@ namespace uba
 					}
 				}
 
+				auto checkHasValue = [&]()
+					{
+						if (value.count)
+							return true;
+						LoggerWithWriter(g_consoleLogWriter, TC("")).Error(TC("   ERROR: %s needs a filename"), name.data);
+						return false;
+					};
+
 				if (isImpLibRsp)
 				{
 					if (name.Equals(TCV("/NOLOGO")))
@@ -135,6 +143,8 @@ namespace uba
 					}
 					else if (name.Equals(TCV("/OUT")))
 					{
+						if (!checkHasValue())
+							return -1;
 						impLibFile = value.data;
 					}
 					else if (name.Equals(TCV("/IGNORE")))
@@ -159,22 +169,32 @@ namespace uba
 				}
 				else if (name.StartsWith(TC("/D")))
 				{
+					if (!checkHasValue())
+						return -1;
 					objFilesDependencies.push_back(value.data);
 				}
 				else if (name.StartsWith(TC("/S")))
 				{
+					if (!checkHasValue())
+						return -1;
 					objFilesToExport.push_back(value.data);
 				}
 				else if (name.StartsWith(TC("/O")))
 				{
+					if (!checkHasValue())
+						return -1;
 					extraObjFile = value.data;
 				}
 				else if (name.StartsWith(TC("/T")))
 				{
+					if (!checkHasValue())
+						return -1;
 					platform = value.data;
 				}
 				else if (name.StartsWith(TC("/M")))
 				{
+					if (!checkHasValue())
+						return -1;
 					moduleName = value.data;
 				}
 				else if (name.StartsWith(TC("/E")))
@@ -190,6 +210,8 @@ namespace uba
 				}
 				else if (name.Equals(TCV("-writeimplib")))
 				{
+					if (!checkHasValue())
+						return -1;
 					impLibFile = value.data;
 					writeImpLib = true;
 					allowLibInputs = true;
