[UBA]
* Added code to be able to use native linux as helper to other platforms. Works similar to how arm64 machines can help x64 machines

diff --git a/Engine/Source/Programs/UnrealBuildAccelerator/Cli/Private/UbaCli.cpp b/Engine/Source/Programs/UnrealBuildAccelerator/Cli/Private/UbaCli.cpp
--- a/Engine/Source/Programs/UnrealBuildAccelerator/Cli/Private/UbaCli.cpp
+++ b/Engine/Source/Programs/UnrealBuildAccelerator/Cli/Private/UbaCli.cpp
@@ -1102,7 +1102,7 @@ namespace uba
 				StringBuffer<> error;
 				bool printImports = true;
 				BinaryInfo info;
-				if (!ParseBinary(application, StringView(application).GetPath(), info, [&](const tchar* import, bool isKnown, const char* const* loaderPaths)
+				if (!ParseBinary(application, StringView(application).GetPath(), info, [&](const tchar* import, bool isKnown, const tchar* const* loaderPaths)
 				{
 					if (printImports)
 					{
diff --git a/Engine/Source/Programs/UnrealBuildAccelerator/Common/Private/UbaObjectFileElf.cpp b/Engine/Source/Programs/UnrealBuildAccelerator/Common/Private/UbaObjectFileElf.cpp
--- a/Engine/Source/Programs/UnrealBuildAccelerator/Common/Private/UbaObjectFileElf.cpp
+++ b/Engine/Source/Programs/UnrealBuildAccelerator/Common/Private/UbaObjectFileElf.cpp
@@ -1,42 +1,11 @@
 // Copyright Epic Games, Inc. All Rights Reserved.
 
 #include "UbaObjectFileElf.h"
+#include "UbaElf.h"
 #include "UbaFileAccessor.h"
 
 namespace uba
 {
-	struct Elf64Header
-	{
-		unsigned char	e_ident[16];	// Magic number and other info
-		u16	e_type;			// Object file type
-		u16	e_machine;		// Architecture
-		u32	e_version;		// Object file version
-		u64	e_entry;		// Entry point virtual address
-		u64	e_phoff;		// Program header table file offset
-		u64	e_shoff;		// Section header table file offset
-		u32	e_flags;		// Processor-specific flags
-		u16	e_ehsize;		// ELF header size in bytes
-		u16	e_phentsize;	// Program header table entry size
-		u16	e_phnum;		// Program header table entry count
-		u16	e_shentsize;	// Section header table entry size
-		u16	e_shnum;		// Section header table entry count
-		u16	e_shstrndx;		// Section header string table index
-	};
-
-	struct Elf64SectionHeader
-	{
-		u32	sh_name;		// Section name (string tbl index)
-		u32	sh_type;		// Section type
-		u64	sh_flags;		// Section flags
-		u64	sh_addr;		// Section virtual addr at execution
-		u64	sh_offset;		// Section file offset
-		u64	sh_size;		// Section size in bytes
-		u32	sh_link;		// Link to another section
-		u32	sh_info;		// Additional section information
-		u64	sh_addralign;	// Section alignment
-		u64	sh_entsize;		// Entry size if section holds table
-	};
-
 	struct Elf64Sym
 	{
 		u32	st_name;		// Symbol name (string tbl index)
@@ -128,7 +97,7 @@ namespace uba
 
 	bool ObjectFileElf::Parse(Logger& logger, ObjectFileParseMode parseMode, const tchar* hint)
 	{
-		auto& header = *(Elf64Header*)m_data;
+		auto& header = *(Elf64FileHeader*)m_data;
 		if (header.e_ident[4] != 2) // 64-bit  (1 is 32-bit)
 			return false;
 		if (header.e_ident[5] != 1) // Little endian
@@ -247,8 +216,8 @@ namespace uba
 
 	bool ObjectFileElf::CreateExtraFile(Logger& logger, const StringView& platform, MemoryBlock& memoryBlock, const AllExternalImports& allExternalImports, const AllInternalImports& allInternalImports, const AllExports& allExports, bool includeExportsInFile)
 	{
-		u8* data = (u8*)memoryBlock.Allocate(sizeof(Elf64Header) + sizeof(Elf64SectionHeader) + 1, 1, TC(""));
-		auto& header = *(Elf64Header*)data;
+		u8* data = (u8*)memoryBlock.Allocate(sizeof(Elf64FileHeader) + sizeof(Elf64SectionHeader) + 1, 1, TC(""));
+		auto& header = *(Elf64FileHeader*)data;
 
 		header.e_ident[0] = 0x7f;
 		header.e_ident[1] = 'E';
@@ -261,19 +230,19 @@ namespace uba
 		header.e_type = 1;
 		header.e_machine = EM_X86_64;
 		header.e_version = 1;
-		header.e_ehsize = sizeof(Elf64Header);
+		header.e_ehsize = sizeof(Elf64FileHeader);
 		header.e_flags = 0;//0x04000000;
-		header.e_shoff = sizeof(Elf64Header);
+		header.e_shoff = sizeof(Elf64FileHeader);
 		header.e_shentsize = sizeof(Elf64SectionHeader);
 		header.e_shnum = 1;
 		header.e_shstrndx = 0;
 
-		auto& section = *(Elf64SectionHeader*)(data + sizeof(Elf64Header));
+		auto& section = *(Elf64SectionHeader*)(data + sizeof(Elf64FileHeader));
 		section.sh_name = 0;
 		section.sh_type = 3;
 		section.sh_flags = 0;//0xC35d00;
 		section.sh_addralign = 1;
-		section.sh_offset = sizeof(Elf64Header) + sizeof(Elf64SectionHeader);
+		section.sh_offset = sizeof(Elf64FileHeader) + sizeof(Elf64SectionHeader);
 		section.sh_size = 1;
 		return true;
 	}
diff --git a/Engine/Source/Programs/UnrealBuildAccelerator/Common/Private/UbaProcess.cpp b/Engine/Source/Programs/UnrealBuildAccelerator/Common/Private/UbaProcess.cpp
--- a/Engine/Source/Programs/UnrealBuildAccelerator/Common/Private/UbaProcess.cpp
+++ b/Engine/Source/Programs/UnrealBuildAccelerator/Common/Private/UbaProcess.cpp
@@ -1317,8 +1317,8 @@ namespace uba
 
 		auto& process = *(ProcessImpl*)h.m_process;
 
-		const char* detoursLib = m_session.m_detoursLibrary[IsArmBinary].c_str();
-		u32 detoursLibLen = u32(m_session.m_detoursLibrary[IsArmBinary].size());
+		const char* detoursLib = m_session.m_detoursLibrary[IsLinux][IsArmBinary].c_str();
+		u32 detoursLibLen = u32(m_session.m_detoursLibrary[IsLinux][IsArmBinary].size());
 
 		writer.WriteU32(childProcessId);
 		writer.WriteU32(info.rules->index);
@@ -1620,7 +1620,7 @@ namespace uba
 
 	bool ProcessImpl::HandleResolveCallstack(BinaryReader& reader, BinaryWriter& writer)
 	{
-		m_session.GetSymbols(m_startInfo.application, m_isArmBinary, reader, writer);
+		m_session.GetSymbols(m_startInfo.application, IsLinux, m_isArmBinary, reader, writer);
 		return true;
 	}
 
@@ -1788,7 +1788,7 @@ namespace uba
 					allowCustomAllocator = false;
 			}
 
-			const char* detoursLib = m_session.m_detoursLibrary[m_isArmBinary].c_str();
+			const char* detoursLib = m_session.m_detoursLibrary[IsLinux][m_isArmBinary].c_str();
 			if (!*detoursLib)
 				detoursLib = UBA_DETOURS_LIBRARY_ANSI;
 
@@ -2200,7 +2200,7 @@ namespace uba
 
 			if (m_detourEnabled)
 			{
-				const char* detoursLib = m_session.m_detoursLibrary[IsArmBinary].c_str();
+				const char* detoursLib = m_session.m_detoursLibrary[IsLinux][IsArmBinary].c_str();
 				if (*detoursLib)
 				{
 #if PLATFORM_LINUX
diff --git a/Engine/Source/Programs/UnrealBuildAccelerator/Common/Private/UbaSession.cpp b/Engine/Source/Programs/UnrealBuildAccelerator/Common/Private/UbaSession.cpp
--- a/Engine/Source/Programs/UnrealBuildAccelerator/Common/Private/UbaSession.cpp
+++ b/Engine/Source/Programs/UnrealBuildAccelerator/Common/Private/UbaSession.cpp
@@ -20,6 +20,10 @@
 extern char **environ;
 #endif
 
+#if !PLATFORM_LINUX
+#include "UbaLinuxBinDependencyParser.h"
+#endif
+
 #if UBA_USE_NATIVE_MAC_SEMAPHORES
 #include <mach/mach.h>
 #include <mach/semaphore.h>
@@ -1661,7 +1665,7 @@ namespace uba
 		return WithVfs(rootsHandle, !entry.locals.empty());
 	}
 
-	bool Session::CopyImports(Vector<BinaryModule>& out, const tchar* library, tchar* applicationDir, tchar* applicationDirEnd, UnorderedSet<TString>& handledImports, const char* const* loaderPaths)
+	bool Session::CopyImports(Vector<BinaryModule>& out, const tchar* library, tchar* applicationDir, tchar* applicationDirEnd, UnorderedSet<TString>& handledImports, const tchar* const* loaderPaths, bool isLinuxBinary)
 	{
 		if (!handledImports.insert(library).second)
 			return true;
@@ -1729,15 +1733,30 @@ namespace uba
 		if (isInSystemFolder && IsKnownSystemFile(applicationName))
 			return true;
 
+		#if PLATFORM_WINDOWS
+		if (isLinuxBinary)
+			temp3.Replace(PathSeparator, NonPathSeparator);
+		#endif
+
 		BinaryModule& binaryModule = out.emplace_back(library, temp3.ToString(), attr, isInSystemFolder);(void)binaryModule;
 
 		StringBuffer<> errorStr;
 		BinaryInfo info;
-		if (!ParseBinary(temp3, ToView(applicationDir), info, [&](const tchar* importName, bool isKnown, const char* const* importLoaderPaths)
+		auto handleImport = [&](const tchar* importName, bool isKnown, const tchar* const* importLoaderPaths)
 			{
 				if (result && !isKnown)
-					result = CopyImports(out, importName, applicationDir, applicationDirEnd, handledImports, importLoaderPaths);
-			}, errorStr))
+					result = CopyImports(out, importName, applicationDir, applicationDirEnd, handledImports, importLoaderPaths, isLinuxBinary);
+			};
+
+		bool parseRes;
+		#if !PLATFORM_LINUX
+		if (isLinuxBinary)
+			parseRes = linux::ParseBinary(temp3, ToView(applicationDir), info, handleImport, errorStr);
+		else
+		#endif
+			parseRes = ParseBinary(temp3, ToView(applicationDir), info, handleImport, errorStr);
+
+		if (!parseRes)
 			return m_logger.Error(TC("Failed to parse binary %s for imports"), temp3.data);
 
 		if (errorStr.count)
@@ -2876,7 +2895,7 @@ namespace uba
 		logger.EndScope();
 	}
 
-	bool Session::GetBinaryModules(Vector<BinaryModule>& out, const tchar* application)
+	bool Session::GetBinaryModules(Vector<BinaryModule>& out, const tchar* application, bool isLinuxBinary)
 	{
 		const tchar* applicationName = application;
 
@@ -2890,7 +2909,7 @@ namespace uba
 		tchar* applicationDirEnd = applicationDir + applicationDirLen;
 
 		UnorderedSet<TString> handledImports;
-		return CopyImports(out, applicationName, applicationDir, applicationDirEnd, handledImports, nullptr);
+		return CopyImports(out, applicationName, applicationDir, applicationDirEnd, handledImports, nullptr, isLinuxBinary);
 	}
 
 	void Session::Free(Vector<BinaryModule>& v)
@@ -3712,12 +3731,12 @@ namespace uba
 		return true;
 	}
 
-	bool Session::GetSymbols(const tchar* application, bool isArm, BinaryReader& reader, BinaryWriter& writer)
+	bool Session::GetSymbols(const tchar* application, bool isLinux, bool isArm, BinaryReader& reader, BinaryWriter& writer)
 	{
 		StringBuffer<256> detoursLibPath;
 		StringBuffer<256> alternativeLibPath;
 
-		detoursLibPath.Append(m_detoursLibrary[IsArmBinary].c_str()).Resize(detoursLibPath.Last(PathSeparator) - detoursLibPath.data);
+		detoursLibPath.Append(m_detoursLibrary[IsLinux][IsArmBinary].c_str()).Resize(detoursLibPath.Last(PathSeparator) - detoursLibPath.data);
 		GetAlternativeUbaPath(m_logger, alternativeLibPath, detoursLibPath, IsArmBinary);
 
 		StringView searchPaths[3] = { detoursLibPath, alternativeLibPath, {} };
diff --git a/Engine/Source/Programs/UnrealBuildAccelerator/Common/Private/UbaSessionClient.cpp b/Engine/Source/Programs/UnrealBuildAccelerator/Common/Private/UbaSessionClient.cpp
--- a/Engine/Source/Programs/UnrealBuildAccelerator/Common/Private/UbaSessionClient.cpp
+++ b/Engine/Source/Programs/UnrealBuildAccelerator/Common/Private/UbaSessionClient.cpp
@@ -469,7 +469,7 @@ namespace uba
 		{
 			StackBinaryWriter<256> writer;
 			NetworkMessage msg(m_client, ServiceId, SessionMessageType_GetApplication, writer);
-			writer.WriteU32(processId);
+			writer.WriteBool(IsLinux);
 			writer.WriteString(application);
 			if (!msg.Send(reader, m_stats.getApplicationMsg))
 				return false;
@@ -1340,6 +1340,7 @@ namespace uba
 		writer.WriteString(m_name.data);
 		writer.WriteU32(SessionNetworkVersion);
 		writer.WriteBool(IsRunningArm());
+		writer.WriteBool(IsLinux);
 
 		CasKey keys[2];
 		if (m_useBinariesAsVersion)
@@ -1438,9 +1439,9 @@ namespace uba
 			#if PLATFORM_WINDOWS
 			char dll[1024];
 			detoursFile[i].Parse(dll, sizeof_array(dll));
-			m_detoursLibrary[i] = dll;
+			m_detoursLibrary[IsLinux][i] = dll;
 			#else
-			m_detoursLibrary[i] = detoursFile[i].data;
+			m_detoursLibrary[IsLinux][i] = detoursFile[i].data;
 			#endif
 		}
 
@@ -1786,7 +1787,7 @@ namespace uba
 				{
 					BinaryReader stackReader(cs.data.data(), 0, cs.data.size());
 					StackBinaryWriter<SendMaxSize> stackWriter;
-					GetSymbols(UBA_AGENT_EXECUTABLE, IsArmBinary, stackReader, stackWriter);
+					GetSymbols(UBA_AGENT_EXECUTABLE, IsLinux, IsArmBinary, stackReader, stackWriter);
 					BinaryReader resultReader(stackWriter.GetData(), 0, stackWriter.GetPosition());
 					TString infoString = resultReader.ReadString();
 					m_logger.Info(TC("%s%s"), cs.desc.c_str(), infoString.c_str());
@@ -2850,11 +2851,12 @@ namespace uba
 		return true;
 	}
 
-	bool SessionClient::GetSymbols(const tchar* application, bool isArm, BinaryReader& reader, BinaryWriter& writer)
+	bool SessionClient::GetSymbols(const tchar* application, bool isLinux, bool isArm, BinaryReader& reader, BinaryWriter& writer)
 	{
 		StackBinaryWriter<SendMaxSize> msgWriter;
 		NetworkMessage msg(m_client, ServiceId, SessionMessageType_GetSymbols, msgWriter);
 		msgWriter.WriteString(application);
+		msgWriter.WriteBool(isLinux);
 		msgWriter.WriteBool(isArm);
 		u32 size = reader.ReadU32();
 		msgWriter.WriteU32(size);
diff --git a/Engine/Source/Programs/UnrealBuildAccelerator/Common/Private/UbaSessionServer.cpp b/Engine/Source/Programs/UnrealBuildAccelerator/Common/Private/UbaSessionServer.cpp
--- a/Engine/Source/Programs/UnrealBuildAccelerator/Common/Private/UbaSessionServer.cpp
+++ b/Engine/Source/Programs/UnrealBuildAccelerator/Common/Private/UbaSessionServer.cpp
@@ -10,6 +10,9 @@
 #include "UbaScheduler.h"
 #include "UbaStorage.h"
 
+// Temp hack to test running linux application from windows. (I have manually copied gcc to d:\\temp)
+#define UBA_TEST_LINUX_HELPING_WINDOWS 0
+
 namespace uba
 {
 	struct SessionServer::ReceivedFile
@@ -168,6 +171,31 @@ namespace uba
 		return true;
 	}
 
+#if !PLATFORM_LINUX
+	bool GetLinuxDir(Logger& logger, StringBufferBase& dir, bool reportError)
+	{
+		// UBT has the path win-x64/native or win-arm64/native
+		bool isUbtPath = dir.EndsWith(TCV("native"));
+		if (isUbtPath)
+		{
+			const tchar* runtimesDir = IsWindows ? TC("\\runtimes\\") : TC("/runtimes/");
+			const tchar* pos = nullptr;
+			if (!dir.Contains(runtimesDir, true, &pos))
+				return false;
+			dir.Resize(pos - dir.data + 10).Append(TCV("linux-x64")).Append(PathSeparator).Append(TCV("native"));
+		}
+		else
+		{
+			const tchar* platformStr = IsWindows ? TC("\\Win64\\") : TC("/Mac/");
+			const tchar* pos = nullptr;
+			if (!dir.Contains(platformStr, true, &pos))
+				return false;
+			dir.Resize(pos - dir.data + 1).Append(TCV("Linux")).Append(PathSeparator).Append(TCV("UnrealBuildAccelerator"));
+		}
+		return true;
+	}
+#endif
+
 	SessionServer::SessionServer(const SessionServerCreateInfo& info, const u8* environment, u32 environmentSize)
 	:	Session(info, TC("UbaSessionServer"), false, info.server)
 	,	m_server(info.server)
@@ -231,7 +259,7 @@ namespace uba
 		m_localEnvironmentVariables.insert(TC("TMPDIR"));
 		#endif
 
-		StringBuffer<> detoursFile;
+		StringBuffer<512> detoursFile;
 		if (!GetDirectoryOfCurrentModule(m_logger, detoursFile))
 		{
 			UBA_ASSERT(false);
@@ -241,17 +269,28 @@ namespace uba
 
 		detoursFile.Append(PathSeparator).Append(UBA_DETOURS_LIBRARY);
 		
-		#if PLATFORM_WINDOWS
 		char temp[1024];
+
+		#if !PLATFORM_LINUX
+		StringBuffer<512> detoursFileLinux;
+		detoursFileLinux.AppendDir(detoursFile);
+		if (GetLinuxDir(m_logger, detoursFileLinux, false))
+		{
+			detoursFileLinux.Append(PathSeparator).Append(UBA_DETOURS_LIBRARY_LINUX).Parse(temp, sizeof_array(temp));
+			m_detoursLibrary[true][false] = temp;
+		}
+		#endif
+
+		#if PLATFORM_WINDOWS
 		detoursFile.Parse(temp, sizeof_array(temp));
-		m_detoursLibrary[IsArmBinary] = temp;
+		m_detoursLibrary[IsLinux][IsArmBinary] = temp;
 		if (GetCrossArchitectureDir(m_logger, detoursFile.Resize(dirLength), false))
 		{
 			detoursFile.Append(PathSeparator).Append(UBA_DETOURS_LIBRARY).Parse(temp, sizeof_array(temp));
-			m_detoursLibrary[!IsArmBinary] = temp;
+			m_detoursLibrary[IsLinux][!IsArmBinary] = temp;
 		}
 		#else
-		m_detoursLibrary[IsArmBinary] = detoursFile.data;
+		m_detoursLibrary[IsLinux][IsArmBinary] = detoursFile.data;
 		#endif
 
 		m_trace.SetThreadUpdateCallback([this]() { TraceSessionUpdate(); }, info.traceIntervalMs);
@@ -261,6 +300,10 @@ namespace uba
 			UBA_SHIPPING_ASSERT(false);
 			return;
 		}
+
+		#if UBA_TEST_LINUX_HELPING_WINDOWS
+		RegisterCrossArchitectureMapping(TC("d:\\temp\\dummy"), TC("d:\\temp\\gcc"));
+		#endif
 	}
 
 	SessionServer::~SessionServer()
@@ -421,6 +464,12 @@ namespace uba
 
 		auto rules = GetRules(remoteProcess->m_startInfo);
 
+		#if UBA_TEST_LINUX_HELPING_WINDOWS
+		remoteProcess->m_startInfo.applicationStr = TC("d:\\temp\\dummy");
+		remoteProcess->m_startInfo.application = remoteProcess->m_startInfo.applicationStr.c_str();
+		allowCrossArchitecture = true;
+		#endif
+
 		remoteProcess->m_startInfo.rules = rules;
 		remoteProcess->m_allowCrossArchitecture = allowCrossArchitecture;
 
@@ -903,8 +952,11 @@ namespace uba
 		reader.ReadString(name);
 		u32 clientVersion = reader.ReadU32();
 		bool isClientArm = false;
+		bool isClientLinux = false;
 		if (clientVersion >= 36)
 			isClientArm = reader.ReadBool();
+		if (clientVersion >= 49)
+			isClientLinux = reader.ReadBool();
 
 		m_logger.Detail(TC("Client session %s connected (Id: %u, Uid: %s%s)"), name.data, connectionInfo.GetId(), GuidToString(connectionInfo.GetUid()).str, (isClientArm ? TC(", IsArm: true") : TC("")));
 
@@ -918,8 +970,8 @@ namespace uba
 		{
 			SCOPED_FUTEX(m_binKeysLock, lock);
 
-			detoursBinaryKey[0] = m_detoursBinaryKey[0];
-			detoursBinaryKey[1] = m_detoursBinaryKey[1];
+			detoursBinaryKey[0] = m_detoursBinaryKey[isClientLinux][0];
+			detoursBinaryKey[1] = m_detoursBinaryKey[isClientLinux][1];
 
 			StringBuffer<> detoursLib;
 			bool deferCreation = true;
@@ -929,14 +981,15 @@ namespace uba
 			{
 				if (detoursBinaryKey[i] != CasKeyZero)
 					continue;
-				detoursLib.Clear().Append(m_detoursLibrary[i].c_str());
+				detoursLib.Clear().Append(m_detoursLibrary[isClientLinux][i].c_str());
 				if (!m_storage.StoreCasFile(detoursBinaryKey[i], detoursLib.data, CasKeyZero, deferCreation) || detoursBinaryKey[i] == CasKeyZero)
 					return m_logger.Error(TC("Failed to create cas for %s"), detoursLib.data);
-				m_detoursBinaryKey[i] = detoursBinaryKey[i];
+				m_detoursBinaryKey[isClientLinux][i] = detoursBinaryKey[i];
 			}
 
 			if (binAsVersion && agentBinaryKey == CasKeyZero)
 			{
+				UBA_ASSERT(IsLinux == isClientLinux);
 				StringBuffer<> agentDir;
 				if (!GetDirectoryOfCurrentModule(m_logger, agentDir))
 					return false;
@@ -1008,6 +1061,7 @@ namespace uba
 		session.processSlotCount = processSlotCount;
 		session.dedicated = dedicated;
 		session.isArm = isClientArm;
+		session.isLinux = isClientLinux;
 		session.memAvail = memAvail;
 		session.memTotal = memTotal;
 		session.cpuLoad = cpuLoad;
@@ -1017,9 +1071,9 @@ namespace uba
 		if (!InitializeNameToHashTable())
 			return false;
 
-		writer.WriteCasKey(m_detoursBinaryKey[0]);
+		writer.WriteCasKey(m_detoursBinaryKey[isClientLinux][0]);
 		if (isClientArm)
-			writer.WriteCasKey(m_detoursBinaryKey[1]);
+			writer.WriteCasKey(m_detoursBinaryKey[isClientLinux][1]);
 		writer.WriteBool(m_resetCas);
 		writer.WriteU32(sessionId);
 		writer.WriteU32(m_uiLanguage);
@@ -1148,9 +1202,15 @@ namespace uba
 
 	bool SessionServer::HandleGetApplication(const ConnectionInfo& connectionInfo, const WorkContext& workContext, BinaryReader& reader, BinaryWriter& writer)
 	{
-		u32 processId = reader.ReadU32(); (void)processId;
+		bool isClientLinux = reader.ReadBool();
 		StringBuffer<> applicationName;
 		reader.ReadString(applicationName);
+
+		#if UBA_TEST_LINUX_HELPING_WINDOWS
+		if (isClientLinux && applicationName.Equals(TCV("/usr/bin/gcc")))
+			applicationName.Clear().Append(TCV("d:\\temp\\gcc"));
+		#endif
+
 		StringKey applicationKey = ToStringKeyLower(applicationName);
 
 		SCOPED_FUTEX(m_applicationDataLock, lock);
@@ -1167,7 +1227,7 @@ namespace uba
 
 		u64 startPos = writer.GetPosition();
 		Vector<BinaryModule> modules;
-		if (!GetBinaryModules(modules, applicationName.data))
+		if (!GetBinaryModules(modules, applicationName.data, isClientLinux))
 			return false;
 
 		#if PLATFORM_WINDOWS
@@ -1188,10 +1248,25 @@ namespace uba
 		for (BinaryModule& m : modules)
 		{
 			CasKey casKey;
-			if (!StoreCasFile(casKey, StringKeyZero, m.path.c_str()))
+			const tchar* path = m.path.c_str();
+			u32 fileAttributes = m.fileAttributes;
+
+			#if PLATFORM_WINDOWS
+			TString tmp;
+			if (isClientLinux)
+			{
+				bool execute = true;
+				tmp = m.path;
+				Replace(tmp.data(), NonPathSeparator, PathSeparator);
+				path = tmp.c_str();
+				fileAttributes = 0644 | (execute ? 0100 : 0);
+			}
+			#endif
+
+			if (!StoreCasFile(casKey, StringKeyZero, path))
 				return false;
 			writer.WriteString(m.path);
-			writer.WriteU32(m.fileAttributes);
+			writer.WriteU32(fileAttributes);
 			writer.WriteBool(m.isSystem);
 			writer.WriteCasKey(casKey);
 			#if PLATFORM_MAC
@@ -1582,6 +1657,7 @@ namespace uba
 		auto processesLostRace = std::move(session.processesLostRace);
 		sessionsLock.Leave();
 
+		bool isCrossPlatform = IsLinux != session.isLinux;
 		bool isCrossArchitecture = IsArmBinary != session.isArm;
 
 		u32 weight32 = reader.ReadU32();
@@ -1604,7 +1680,8 @@ namespace uba
 			auto& startInfo = process->m_startInfo;
 
 			StringBuffer<> applicationOverride;
-			if (isCrossArchitecture)
+
+			if (isCrossArchitecture || isCrossPlatform)
 			{
 				auto returnProcess = [&]()
 				{
@@ -1641,6 +1718,11 @@ namespace uba
 					}
 				}
 
+				#if UBA_TEST_LINUX_HELPING_WINDOWS // PLATFORM_WINDOWS
+				applicationOverride.Clear().Append(TCV("/usr/bin/gcc"));
+				startInfo.workingDirStr = TC("/usr/bin");
+				startInfo.workingDir = startInfo.workingDirStr.c_str();
+				#endif
 			}
 
 			ProcessAdded(*process, sessionId);
@@ -2842,6 +2924,7 @@ namespace uba
 	bool SessionServer::HandleGetSymbols(const ConnectionInfo& connectionInfo, const WorkContext& workContext, BinaryReader& reader, BinaryWriter& writer)
 	{
 		TString application = reader.ReadString();
+		bool isClientLinux = reader.ReadBool();
 		bool isClientArm = reader.ReadBool();
 
 		if (IsArmBinary != isClientArm)
@@ -2850,7 +2933,7 @@ namespace uba
 			return true;
 		}
 
-		GetSymbols(application.c_str(), isClientArm, reader, writer);
+		GetSymbols(application.c_str(), isClientLinux, isClientArm, reader, writer);
 
 		if constexpr (DownloadDebugSymbols)
 		{
diff --git a/Engine/Source/Programs/UnrealBuildAccelerator/Common/Public/UbaNetwork.h b/Engine/Source/Programs/UnrealBuildAccelerator/Common/Public/UbaNetwork.h
--- a/Engine/Source/Programs/UnrealBuildAccelerator/Common/Public/UbaNetwork.h
+++ b/Engine/Source/Programs/UnrealBuildAccelerator/Common/Public/UbaNetwork.h
@@ -16,7 +16,7 @@ namespace uba
 	
 	static constexpr u32 SystemNetworkVersion = 1339;
 	static constexpr u32 StorageNetworkVersion = 4;
-	static constexpr u32 SessionNetworkVersion = 48;
+	static constexpr u32 SessionNetworkVersion = 49;
 	static constexpr u32 CacheNetworkVersion = 5;
 	static constexpr u32 CacheBucketVersion = 3;
 
diff --git a/Engine/Source/Programs/UnrealBuildAccelerator/Common/Public/UbaSession.h b/Engine/Source/Programs/UnrealBuildAccelerator/Common/Public/UbaSession.h
--- a/Engine/Source/Programs/UnrealBuildAccelerator/Common/Public/UbaSession.h
+++ b/Engine/Source/Programs/UnrealBuildAccelerator/Common/Public/UbaSession.h
@@ -140,7 +140,7 @@ namespace uba
 		SessionStats& Stats();
 
 		struct BinaryModule { TString name; TString path; u32 fileAttributes = 0; bool isSystem = false; u32 minOsVersion = 0; };
-		bool GetBinaryModules(Vector<BinaryModule>& out, const tchar* application);
+		bool GetBinaryModules(Vector<BinaryModule>& out, const tchar* application, bool isLinuxBinary);
 		void Free(Vector<BinaryModule>& v);
 		bool IsRarelyRead(ProcessImpl& process, const StringView& fileName) const;
 		bool IsMaybeReadAfterWritten(ProcessImpl& process, const StringView& fileName) const;
@@ -192,7 +192,7 @@ namespace uba
 		virtual bool CustomMessage(Process& process, BinaryReader& reader, BinaryWriter& writer);
 		virtual bool SHGetKnownFolderPath(Process& process, BinaryReader& reader, BinaryWriter& writer);
 		virtual bool HostRun(BinaryReader& reader, BinaryWriter& writer);
-		virtual bool GetSymbols(const tchar* application, bool isArm, BinaryReader& reader, BinaryWriter& writer);
+		virtual bool GetSymbols(const tchar* application, bool isLinux, bool isArm, BinaryReader& reader, BinaryWriter& writer);
 		virtual bool CheckRemapping(ProcessImpl& process, BinaryReader& reader, BinaryWriter& writer);
 		virtual bool RunSpecialProgram(ProcessImpl& process, BinaryReader& reader, BinaryWriter& writer);
 		virtual void FileEntryAdded(StringKey fileNameKey, u64 lastWritten, u64 size);
@@ -210,7 +210,7 @@ namespace uba
 		void AddEnvironmentVariableNoLock(StringView key, StringView value);
 		bool WriteDirectoryEntriesInternal(DirectoryTable::Directory& dir, const StringKey& dirKey, StringView dirPath, bool isRefresh, u32& outTableOffset);
 		void WriteDirectoryEntriesRecursive(const StringKey& dirKey, StringView dirPath, u32& outTableOffset);
-		bool CopyImports(Vector<BinaryModule>& out, const tchar* library, tchar* applicationDir, tchar* applicationDirEnd, UnorderedSet<TString>& handledImports, const char* const* loaderPaths);
+		bool CopyImports(Vector<BinaryModule>& out, const tchar* library, tchar* applicationDir, tchar* applicationDirEnd, UnorderedSet<TString>& handledImports, const tchar* const* loaderPaths, bool isLinuxBinary);
 		bool CreateProcessJobObject();
 		void EnsureDirectoryTableMemory(u64 neededSize);
 		void GetSessionInfo(StringBufferBase& out);
@@ -326,7 +326,7 @@ namespace uba
 		Atomic<u64> m_fileIndexCounter = 8000000000;
 
 		FixedSizeFileMappingAllocator m_processCommunicationAllocator;
-		std::string m_detoursLibrary[2]; // x64 and arm64
+		std::string m_detoursLibrary[2][2]; // [nonlinux/linux][x64/arm64]
 
 		Futex m_processStatsLock;
 		ProcessStats m_processStats;
diff --git a/Engine/Source/Programs/UnrealBuildAccelerator/Common/Public/UbaSessionClient.h b/Engine/Source/Programs/UnrealBuildAccelerator/Common/Public/UbaSessionClient.h
--- a/Engine/Source/Programs/UnrealBuildAccelerator/Common/Public/UbaSessionClient.h
+++ b/Engine/Source/Programs/UnrealBuildAccelerator/Common/Public/UbaSessionClient.h
@@ -77,7 +77,7 @@ namespace uba
 		virtual bool CustomMessage(Process& process, BinaryReader& reader, BinaryWriter& writer) override;
 		virtual bool SHGetKnownFolderPath(Process& process, BinaryReader& reader, BinaryWriter& writer) override;
 		virtual bool HostRun(BinaryReader& reader, BinaryWriter& writer) override;
-		virtual bool GetSymbols(const tchar* application, bool isArm, BinaryReader& reader, BinaryWriter& writer);
+		virtual bool GetSymbols(const tchar* application, bool isLinux, bool isArm, BinaryReader& reader, BinaryWriter& writer);
 		virtual bool FlushWrittenFiles(ProcessImpl& process) override;
 		virtual bool UpdateEnvironment(ProcessImpl& process, const StringView& reason, bool resetStats) override;
 		virtual bool LogLine(ProcessImpl& process, const tchar* line, LogEntryType logType) override;
diff --git a/Engine/Source/Programs/UnrealBuildAccelerator/Common/Public/UbaSessionServer.h b/Engine/Source/Programs/UnrealBuildAccelerator/Common/Public/UbaSessionServer.h
--- a/Engine/Source/Programs/UnrealBuildAccelerator/Common/Public/UbaSessionServer.h
+++ b/Engine/Source/Programs/UnrealBuildAccelerator/Common/Public/UbaSessionServer.h
@@ -166,8 +166,8 @@ namespace uba
 		u32 m_connectionCount = 0;
 
 		Futex m_binKeysLock;
-		CasKey m_detoursBinaryKey[2];
-		CasKey m_agentBinaryKey[2];
+		CasKey m_detoursBinaryKey[2][2]; // [platform][isarm]
+		CasKey m_agentBinaryKey[2]; // [isarm]
 
 		struct ClientSession
 		{
@@ -191,6 +191,7 @@ namespace uba
 			bool crashdump = false;
 			bool hasNotification = false;
 			bool isArm = false;
+			bool isLinux = false; // Only used for non-linux
 		};
 		Vector<ClientSession*> m_clientSessions;
 
diff --git a/Engine/Source/Programs/UnrealBuildAccelerator/Core/Public/UbaBase.h b/Engine/Source/Programs/UnrealBuildAccelerator/Core/Public/UbaBase.h
--- a/Engine/Source/Programs/UnrealBuildAccelerator/Core/Public/UbaBase.h
+++ b/Engine/Source/Programs/UnrealBuildAccelerator/Core/Public/UbaBase.h
@@ -21,11 +21,13 @@ namespace uba
 
 	#if PLATFORM_WINDOWS
 	inline constexpr bool IsWindows = true;
+	inline constexpr bool IsLinux = false;
 	using tchar = wchar_t;
 	#define TC(x) L##x
 	#define PERCENT_HS L"%hs"
 	#else
 	inline constexpr bool IsWindows = false;
+	inline constexpr bool IsLinux = PLATFORM_LINUX != 0;
 	using tchar = char;
 	#define TC(x) x
 	#define PERCENT_HS "%s"
diff --git a/Engine/Source/Programs/UnrealBuildAccelerator/Core/Public/UbaElf.h b/Engine/Source/Programs/UnrealBuildAccelerator/Core/Public/UbaElf.h
new file mode 100644
--- /dev/null
+++ b/Engine/Source/Programs/UnrealBuildAccelerator/Core/Public/UbaElf.h
@@ -0,0 +1,78 @@
+// Copyright Epic Games, Inc. All Rights Reserved.
+
+#pragma once
+
+namespace uba
+{
+	struct Elf64FileHeader
+	{
+		unsigned char	e_ident[16];	// Magic number and other info
+		u16	e_type;			// Object file type
+		u16	e_machine;		// Architecture
+		u32	e_version;		// Object file version
+		u64	e_entry;		// Entry point virtual address
+		u64	e_phoff;		// Program header table file offset
+		u64	e_shoff;		// Section header table file offset
+		u32	e_flags;		// Processor-specific flags
+		u16	e_ehsize;		// ELF header size in bytes
+		u16	e_phentsize;	// Program header table entry size
+		u16	e_phnum;		// Program header table entry count
+		u16	e_shentsize;	// Section header table entry size
+		u16	e_shnum;		// Section header table entry count
+		u16	e_shstrndx;		// Section header string table index
+	};
+
+	inline const constexpr char* ELFMAG	= "\177ELF";
+	inline constexpr int SELFMAG	= 4;
+	inline constexpr int EI_CLASS	= 4;
+	inline constexpr int ELFCLASS32	= 1;
+	inline constexpr int ELFCLASS64	= 2;
+
+	struct Elf64ProgramHeader
+	{
+		u32 p_type;    // Segment type
+		u32 p_flags;   // Segment flags
+		u64 p_offset;  // Offset in file
+		u64 p_vaddr;   // Virtual address in memory
+		u64 p_paddr;   // Physical address (ignored)
+		u64 p_filesz;  // Size of segment in file
+		u64 p_memsz;   // Size of segment in memory
+		u64 p_align;   // Alignment
+	};
+
+	struct Elf64SectionHeader
+	{
+		u32	sh_name;		// Section name (string tbl index)
+		u32	sh_type;		// Section type
+		u64	sh_flags;		// Section flags
+		u64	sh_addr;		// Section virtual addr at execution
+		u64	sh_offset;		// Section file offset
+		u64	sh_size;		// Section size in bytes
+		u32	sh_link;		// Link to another section
+		u32	sh_info;		// Additional section information
+		u64	sh_addralign;	// Section alignment
+		u64	sh_entsize;		// Entry size if section holds table
+	};
+
+	inline constexpr int SHT_PROGBITS		= 1;	// Program data
+	inline constexpr int SHT_SYMTAB			= 2;	// Symbol table
+	inline constexpr int SHT_STRTAB			= 3;
+	inline constexpr int SHT_RELA			= 4;	// Relocation entries with addends
+	inline constexpr int SHT_DYNAMIC		= 6;	// Dynamic linking information
+	inline constexpr int SHT_REL			= 9;	// Relocation entries, no addends
+	inline constexpr int SHT_DYNSYM			= 11;	// Dynamic linker symbol table
+	inline constexpr int SHT_SYMTAB_SHNDX	= 18;	// Extended section indeces
+
+	struct Elf64Dyn
+	{
+		int64_t d_tag;		// Dynamic entry type
+		union
+		{
+			u64 d_val;  /* Integer value */
+			u64  d_ptr;  /* Address value */
+		} d_un;
+	};
+
+	using Elf64Addr = u64;
+	using Elf64Off = u64;
+}
\ No newline at end of file
diff --git a/Engine/Source/Programs/UnrealBuildAccelerator/Core/Public/UbaLinuxBinDependencyParser.h b/Engine/Source/Programs/UnrealBuildAccelerator/Core/Public/UbaLinuxBinDependencyParser.h
--- a/Engine/Source/Programs/UnrealBuildAccelerator/Core/Public/UbaLinuxBinDependencyParser.h
+++ b/Engine/Source/Programs/UnrealBuildAccelerator/Core/Public/UbaLinuxBinDependencyParser.h
@@ -2,10 +2,8 @@
 
 #pragma once
 
-#if PLATFORM_LINUX
-
+#include "UbaElf.h"
 #include "UbaStringBuffer.h"
-#include <elf.h>
 
 #if !defined(UBA_IS_DETOURED_INCLUDE)
 #define TRUE_WRAPPER(func) func
@@ -15,9 +13,22 @@
 
 namespace uba
 {
+	inline constexpr int PT_LOAD	= 1;	// Elf64_Phdr	Load segment into memory
+	inline constexpr int PT_DYNAMIC = 2;	// Elf64_Phdr	Dynamic linker info
+	inline constexpr int DT_NULL	= 0;	// Elf64_Dyn	End marker
+	inline constexpr int DT_NEEDED	= 1;	// Elf64_Dyn	Dependency
+	inline constexpr int DT_STRTAB	= 5;	// Elf64_Dyn	String table address
+	inline constexpr int DT_STRSZ	= 10;	// Elf64_Dyn	String table size
+	inline constexpr int DT_RPATH	= 15;	// Elf64_Dyn	Legacy library path
+	inline constexpr int DT_RUNPATH = 29;	// Elf64_Dyn	Modern library path
+
+	#if !PLATFORM_LINUX
+	namespace linux {
+	#else
 	struct BinaryInfo
 	{
 	};
+	#endif
 
 	constexpr StringView g_systemFiles[] =
 	{
@@ -44,19 +55,19 @@ namespace uba
 	inline bool IsKnownSystemFile(const tchar* fileName)
 	{
 		StringView file = ToView(fileName);
-		bool isSystem = false;
 		for (auto systemFile : g_systemFiles)
 			if (file.StartsWith(systemFile.data))
 				return true;
 		return false;
 	}
 
-	inline bool ParseBinary(const StringView& filePath, const StringView& originalPath, BinaryInfo& outInfo, const Function<void(const tchar* import, bool isKnown, const char* const* loaderPaths)>& func, StringBufferBase& outError)
+	inline bool ParseBinary(const StringView& filePath, const StringView& originalPath, BinaryInfo& outInfo, const Function<void(const tchar* import, bool isKnown, const tchar* const* loaderPaths)>& func, StringBufferBase& outError)
 	{
 		#if UBA_FINDIMPORTS_DEBUG
 		printf("FINDIMPORTS: %s\n", filePath.data);
 		#endif
 
+		#if PLATFORM_LINUX
 		int fd = TRUE_WRAPPER(open)(filePath.data, O_RDONLY);
 		if (fd == -1)
 			return outError.Appendf("Open failed for file %s", filePath.data).ToFalse();
@@ -66,27 +77,45 @@ namespace uba
 		if (TRUE_WRAPPER(fstat)(fd, &sb) == -1)
 			return outError.Appendf("Stat failed for file %s", filePath.data).ToFalse();
 
-		u32 size = sb.st_size;
+		u32 fileSize = sb.st_size;
 
-		auto mem = (u8*)mmap(nullptr, size, PROT_READ, MAP_PRIVATE, fd, 0);
+		auto mem = (u8*)mmap(nullptr, fileSize, PROT_READ, MAP_PRIVATE, fd, 0);
 		if (mem == MAP_FAILED)
 			return outError.Appendf("Mmap failed for file %s", filePath.data).ToFalse();
-		auto unmap = MakeGuard([&]() { munmap(mem, size); });
+		auto unmap = MakeGuard([&]() { munmap(mem, fileSize); });
+		#else
+		HANDLE fileHandle = ::CreateFileW(filePath.data, GENERIC_READ, FILE_SHARE_READ, 0, OPEN_EXISTING, FILE_ATTRIBUTE_NORMAL | FILE_FLAG_SEQUENTIAL_SCAN, 0);
+		if (fileHandle == INVALID_HANDLE_VALUE)
+			return true;
+		auto closeFileHandle = MakeGuard([&]() { ::CloseHandle(fileHandle); });
+		HANDLE fileMapping = ::CreateFileMappingW(fileHandle, NULL, PAGE_READONLY, 0, 0, NULL);
+		if (!fileMapping)
+			return false;
+		auto closeMappingHandle = MakeGuard([&]() { ::CloseHandle(fileMapping); });
+		u8* mem = (u8*)::MapViewOfFile(fileMapping, FILE_MAP_READ, 0, 0, 0);
+		if (!mem)
+			return false;
+		auto unmap = MakeGuard([&]() { ::UnmapViewOfFile(mem); });
+		LARGE_INTEGER fileSize2;
+		if (!::GetFileSizeEx(fileHandle, &fileSize2))
+			return false;
+		u32 fileSize = u32(fileSize2.QuadPart);
+		#endif
 
-		auto ehdr = (Elf64_Ehdr*)mem;
+		auto ehdr = (Elf64FileHeader*)mem;
 		if (memcmp(ehdr->e_ident, ELFMAG, SELFMAG) != 0 || ehdr->e_ident[EI_CLASS] != ELFCLASS64)
 			return outError.Append("Not a valid 64-bit ELF file").ToFalse();
 
-		Elf64_Phdr* phdrs = (Elf64_Phdr*)(mem + ehdr->e_phoff);
-		Elf64_Shdr* shdrs = (Elf64_Shdr *)(mem + ehdr->e_shoff);
-		Elf64_Dyn* dynamic = NULL;
+		auto phdrs = (Elf64ProgramHeader*)(mem + ehdr->e_phoff);
+		auto shdrs = (Elf64SectionHeader*)(mem + ehdr->e_shoff);
+		Elf64Dyn* dynamic = NULL;
 		size_t dyn_count = 0;
 
-		auto vaddr_to_offset = [](Elf64_Addr vaddr, Elf64_Phdr *phdrs, int phnum) -> Elf64_Off {
+		auto vaddr_to_offset = [](Elf64Addr vaddr, Elf64ProgramHeader* phdrs, int phnum) -> Elf64Off {
 			for (int i = 0; i < phnum; i++) {
 				if (phdrs[i].p_type != PT_LOAD) continue;
-				Elf64_Addr start = phdrs[i].p_vaddr;
-				Elf64_Addr end = start + phdrs[i].p_memsz;
+				Elf64Addr start = phdrs[i].p_vaddr;
+				Elf64Addr end = start + phdrs[i].p_memsz;
 				if (vaddr >= start && vaddr < end) {
 					return phdrs[i].p_offset + (vaddr - start);
 				}
@@ -98,22 +127,22 @@ namespace uba
 		{
 			if (phdrs[i].p_type != PT_DYNAMIC)
 				continue;
-			Elf64_Off dyn_offset = vaddr_to_offset(phdrs[i].p_vaddr, phdrs, ehdr->e_phnum);
-			if (dyn_offset >= sb.st_size || dyn_offset + sizeof(Elf64_Dyn) > sb.st_size)
-				return outError.Appendf("Dynamic offset out of file bounds!").ToFalse();
-			dynamic = (Elf64_Dyn*)(mem + dyn_offset);			
-			dyn_count = phdrs[i].p_filesz / sizeof(Elf64_Dyn);
+			Elf64Off dyn_offset = vaddr_to_offset(phdrs[i].p_vaddr, phdrs, ehdr->e_phnum);
+			if (dyn_offset >= fileSize || dyn_offset + sizeof(Elf64Dyn) > fileSize)
+				return outError.Appendf(TC("Dynamic offset out of file bounds!")).ToFalse();
+			dynamic = (Elf64Dyn*)(mem + dyn_offset);			
+			dyn_count = phdrs[i].p_filesz / sizeof(Elf64Dyn);
 			break;
 		}
 
 		if (!dynamic)
-			return outError.Appendf("No PT_DYNAMIC segment found.").ToFalse();
-		if ((u8*)dynamic < mem || (u8*)dynamic >= mem + sb.st_size)
-			return outError.Appendf("Dynamic segment is out of file bounds!").ToFalse();
+			return outError.Appendf(TC("No PT_DYNAMIC segment found.")).ToFalse();
+		if ((u8*)dynamic < mem || (u8*)dynamic >= mem + fileSize)
+			return outError.Appendf(TC("Dynamic segment is out of file bounds!")).ToFalse();
 
-		Elf64_Addr strtab_addr = 0;
-		Elf64_Xword strsz = 0;
-		Elf64_Xword needed[512];
+		Elf64Addr strtab_addr = 0;
+		u64 strsz = 0;
+		u64 needed[512];
 		u32 needed_count = 0;
 
 		for (size_t i = 0; i < dyn_count; ++i)
@@ -145,14 +174,14 @@ namespace uba
 		// If we have strtab_addr but not strtab, resolve via vaddr to file offset
 		if (strtab_addr && !strtab)
 		{
-			Elf64_Off str_off = vaddr_to_offset(strtab_addr, phdrs, ehdr->e_phnum);
+			Elf64Off str_off = vaddr_to_offset(strtab_addr, phdrs, ehdr->e_phnum);
 			strtab = (char *)mem + str_off;
 		}
 
 		if (!strtab)
-			return outError.Appendf("Failed to find string table.").ToFalse();
+			return outError.Appendf(TC("Failed to find string table.")).ToFalse();
 
-		const char* loaderPaths[256];
+		const tchar* loaderPaths[256];
 		u32 loaderPathsCount = 0;
 
 		List<TString> fixedLoaderPaths;
@@ -168,7 +197,7 @@ namespace uba
 		
 		//StringView origin(fileName, strrchr(fileName, '/') - fileName);
 
-		for (Elf64_Dyn *dyn = dynamic; dyn->d_tag != DT_NULL; dyn++)
+		for (Elf64Dyn* dyn = dynamic; dyn->d_tag != DT_NULL; dyn++)
 		{
 			if (dyn->d_tag != DT_RPATH && dyn->d_tag != DT_RUNPATH)
 				continue;
@@ -188,8 +217,15 @@ namespace uba
 				char* pathEnd = strchr(pathBegin, ':');
 				if (!pathEnd)
 					pathEnd = end;
+				#if PLATFORM_WINDOWS
+				StringBuffer temp;
+				temp.Append(pathBegin, u32(pathEnd - pathBegin));
+				StringView loaderPath(temp);
+				#else
 				StringView loaderPath(pathBegin, pathEnd - pathBegin);
-				if (loaderPath.StartsWith("${ORIGIN}"))
+				#endif
+
+				if (loaderPath.StartsWith(TCV("${ORIGIN}")))
 					loaderPaths[loaderPathsCount++] = fixedLoaderPaths.emplace_back(originalPath.ToString() + loaderPath.Skip(9).ToString()).c_str();
 				else
 					loaderPaths[loaderPathsCount++] = fixedLoaderPaths.emplace_back(loaderPath.ToString()).c_str();
@@ -208,7 +244,7 @@ namespace uba
 		#endif
 		
 		#if 0
-		for (Elf64_Dyn *dyn = dynamic; dyn->d_tag != DT_NULL; dyn++) {
+		for (Elf64Dyn* dyn = dynamic; dyn->d_tag != DT_NULL; dyn++) {
 			if (dyn->d_tag != DT_NEEDED)
 				continue;
 			printf("  %s\n", strtab + dyn->d_un.d_val);
@@ -216,10 +252,16 @@ namespace uba
 		}
 		#endif
 
-		for (int i = 0; i < needed_count; i++) {
+		for (u32 i = 0; i < needed_count; i++) {
 			if (needed[i] < strsz || strsz == 0)
 			{
+#if PLATFORM_WINDOWS
+				StringBuffer<> temp;
+				temp.Append(strtab + needed[i]);
+				StringView file(temp);
+#else
 				StringView file = ToView(strtab + needed[i]);
+#endif
 				bool isSystem = false;
 				for (auto systemFile : g_systemFiles)
 					isSystem |= file.StartsWith(systemFile.data);
@@ -241,6 +283,9 @@ namespace uba
 		}		
 		return true;
 	}
-}
 
-#endif
+	#if !PLATFORM_LINUX
+	}
+	#endif
+
+}
diff --git a/Engine/Source/Programs/UnrealBuildAccelerator/Core/Public/UbaWinBinDependencyParser.h b/Engine/Source/Programs/UnrealBuildAccelerator/Core/Public/UbaWinBinDependencyParser.h
--- a/Engine/Source/Programs/UnrealBuildAccelerator/Core/Public/UbaWinBinDependencyParser.h
+++ b/Engine/Source/Programs/UnrealBuildAccelerator/Core/Public/UbaWinBinDependencyParser.h
@@ -111,7 +111,7 @@ namespace uba
 
 	DWORD UbaExceptionHandlerAllowAccessViolation(DWORD code, _EXCEPTION_POINTERS* info, const tchar* fileName);
 
-	inline bool FindImportsInMem(const tchar* fileName, void* mem, u64 memSize, const Function<void(const tchar* import, bool isKnown, const char* const* importLoaderPaths)>& func)
+	inline bool FindImportsInMem(const tchar* fileName, void* mem, u64 memSize, const Function<void(const tchar* import, bool isKnown, const tchar* const* importLoaderPaths)>& func)
 	{
 		auto hdrs = (PIMAGE_NT_HEADERS)(PCHAR(mem) + PIMAGE_DOS_HEADER(mem)->e_lfanew);
 		auto pSech=IMAGE_FIRST_SECTION(hdrs);
@@ -145,7 +145,7 @@ namespace uba
 		return true;
 	}
 
-	inline bool ParseBinary(const StringView& filePath, const StringView& originalPath, BinaryInfo& outInfo, const Function<void(const tchar* import, bool isKnown, const char* const* importLoaderPaths)>& func, StringBufferBase& outError)
+	inline bool ParseBinary(const StringView& filePath, const StringView& originalPath, BinaryInfo& outInfo, const Function<void(const tchar* import, bool isKnown, const tchar* const* importLoaderPaths)>& func, StringBufferBase& outError)
 	{
 		if (filePath.EndsWith(TCV(".bat")))
 			return true;
@@ -160,12 +160,11 @@ namespace uba
 		void* mem = True_MapViewOfFile(fileMapping, FILE_MAP_READ, 0, 0, 0);
 		if (!mem)
 			return false;
+		auto unmap = MakeGuard([&]() { True_UnmapViewOfFile(mem); });
 		LARGE_INTEGER fileSize;
 		if (!True_GetFileSizeEx(fileHandle, &fileSize))
 			return false;
 
-		auto unmap = MakeGuard([&]() { True_UnmapViewOfFile(mem); });
-
 		return FindImportsInMem(filePath.data, mem, fileSize.QuadPart, func);
 	}
 }
diff --git a/Engine/Source/Programs/UnrealBuildAccelerator/Detours/Private/Windows/UbaDetoursFunctionsKernelBase.inl b/Engine/Source/Programs/UnrealBuildAccelerator/Detours/Private/Windows/UbaDetoursFunctionsKernelBase.inl
--- a/Engine/Source/Programs/UnrealBuildAccelerator/Detours/Private/Windows/UbaDetoursFunctionsKernelBase.inl
+++ b/Engine/Source/Programs/UnrealBuildAccelerator/Detours/Private/Windows/UbaDetoursFunctionsKernelBase.inl
@@ -2776,7 +2776,7 @@ HMODULE Recursive_LoadLibraryExW(const StringView& filePath, LPCWSTR originalNam
 		SuppressCreateFileDetourScope cfs;
 		StringBuffer<256> error;
 		BinaryInfo info;
-		if (!ParseBinary(filePath, {}, info, [&](const wchar_t* import, bool isKnown, const char* const* importLoaderPaths)
+		if (!ParseBinary(filePath, {}, info, [&](const wchar_t* import, bool isKnown, const tchar* const* importLoaderPaths)
 			{
 				if (!GetModuleHandleW(import))
 					importedModules.emplace_back(import, isKnown);
diff --git a/Engine/Source/Programs/UnrealBuildAccelerator/Host/Private/UbaExports.cpp b/Engine/Source/Programs/UnrealBuildAccelerator/Host/Private/UbaExports.cpp
--- a/Engine/Source/Programs/UnrealBuildAccelerator/Host/Private/UbaExports.cpp
+++ b/Engine/Source/Programs/UnrealBuildAccelerator/Host/Private/UbaExports.cpp
@@ -1142,7 +1142,7 @@ extern "C"
 #if PLATFORM_WINDOWS
 		uba::StringBuffer<> errors;
 		uba::BinaryInfo info;
-		uba::ParseBinary(uba::ToView(binary), {}, info, [&](const uba::tchar* importName, bool isKnown, const char* const* importLoaderPaths) { func(importName, userData); }, errors);
+		uba::ParseBinary(uba::ToView(binary), {}, info, [&](const uba::tchar* importName, bool isKnown, const uba::tchar* const* importLoaderPaths) { func(importName, userData); }, errors);
 #else
 		UBA_ASSERT(false);
 #endif
diff --git a/Engine/Source/Programs/UnrealBuildAccelerator/Test/Private/UbaTestBasics.cpp b/Engine/Source/Programs/UnrealBuildAccelerator/Test/Private/UbaTestBasics.cpp
--- a/Engine/Source/Programs/UnrealBuildAccelerator/Test/Private/UbaTestBasics.cpp
+++ b/Engine/Source/Programs/UnrealBuildAccelerator/Test/Private/UbaTestBasics.cpp
@@ -637,7 +637,7 @@ namespace uba
 		bool importKernel = false;
 		StringBuffer<> error;
 		BinaryInfo info;
-		ParseBinary(path, {}, info, [&](const tchar* import, bool isKnown, const char* const* importLoaderPaths)
+		ParseBinary(path, {}, info, [&](const tchar* import, bool isKnown, const tchar* const* importLoaderPaths)
 		{
 			importKernel |= isKnown && Contains(import, TC("KERNEL32.dll"));
 		}, error);
diff --git a/Engine/Source/Programs/UnrealBuildAccelerator/UbaAgent.Target.cs b/Engine/Source/Programs/UnrealBuildAccelerator/UbaAgent.Target.cs
--- a/Engine/Source/Programs/UnrealBuildAccelerator/UbaAgent.Target.cs
+++ b/Engine/Source/Programs/UnrealBuildAccelerator/UbaAgent.Target.cs
@@ -65,6 +65,7 @@ public class UbaAgentTarget : TargetRules
 				"UBA_DETOURS_LIBRARY=L\"UbaDetours.dll\"",
 				"UBA_AGENT_EXECUTABLE=L\"UbaAgent.exe\"",
 				"UBA_DETOURS_LIBRARY_ANSI=\"UbaDetours.dll\"",
+				"UBA_DETOURS_LIBRARY_LINUX=\"libUbaDetours.so\"",
 			});
 		}
 		else
@@ -115,6 +116,7 @@ public class UbaAgentTarget : TargetRules
 			Rules.GlobalDefinitions.AddRange(new string[] {
 				"UBA_DETOURS_LIBRARY=\"libUbaDetours.dylib\"",
 				"UBA_AGENT_EXECUTABLE=\"UbaAgent\"",
+				"UBA_DETOURS_LIBRARY_LINUX=\"libUbaDetours.so\"",
 			});
 		}
 		else
