[UBA]
* Added so running processes remotely from ubacli always allow cross platform/architecture

diff --git a/Engine/Source/Programs/UnrealBuildAccelerator/Cli/Private/UbaCli.cpp b/Engine/Source/Programs/UnrealBuildAccelerator/Cli/Private/UbaCli.cpp
--- a/Engine/Source/Programs/UnrealBuildAccelerator/Cli/Private/UbaCli.cpp
+++ b/Engine/Source/Programs/UnrealBuildAccelerator/Cli/Private/UbaCli.cpp
@@ -1406,6 +1406,8 @@ namespace uba
 
 		auto RunRemote = [&](SessionServer& sessionServer, const TString& app, const TString& arg)
 		{
+			bool allowCross = true;
+
 			u64 start = GetTime();
 			ProcessStartInfo pinfo;
 			pinfo.description = app.c_str();
@@ -1418,7 +1420,7 @@ namespace uba
 			if (enableStdOut)
 				pinfo.logLineFunc = [](void* userData, const tchar* line, u32 length, LogEntryType type) { ((Logger*)userData)->Log(type, line, length); };
 			logger.Info(TC("Running %s %s"), app.c_str(), arg.c_str());
-			ProcessHandle process = sessionServer.RunProcessRemote(pinfo);
+			ProcessHandle process = sessionServer.RunProcessRemote(pinfo, 1.0f, nullptr, 0, allowCross);
 			process.WaitForExit(~0u);
 			if (process.GetExitCode() != 0)
 				return logger.Error(TC("Error exit code: %u"), process.GetExitCode());
