[UBA]
* Added ubacli -cross=<from>;<to> that can be used to add converters when using cross platform/architecture

diff --git a/Engine/Source/Programs/UnrealBuildAccelerator/Cli/Private/UbaCli.cpp b/Engine/Source/Programs/UnrealBuildAccelerator/Cli/Private/UbaCli.cpp
--- a/Engine/Source/Programs/UnrealBuildAccelerator/Cli/Private/UbaCli.cpp
+++ b/Engine/Source/Programs/UnrealBuildAccelerator/Cli/Private/UbaCli.cpp
@@ -82,6 +82,7 @@ namespace uba
 		logger.Info(TC("   -workdir=<dir>          Working directory"));
 		logger.Info(TC("   -config=<file>          Config file that contains options for various systems"));
 		logger.Info(TC("   -vfs=<virtual>;<local>  Will convert virtual path to local under the hood. Can have multiple -vfs"));
+		logger.Info(TC("   -cross=<virtual>;<local> Will add cross architecture/platform entries. Can have multiple -cross"));
 		logger.Info(TC("   -checkcas               Check so all cas entries are correct. Combine with -deleteCas to delete bad ones"));
 		logger.Info(TC("   -checkfiletable         Check so file table has correct cas stored. if bad entries it is enough to delete 'casdb' file"));
 		logger.Info(TC("   -checkcloud             Check if we are inside cloud and output information about cloud"));
@@ -246,13 +247,12 @@ namespace uba
 		TString symbolFile;
 		u64 symbolOffset = 0;
 
-		struct VfsEntry
-		{
-			TString virtualPath;
-			TString localPath;
-		};
+		struct VfsEntry { TString virtualPath; TString localPath; };
 		Vector<VfsEntry> vfsEntries;
 
+		struct CrossEntry { TString from; TString to; };
+		Vector<CrossEntry> crossEntries;
+
 		u32 loopCount = 1;
 
 		enum CommandType
@@ -326,6 +326,16 @@ namespace uba
 					u32 semiPos = u32(semi - value.data);
 					vfsEntries.push_back({StringView(value.data, semiPos).ToString(), StringView(value).Skip(semiPos + 1).ToString()});
 				}
+				else if (name.Equals(TCV("-cross")))
+				{
+					if (value.IsEmpty())
+						return PrintHelp(TC("-cross needs a value"));
+					const tchar* semi = value.First(';');
+					if (!semi)
+						return PrintHelp(TC("-cross needs a semicolon between virtual and local path"));
+					u32 semiPos = u32(semi - value.data);
+					crossEntries.push_back({StringView(value.data, semiPos).ToString(), StringView(value).Skip(semiPos + 1).ToString()});
+				}
 				else if (name.Equals(TCV("-capacity")))
 				{
 					if (!value.Parse(storageCapacityGb))
@@ -1301,7 +1311,8 @@ namespace uba
 				sessionServer.DevirtualizeString(application, rootsHandle, true);
 			}
 
-
+			for (CrossEntry& entry : crossEntries)
+				sessionServer.RegisterCrossArchitectureMapping(entry.from.c_str(), entry.to.c_str());
 
 			if (isRemote || useScheduler)
 			{
