[UBA]
* Fixed mac compile errors

diff --git a/Engine/Source/Programs/UnrealBuildAccelerator/Core/Public/UbaElf.h b/Engine/Source/Programs/UnrealBuildAccelerator/Core/Public/UbaElf.h
--- a/Engine/Source/Programs/UnrealBuildAccelerator/Core/Public/UbaElf.h
+++ b/Engine/Source/Programs/UnrealBuildAccelerator/Core/Public/UbaElf.h
@@ -2,6 +2,8 @@
 
 #pragma once
 
+#include "UbaBase.h"
+
 namespace uba
 {
 	struct Elf64FileHeader
@@ -65,7 +67,7 @@ namespace uba
 
 	struct Elf64Dyn
 	{
-		int64_t d_tag;		// Dynamic entry type
+		s64 d_tag;		// Dynamic entry type
 		union
 		{
 			u64 d_val;  /* Integer value */
diff --git a/Engine/Source/Programs/UnrealBuildAccelerator/Core/Public/UbaLinuxBinDependencyParser.h b/Engine/Source/Programs/UnrealBuildAccelerator/Core/Public/UbaLinuxBinDependencyParser.h
--- a/Engine/Source/Programs/UnrealBuildAccelerator/Core/Public/UbaLinuxBinDependencyParser.h
+++ b/Engine/Source/Programs/UnrealBuildAccelerator/Core/Public/UbaLinuxBinDependencyParser.h
@@ -5,6 +5,10 @@
 #include "UbaElf.h"
 #include "UbaStringBuffer.h"
 
+#if PLATFORM_MAC
+#include "UbaMacBinDependencyParser.h"
+#endif
+
 #if !defined(UBA_IS_DETOURED_INCLUDE)
 #define TRUE_WRAPPER(func) func
 #endif
@@ -67,7 +71,7 @@ namespace uba
 		printf("FINDIMPORTS: %s\n", filePath.data);
 		#endif
 
-		#if PLATFORM_LINUX
+		#if !PLATFORM_WINDOWS
 		int fd = TRUE_WRAPPER(open)(filePath.data, O_RDONLY);
 		if (fd == -1)
 			return outError.Appendf("Open failed for file %s", filePath.data).ToFalse();
