[UBA]
* Fixes around cross platform communication

diff --git a/Engine/Source/Programs/UnrealBuildAccelerator/Cli/Private/UbaCli.cpp b/Engine/Source/Programs/UnrealBuildAccelerator/Cli/Private/UbaCli.cpp
--- a/Engine/Source/Programs/UnrealBuildAccelerator/Cli/Private/UbaCli.cpp
+++ b/Engine/Source/Programs/UnrealBuildAccelerator/Cli/Private/UbaCli.cpp
@@ -81,8 +81,8 @@ namespace uba
 		logger.Info(TC("   -loop=<count>           Loop the commandline <count> number of times. Will exit when/if it fails"));
 		logger.Info(TC("   -workdir=<dir>          Working directory"));
 		logger.Info(TC("   -config=<file>          Config file that contains options for various systems"));
-		logger.Info(TC("   -vfs=<virtual>;<local>  Will convert virtual path to local under the hood. Can have multiple -vfs"));
-		logger.Info(TC("   -cross=<virtual>;<local> Will add cross architecture/platform entries. Can have multiple -cross"));
+		logger.Info(TC("   -vfs=<virtual>;<local>  Will convert virtual path to local under the hood. Can have multiple -vfs."));
+		logger.Info(TC("   -cross=<from>;<to>.     Will add cross architecture/platform entries. Can have multiple -cross."));
 		logger.Info(TC("   -checkcas               Check so all cas entries are correct. Combine with -deleteCas to delete bad ones"));
 		logger.Info(TC("   -checkfiletable         Check so file table has correct cas stored. if bad entries it is enough to delete 'casdb' file"));
 		logger.Info(TC("   -checkcloud             Check if we are inside cloud and output information about cloud"));
@@ -322,7 +322,7 @@ namespace uba
 						return PrintHelp(TC("-vfs needs a value"));
 					const tchar* semi = value.First(';');
 					if (!semi)
-						return PrintHelp(TC("-vfs needs a semicolon between virtual and local path"));
+						return PrintHelp(TC("-vfs needs a semicolon between virtual and local path (And don't forget quotes around value)"));
 					u32 semiPos = u32(semi - value.data);
 					vfsEntries.push_back({StringView(value.data, semiPos).ToString(), StringView(value).Skip(semiPos + 1).ToString()});
 				}
@@ -332,7 +332,7 @@ namespace uba
 						return PrintHelp(TC("-cross needs a value"));
 					const tchar* semi = value.First(';');
 					if (!semi)
-						return PrintHelp(TC("-cross needs a semicolon between virtual and local path"));
+						return PrintHelp(TC("-cross needs a semicolon between paths (And don't forget quotes around value)"));
 					u32 semiPos = u32(semi - value.data);
 					crossEntries.push_back({StringView(value.data, semiPos).ToString(), StringView(value).Skip(semiPos + 1).ToString()});
 				}
@@ -1312,7 +1312,8 @@ namespace uba
 			}
 
 			for (CrossEntry& entry : crossEntries)
-				sessionServer.RegisterCrossArchitectureMapping(entry.from.c_str(), entry.to.c_str());
+				if (!sessionServer.RegisterCrossArchitectureMapping(entry.from.c_str(), entry.to.c_str()))
+					return false;
 
 			if (isRemote || useScheduler)
 			{
diff --git a/Engine/Source/Programs/UnrealBuildAccelerator/Common/Private/UbaSessionClient.cpp b/Engine/Source/Programs/UnrealBuildAccelerator/Common/Private/UbaSessionClient.cpp
--- a/Engine/Source/Programs/UnrealBuildAccelerator/Common/Private/UbaSessionClient.cpp
+++ b/Engine/Source/Programs/UnrealBuildAccelerator/Common/Private/UbaSessionClient.cpp
@@ -333,6 +333,7 @@ namespace uba
 			UBA_ASSERT(fileNameKey == StringKeyZero);
 			StackBinaryWriter<1024> writer;
 			NetworkMessage msg(m_client, ServiceId, SessionMessageType_EnsureBinaryFile, writer);
+			writer.WriteBool(IsLinux);
 			writer.WriteBool(IsRunningArm());
 			writer.WriteString(fileName);
 			writer.WriteStringKey(fileNameKey);
diff --git a/Engine/Source/Programs/UnrealBuildAccelerator/Common/Private/UbaSessionServer.cpp b/Engine/Source/Programs/UnrealBuildAccelerator/Common/Private/UbaSessionServer.cpp
--- a/Engine/Source/Programs/UnrealBuildAccelerator/Common/Private/UbaSessionServer.cpp
+++ b/Engine/Source/Programs/UnrealBuildAccelerator/Common/Private/UbaSessionServer.cpp
@@ -801,9 +801,14 @@ namespace uba
 		}
 	}
 
-	void SessionServer::RegisterCrossArchitectureMapping(const tchar* from, const tchar* to)
+	bool SessionServer::RegisterCrossArchitectureMapping(const tchar* from, const tchar* to)
 	{
+		if (!IsAbsolutePath(from))
+			return m_logger.Error(TC("Must register absolute paths for cross mapping (%s)"), from);
+		if (!IsAbsolutePath(to))
+			return m_logger.Error(TC("Must register absolute paths for cross mapping (%s)"), to);
 		m_crossArchitectureMappings.emplace_back(CrossArchitectureMapping{from, to});
+		return true;
 	}
 
 	void SessionServer::SetOuterScheduler(Scheduler* scheduler)
@@ -1099,7 +1104,9 @@ namespace uba
 
 	bool SessionServer::HandleEnsureBinaryFile(const ConnectionInfo& connectionInfo, const WorkContext& workContext, BinaryReader& reader, BinaryWriter& writer)
 	{
-		bool clientIsArm = reader.ReadBool();
+		bool isClientLinux = reader.ReadBool();
+		bool isClientArm = reader.ReadBool();
+
 		StringBuffer<> fileName;
 		reader.ReadString(fileName);
 		StringKey fileNameKey = reader.ReadStringKey();
@@ -1130,9 +1137,11 @@ namespace uba
 		CasKey casKey = CasKeyZero;
 		StringBuffer<> absoluteFile;
 
-		auto FixCrossArchitecture = [this, clientIsArm](StringBuffer<>& absoluteFile)
+		bool isCross = isClientLinux != IsLinux || isClientArm != IsArmBinary;
+
+		auto FixCrossArchitecture = [this, isCross](StringBuffer<>& absoluteFile)
 			{
-				if (clientIsArm == IsArmBinary)
+				if (!isCross)
 					return;
 				for (auto& mapping : m_crossArchitectureMappings)
 					if (absoluteFile.StartsWith(mapping.from.c_str()))
@@ -2927,6 +2936,12 @@ namespace uba
 		bool isClientLinux = reader.ReadBool();
 		bool isClientArm = reader.ReadBool();
 
+		if (IsLinux != isClientLinux)
+		{
+			writer.WriteString(TC("Can't resolve callstack on cross platforms"));
+			return true;
+		}
+
 		if (IsArmBinary != isClientArm)
 		{
 			writer.WriteString(TC("Can't resolve callstack on cross architectures"));
diff --git a/Engine/Source/Programs/UnrealBuildAccelerator/Common/Public/UbaSessionServer.h b/Engine/Source/Programs/UnrealBuildAccelerator/Common/Public/UbaSessionServer.h
--- a/Engine/Source/Programs/UnrealBuildAccelerator/Common/Public/UbaSessionServer.h
+++ b/Engine/Source/Programs/UnrealBuildAccelerator/Common/Public/UbaSessionServer.h
@@ -85,7 +85,7 @@ namespace uba
 		void UnregisterNetworkTrafficProvider(u64 key);
 
 		// Register mappings for cross architecture helpers (helpers with different architecture will use the mapping)
-		void RegisterCrossArchitectureMapping(const tchar* from, const tchar* to);
+		bool RegisterCrossArchitectureMapping(const tchar* from, const tchar* to);
 
 		void SetOuterScheduler(Scheduler* scheduler);
 		Scheduler* GetOuterScheduler();
