[UBA]
* Fixed broken cross architecture/platform path in scheduler

diff --git a/Engine/Source/Programs/UnrealBuildAccelerator/Common/Private/UbaScheduler.cpp b/Engine/Source/Programs/UnrealBuildAccelerator/Common/Private/UbaScheduler.cpp
--- a/Engine/Source/Programs/UnrealBuildAccelerator/Common/Private/UbaScheduler.cpp
+++ b/Engine/Source/Programs/UnrealBuildAccelerator/Common/Private/UbaScheduler.cpp
@@ -234,7 +234,7 @@ namespace uba
 	void Scheduler::Start()
 	{
 		m_session.SetRemoteProcessReturnedEvent([this](Process& process) { ProcessReturned(process, false); });
-		m_session.SetRemoteProcessSlotAvailableEvent([this](bool isCrossArchitecture) { RemoteSlotAvailable(isCrossArchitecture); });
+		m_session.SetRemoteProcessSlotAvailableEvent([this](bool isCrossArchitecture, bool isCrossPlatform) { RemoteSlotAvailable(isCrossArchitecture, isCrossPlatform); });
 
 		if (!m_cacheClients.empty())
 			m_cacheWorkManager = new WorkManagerImpl(m_maxParallelCacheQueries, TC("UbaSchedCach"));
@@ -378,6 +378,8 @@ namespace uba
 		bool useCache = info.cacheBucketId && !m_cacheClients.empty() && rules->IsCacheable();
 		bool canDetour = info.canDetour;
 		bool canExecuteRemotely = info.canExecuteRemotely && info.canDetour;
+		bool canCrossArchitecture = info.canCrossArchitecture;
+		bool canCrossPlatform = info.canCrossPlatform;
 
 		info2->trackInputs |= m_writeTrackedInputs;
 
@@ -406,6 +408,8 @@ namespace uba
 		entry.cacheStatus = useCache ? CacheStatus_QueuedForTest : CacheStatus_Failed;
 		entry.canDetour = canDetour;
 		entry.canExecuteRemotely = canExecuteRemotely;
+		entry.canCrossArchitecture = canCrossArchitecture;
+		entry.canCrossPlatform = canCrossPlatform;
 		lock.Leave();
 
 		++m_totalProcesses;
@@ -518,7 +522,7 @@ namespace uba
 			while (RunCacheQuery())
 				;
 
-			while (RunQueuedProcess(true))
+			while (RunQueuedProcess(true, false, false))
 				;
 		}
 	}
@@ -653,10 +657,9 @@ namespace uba
 		m_updateThreadWakeup.Set();
 	}
 
-	void Scheduler::RemoteSlotAvailable(bool isCrossArchitecture)
+	void Scheduler::RemoteSlotAvailable(bool isCrossArchitecture, bool isCrossPlatform)
 	{
-		UBA_ASSERTF(!isCrossArchitecture, TC("Cross architecture code path not implemented"));
-		if (RunQueuedProcess(false))
+		if (RunQueuedProcess(false, isCrossArchitecture, isCrossPlatform))
 			return;
 		if (!m_allowDisableRemoteExecution)
 			return;
@@ -845,7 +848,7 @@ namespace uba
 		return entryToProcess;
 	}
 
-	u32 Scheduler::PopProcess(bool isLocal, RunStatus& outPrevStatus, bool forReuse)
+	u32 Scheduler::PopProcess(bool isLocal, bool isCrossArchitecture, bool isCrossPlatform, RunStatus& outPrevStatus, bool forReuse)
 	{
 		float activeWeight = m_activeLocalProcessWeight;
 		float maxWeight = float(m_maxLocalProcessors);
@@ -895,6 +898,8 @@ namespace uba
 			{
 				if (!entry.canExecuteRemotely)
 					continue;
+				if (isCrossArchitecture && !entry.canCrossArchitecture || isCrossPlatform && !entry.canCrossPlatform)
+					continue;
 			}
 
 			bool canRun = true;
@@ -1062,13 +1067,13 @@ namespace uba
 		return true;
 	}
 
-	bool Scheduler::RunQueuedProcess(bool isLocal)
+	bool Scheduler::RunQueuedProcess(bool isLocal, bool isCrossArchitecture, bool isCrossPlatform)
 	{
 		while (true)
 		{
 			RunStatus prevStatus;
 			SCOPED_FUTEX(m_processEntriesLock, lock);
-			u32 indexToRun = PopProcess(isLocal, prevStatus, false);
+			u32 indexToRun = PopProcess(isLocal, isCrossArchitecture, isCrossPlatform, prevStatus, false);
 			if (indexToRun == ~0u)
 			{
 				lock.Leave();
@@ -1216,7 +1221,7 @@ namespace uba
 		{
 			RunStatus prevStatus;
 			SCOPED_FUTEX(m_processEntriesLock, lock);
-			u32 indexToRun = PopProcess(isLocal, prevStatus, forReuse);
+			u32 indexToRun = PopProcess(isLocal, false, false, prevStatus, forReuse);
 			if (indexToRun == ~0u)
 			{
 				u64 usedTimeMs = TimeToMs(GetTime() - startTime);
diff --git a/Engine/Source/Programs/UnrealBuildAccelerator/Common/Private/UbaSessionServer.cpp b/Engine/Source/Programs/UnrealBuildAccelerator/Common/Private/UbaSessionServer.cpp
--- a/Engine/Source/Programs/UnrealBuildAccelerator/Common/Private/UbaSessionServer.cpp
+++ b/Engine/Source/Programs/UnrealBuildAccelerator/Common/Private/UbaSessionServer.cpp
@@ -681,7 +681,7 @@ namespace uba
 		return true;
 	}
 
-	void SessionServer::SetRemoteProcessSlotAvailableEvent(const Function<void(bool isCrossArchitecture)>& remoteProcessSlotAvailableEvent)
+	void SessionServer::SetRemoteProcessSlotAvailableEvent(const RemoteProcessAvailableCallback& remoteProcessSlotAvailableEvent)
 	{
 		SCOPED_WRITE_LOCK(m_remoteProcessSlotAvailableEventLock, lock);
 		m_remoteProcessSlotAvailableEvent = remoteProcessSlotAvailableEvent;
@@ -2411,7 +2411,7 @@ namespace uba
 			if (hasCalledCallback)
 				return nullptr;
 
-			m_remoteProcessSlotAvailableEvent(IsArmBinary != session.isArm);
+			m_remoteProcessSlotAvailableEvent(IsArmBinary != session.isArm, IsLinux != session.isLinux);
 			hasCalledCallback = true;
 		}
 		return nullptr;
diff --git a/Engine/Source/Programs/UnrealBuildAccelerator/Common/Public/UbaScheduler.h b/Engine/Source/Programs/UnrealBuildAccelerator/Common/Public/UbaScheduler.h
--- a/Engine/Source/Programs/UnrealBuildAccelerator/Common/Public/UbaScheduler.h
+++ b/Engine/Source/Programs/UnrealBuildAccelerator/Common/Public/UbaScheduler.h
@@ -58,6 +58,8 @@ namespace uba
 		float weight = 1.0f; // Weight of process. This is used towards max local processors. If a process is multithreaded it is likely it's weight should be more than 1.0
 		bool canDetour = true; // If true, uba will detour the process. If false it will just create pipes for std out and then run the process as-is.
 		bool canExecuteRemotely = true; // If true, this process can run on other machines, if false it will always be executed locally
+		bool canCrossArchitecture = false; // If true, this process can run on other architecture
+		bool canCrossPlatform = false; // If true, this process can run on other platform.
 
 		const void* knownInputs = nullptr; // knownInputs is a memory block with null terminated tchar strings followed by an empty null terminated string to end.
 		u32 knownInputsBytes = 0; // knownInputsBytes is the total size in bytes of knownInputs
@@ -131,13 +133,13 @@ namespace uba
 		void SkipAllQueued();
 		bool ProcessReturned(Process& process, bool isLocal);
 		void HandleCacheMissed(u32 processIndex);
-		void RemoteSlotAvailable(bool isCrossArchitecture);
+		void RemoteSlotAvailable(bool isCrossArchitecture, bool isCrossPlatform);
 		void ProcessExited(ExitProcessInfo* info, const ProcessHandle& handle);
 		bool CanRun();
 		u32 PopCache();
-		u32 PopProcess(bool isLocal, RunStatus& outPrevStatus, bool forReuse);
+		u32 PopProcess(bool isLocal, bool isCrossArchitecture, bool isCrossPlatform, RunStatus& outPrevStatus, bool forReuse);
 		bool RunCacheQuery();
-		bool RunQueuedProcess(bool isLocal);
+		bool RunQueuedProcess(bool isLocal, bool isCrossArchitecture, bool isCrossPlatform);
 		bool RaceRemoteProcess();
 		bool HandleGetNextProcess(Process& process, NextProcessInfo& outNextProcess, u32 prevExitCode, BinaryReader& prevStats, u32 timeoutMs);
 		void ExitProcess(ExitProcessInfo& info, Process& process, u32 exitCode, bool fromCache, bool forReuse);
@@ -171,6 +173,8 @@ namespace uba
 			RunStatus runStatus;
 			bool canDetour;
 			bool canExecuteRemotely;
+			bool canCrossArchitecture;
+			bool canCrossPlatform;
 		};
 
 		Futex m_processEntriesLock;
diff --git a/Engine/Source/Programs/UnrealBuildAccelerator/Common/Public/UbaSessionServer.h b/Engine/Source/Programs/UnrealBuildAccelerator/Common/Public/UbaSessionServer.h
--- a/Engine/Source/Programs/UnrealBuildAccelerator/Common/Public/UbaSessionServer.h
+++ b/Engine/Source/Programs/UnrealBuildAccelerator/Common/Public/UbaSessionServer.h
@@ -44,7 +44,8 @@ namespace uba
 		bool GetCasKeyFromTrackedInputs(CasKey& out, const tchar* fileName, const tchar* workingDir, const u8* data, u32 dataLen);
 
 		// Callback that will be called when a client is asking for processes to run. All remotes frequently ask for processes to run if they have free process slots
-		void SetRemoteProcessSlotAvailableEvent(const Function<void(bool isCrossArchitecture)>& remoteProcessSlotAvailableEvent);
+		using RemoteProcessAvailableCallback = Function<void(bool isCrossArchitecture, bool isCrossPlatform)>;
+		void SetRemoteProcessSlotAvailableEvent(const RemoteProcessAvailableCallback& remoteProcessSlotAvailableEvent);
 
 		// Callback that is called when process is returned. This could be because a client unexpectedly disconnected or is running out of memory
 		void SetRemoteProcessReturnedEvent(const Function<void(Process&)>& remoteProcessReturnedEvent);
@@ -150,7 +151,7 @@ namespace uba
 		static constexpr u8 ServiceId = SessionServiceId;
 
 		ReaderWriterLock m_remoteProcessSlotAvailableEventLock;
-		Function<void(bool isCrossArchitecture)> m_remoteProcessSlotAvailableEvent;
+		RemoteProcessAvailableCallback m_remoteProcessSlotAvailableEvent;
 
 		ReaderWriterLock m_remoteProcessReturnedEventLock;
 		Function<void(Process&)> m_remoteProcessReturnedEvent;
diff --git a/Engine/Source/Programs/UnrealBuildAccelerator/Host/Private/UbaExports.cpp b/Engine/Source/Programs/UnrealBuildAccelerator/Host/Private/UbaExports.cpp
--- a/Engine/Source/Programs/UnrealBuildAccelerator/Host/Private/UbaExports.cpp
+++ b/Engine/Source/Programs/UnrealBuildAccelerator/Host/Private/UbaExports.cpp
@@ -527,7 +527,7 @@ extern "C"
 	}
 	void SessionServer_SetRemoteProcessAvailable(uba::SessionServer* server, SessionServer_RemoteProcessAvailableCallback* available, void* userData)
 	{
-		server->SetRemoteProcessSlotAvailableEvent([available, userData](bool isCrossArchitecture) { available(userData, isCrossArchitecture); });
+		server->SetRemoteProcessSlotAvailableEvent([available, userData](bool isCrossArchitecture, bool isCrossPlatform) { available(userData, isCrossArchitecture); });
 	}
 	void SessionServer_SetRemoteProcessReturned(uba::SessionServer* server, SessionServer_RemoteProcessReturnedCallback* returned, void* userData)
 	{
@@ -817,6 +817,11 @@ extern "C"
 	}
 
 	uba::u32 Scheduler_EnqueueProcess3(uba::Scheduler* scheduler, const uba::ProcessStartInfo& info, float weight, bool canDetour, bool canExecuteRemotely, const uba::u32* dependencies, uba::u32 dependencyCount, const void* knownInputs, uba::u32 knownInputsBytes, uba::u32 knownInputsCount, uba::u32 cacheBucket, uba::u32 memoryGroup, uba::u64 predictedMemoryUsage)
+	{
+		return Scheduler_EnqueueProcess4(scheduler, info, weight, canDetour, canExecuteRemotely, false, false, dependencies, dependencyCount, knownInputs, knownInputsBytes, knownInputsCount, cacheBucket, memoryGroup, predictedMemoryUsage);
+	}
+
+	uba::u32 Scheduler_EnqueueProcess4(uba::Scheduler* scheduler, const uba::ProcessStartInfo& info, float weight, bool canDetour, bool canExecuteRemotely, bool canCrossArchitecture, bool canCrossPlatform, const uba::u32* dependencies, uba::u32 dependencyCount, const void* knownInputs, uba::u32 knownInputsBytes, uba::u32 knownInputsCount, uba::u32 cacheBucket, uba::u32 memoryGroup, uba::u64 predictedMemoryUsage)
 	{
 		uba::EnqueueProcessInfo epi(info);
 		epi.weight = weight;
@@ -827,6 +832,8 @@ extern "C"
 		epi.knownInputsCount = knownInputsCount;
 		epi.canDetour = canDetour;
 		epi.canExecuteRemotely = canExecuteRemotely;
+		epi.canCrossArchitecture = canCrossArchitecture;
+		epi.canCrossPlatform = canCrossPlatform;
 		epi.cacheBucketId = cacheBucket;
 		epi.memoryGroupId = memoryGroup;
 		epi.predictedMemoryUsage = predictedMemoryUsage;
diff --git a/Engine/Source/Programs/UnrealBuildAccelerator/Host/Public/UbaExports.h b/Engine/Source/Programs/UnrealBuildAccelerator/Host/Public/UbaExports.h
--- a/Engine/Source/Programs/UnrealBuildAccelerator/Host/Public/UbaExports.h
+++ b/Engine/Source/Programs/UnrealBuildAccelerator/Host/Public/UbaExports.h
@@ -164,6 +164,7 @@ extern "C"
 	UBA_API uba::u32 Scheduler_EnqueueProcess(uba::Scheduler* scheduler, const uba::ProcessStartInfo& info, float weight = 1.0f, const void* knownInputs = nullptr, uba::u32 knownInputsBytes = 0, uba::u32 knownInputsCount = 0);
 	UBA_API uba::u32 Scheduler_EnqueueProcess2(uba::Scheduler* scheduler, const uba::ProcessStartInfo& info, float weight, bool canDetour, bool canExecuteRemotely, const uba::u32* dependencies, uba::u32 dependencyCount, const void* knownInputs, uba::u32 knownInputsBytes, uba::u32 knownInputsCount, uba::u32 cacheBucket);
 	UBA_API uba::u32 Scheduler_EnqueueProcess3(uba::Scheduler* scheduler, const uba::ProcessStartInfo& info, float weight, bool canDetour, bool canExecuteRemotely, const uba::u32* dependencies, uba::u32 dependencyCount, const void* knownInputs, uba::u32 knownInputsBytes, uba::u32 knownInputsCount, uba::u32 cacheBucket, uba::u32 memoryGroup, uba::u64 predictedMemoryUsage);
+	UBA_API uba::u32 Scheduler_EnqueueProcess4(uba::Scheduler* scheduler, const uba::ProcessStartInfo& info, float weight, bool canDetour, bool canExecuteRemotely, bool canCrossArchitecture, bool canCrossPlatform, const uba::u32* dependencies, uba::u32 dependencyCount, const void* knownInputs, uba::u32 knownInputsBytes, uba::u32 knownInputsCount, uba::u32 cacheBucket, uba::u32 memoryGroup, uba::u64 predictedMemoryUsage);
 	UBA_API void Scheduler_SetMaxLocalProcessors(uba::Scheduler* scheduler, uba::u32 maxLocalProcessors);
 	UBA_API void Scheduler_Stop(uba::Scheduler* scheduler);
 	UBA_API void Scheduler_Cancel(uba::Scheduler* scheduler);
