[UBA]
* Fixed non-windows compile warning

diff --git a/Engine/Source/Programs/UnrealBuildAccelerator/Common/Private/UbaScheduler.cpp b/Engine/Source/Programs/UnrealBuildAccelerator/Common/Private/UbaScheduler.cpp
--- a/Engine/Source/Programs/UnrealBuildAccelerator/Common/Private/UbaScheduler.cpp
+++ b/Engine/Source/Programs/UnrealBuildAccelerator/Common/Private/UbaScheduler.cpp
@@ -898,7 +898,7 @@ namespace uba
 			{
 				if (!entry.canExecuteRemotely)
 					continue;
-				if (isCrossArchitecture && !entry.canCrossArchitecture || isCrossPlatform && !entry.canCrossPlatform)
+				if ((isCrossArchitecture && !entry.canCrossArchitecture) || (isCrossPlatform && !entry.canCrossPlatform))
 					continue;
 			}
 
