[UBA]
* Changed communication lock on detours side to be futex instead of rwlock

diff --git a/Engine/Source/Programs/UnrealBuildAccelerator/Core/Private/UbaSynchronization.cpp b/Engine/Source/Programs/UnrealBuildAccelerator/Core/Private/UbaSynchronization.cpp
--- a/Engine/Source/Programs/UnrealBuildAccelerator/Core/Private/UbaSynchronization.cpp
+++ b/Engine/Source/Programs/UnrealBuildAccelerator/Core/Private/UbaSynchronization.cpp
@@ -69,6 +69,19 @@ namespace uba
 		#endif
 	}
 
+	bool CriticalSection::TryEnter() const
+	{
+		#if PLATFORM_WINDOWS
+		return TryEnterCriticalSection((CRITICAL_SECTION*)&data);
+		#else
+		int res = pthread_mutex_trylock((pthread_mutex_t*)data);(void)res;
+		if (res == 0)
+			return true;
+		UBA_ASSERTF(res == EBUSY, TC("pthread_mutex_lock failed: %i"), res);
+		return false;
+		#endif
+	}
+
 	ReaderWriterLock::ReaderWriterLock()
 	{
 		#if PLATFORM_WINDOWS
diff --git a/Engine/Source/Programs/UnrealBuildAccelerator/Core/Public/UbaSynchronization.h b/Engine/Source/Programs/UnrealBuildAccelerator/Core/Public/UbaSynchronization.h
--- a/Engine/Source/Programs/UnrealBuildAccelerator/Core/Public/UbaSynchronization.h
+++ b/Engine/Source/Programs/UnrealBuildAccelerator/Core/Public/UbaSynchronization.h
@@ -40,6 +40,8 @@ namespace uba
 		void Enter() const;
 		void Leave() const;
 
+		bool TryEnter() const;
+
 	private:
 		#if PLATFORM_WINDOWS
 		u64 data[5];
diff --git a/Engine/Source/Programs/UnrealBuildAccelerator/Detours/Private/PosixOS/UbaDetoursFunctionsPosix.cpp b/Engine/Source/Programs/UnrealBuildAccelerator/Detours/Private/PosixOS/UbaDetoursFunctionsPosix.cpp
--- a/Engine/Source/Programs/UnrealBuildAccelerator/Detours/Private/PosixOS/UbaDetoursFunctionsPosix.cpp
+++ b/Engine/Source/Programs/UnrealBuildAccelerator/Detours/Private/PosixOS/UbaDetoursFunctionsPosix.cpp
@@ -936,7 +936,7 @@ UBA_EXPORT int UBA_WRAPPER(mkdir)(const char* path, mode_t mode)
 
 	{
 		TimerScope ts(g_stats.createFile);
-		SCOPED_WRITE_LOCK(g_communicationLock, pcs);
+		SCOPED_FUTEX(g_communicationLock, pcs);
 		BinaryWriter writer;
 		writer.WriteByte(MessageType_CreateDirectory);
 		writer.WriteStringKey(pathNameKey);
@@ -976,7 +976,7 @@ UBA_EXPORT int UBA_WRAPPER(rmdir)(const char* path)
 
 	{
 		TimerScope ts(g_stats.deleteFile);
-		SCOPED_WRITE_LOCK(g_communicationLock, pcs);
+		SCOPED_FUTEX(g_communicationLock, pcs);
 		BinaryWriter writer;
 		writer.WriteByte(MessageType_RemoveDirectory);
 		writer.WriteStringKey(pathNameKey);
@@ -1762,7 +1762,7 @@ UBA_EXPORT int UBA_WRAPPER(rename)(const char* oldpath, const char* newpath)
 	bool result;
 	{
 		TimerScope ts(g_stats.moveFile);
-		SCOPED_WRITE_LOCK(g_communicationLock, pcs);
+		SCOPED_FUTEX(g_communicationLock, pcs);
 		BinaryWriter writer;
 		writer.WriteByte(MessageType_MoveFile);
 		writer.WriteStringKey(oldKey);
@@ -1826,7 +1826,7 @@ UBA_EXPORT int UBA_WRAPPER(chmod)(const char* pathname, mode_t mode)
 	u32 errorCode;
 	{
 		TimerScope ts(g_stats.chmod);
-		SCOPED_WRITE_LOCK(g_communicationLock, pcs);
+		SCOPED_FUTEX(g_communicationLock, pcs);
 		BinaryWriter writer;
 		writer.WriteByte(MessageType_Chmod);
 		writer.WriteStringKey(key);
@@ -1928,7 +1928,7 @@ int Shared_DeleteFile(const char* funcName, const char* pathname)
 	{
 		u32 closeId = 0;
 		TimerScope ts(g_stats.deleteFile);
-		SCOPED_WRITE_LOCK(g_communicationLock, pcs);
+		SCOPED_FUTEX(g_communicationLock, pcs);
 		BinaryWriter writer;
 		writer.WriteByte(MessageType_DeleteFile);
 		writer.WriteString(fixedName);
@@ -2027,7 +2027,7 @@ bool ExecuteHostRun(StringBufferBase& out, const char* const* argv, bool removeL
 		DEBUG_LOG_DETOURED("HostRun", "%s", command.data)
 	}
 	TimerScope ts(g_stats.getFullFileName);
-	SCOPED_WRITE_LOCK(g_communicationLock, pcs);
+	SCOPED_FUTEX(g_communicationLock, pcs);
 	BinaryWriter writer;
 	writer.WriteByte(MessageType_HostRun);
 	u16& size = *(u16*)writer.AllocWrite(2);
@@ -2187,7 +2187,7 @@ int shared_posix_spawn(pid_t* pid, const char* path, const posix_spawn_file_acti
 	
 	{
 		TimerScope ts(g_stats.createProcess);
-		SCOPED_WRITE_LOCK(g_communicationLock, pcs);
+		SCOPED_FUTEX(g_communicationLock, pcs);
 
 		const char* pwd = "";
 		//while (char* env = envp[i++])
@@ -2277,7 +2277,7 @@ int shared_posix_spawn(pid_t* pid, const char* path, const posix_spawn_file_acti
 
 	{
 		TimerScope ts(g_stats.createProcess);
-		SCOPED_WRITE_LOCK(g_communicationLock, pcs);
+		SCOPED_FUTEX(g_communicationLock, pcs);
 		BinaryWriter writer;
 		writer.WriteByte(MessageType_StartProcess);
 		writer.WriteU32(processId);
@@ -2495,7 +2495,7 @@ int Internal_execve(const char* pathname, char* const _Nullable argv[], char* co
 
 	{
 		TimerScope ts(g_stats.createProcess);
-		SCOPED_WRITE_LOCK(g_communicationLock, pcs);
+		SCOPED_FUTEX(g_communicationLock, pcs);
 		BinaryWriter writer;
 		writer.WriteByte(MessageType_ExitChildProcess);
 		writer.WriteU32(pid);
@@ -2866,7 +2866,7 @@ namespace uba
 
 		{
 			TimerScope ts(g_stats.init);
-			SCOPED_WRITE_LOCK(g_communicationLock, pcs);
+			SCOPED_FUTEX(g_communicationLock, pcs);
 			BinaryWriter writer;
 			writer.WriteByte(MessageType_Init);
 			writer.Flush();
@@ -3102,7 +3102,7 @@ extern "C"
 		u32 overlayFilesSize = 0;
 		u32 directoryTableSize = 0;
 		{
-			SCOPED_WRITE_LOCK(g_communicationLock, pcs);
+			SCOPED_FUTEX(g_communicationLock, pcs);
 			BinaryWriter writer;
 			writer.WriteByte(MessageType_GetNextProcess);
 			writer.WriteU32(prevExitCode);
diff --git a/Engine/Source/Programs/UnrealBuildAccelerator/Detours/Private/UbaDetoursFileMappingTable.cpp b/Engine/Source/Programs/UnrealBuildAccelerator/Detours/Private/UbaDetoursFileMappingTable.cpp
--- a/Engine/Source/Programs/UnrealBuildAccelerator/Detours/Private/UbaDetoursFileMappingTable.cpp
+++ b/Engine/Source/Programs/UnrealBuildAccelerator/Detours/Private/UbaDetoursFileMappingTable.cpp
@@ -239,7 +239,7 @@ namespace uba
 
 	void Rpc_UpdateTables()
 	{
-		SCOPED_WRITE_LOCK(g_communicationLock, pcs);
+		SCOPED_FUTEX(g_communicationLock, pcs);
 		u32 directoryTableSize;
 		u32 fileMappingTableSize;
 		bool done;
@@ -265,7 +265,7 @@ namespace uba
 
 	void Rpc_GetWrittenFiles()
 	{
-		SCOPED_WRITE_LOCK(g_communicationLock, pcs);
+		SCOPED_FUTEX(g_communicationLock, pcs);
 		Rpc_GetWrittenFilesNoLock(true);
 	}
 
diff --git a/Engine/Source/Programs/UnrealBuildAccelerator/Detours/Private/UbaDetoursFileMappingTable.h b/Engine/Source/Programs/UnrealBuildAccelerator/Detours/Private/UbaDetoursFileMappingTable.h
--- a/Engine/Source/Programs/UnrealBuildAccelerator/Detours/Private/UbaDetoursFileMappingTable.h
+++ b/Engine/Source/Programs/UnrealBuildAccelerator/Detours/Private/UbaDetoursFileMappingTable.h
@@ -85,7 +85,7 @@ namespace uba
 		u32 m_memPosition = 0;
 		ReaderWriterLock m_lookupLock;
 		GrowingUnorderedMap<StringKey, FileInfo> m_lookup;
-		ReaderWriterLock m_memLookupLock;
+		Futex m_memLookupLock;
 		struct Entry { int refCount = 0; DetouredHandle* handle = nullptr; };
 		UnorderedMap<const void*, Entry> m_memLookup;
 	};
diff --git a/Engine/Source/Programs/UnrealBuildAccelerator/Detours/Private/UbaDetoursShared.cpp b/Engine/Source/Programs/UnrealBuildAccelerator/Detours/Private/UbaDetoursShared.cpp
--- a/Engine/Source/Programs/UnrealBuildAccelerator/Detours/Private/UbaDetoursShared.cpp
+++ b/Engine/Source/Programs/UnrealBuildAccelerator/Detours/Private/UbaDetoursShared.cpp
@@ -15,7 +15,7 @@ namespace uba
 	VARIABLE_MEM(StringBuffer<512>, g_virtualApplicationDir);
 	VARIABLE_MEM(ProcessStats, g_stats);
 	VARIABLE_MEM(KernelStats, g_kernelStats);
-	VARIABLE_MEM(ReaderWriterLock, g_communicationLock);
+	VARIABLE_MEM(Futex, g_communicationLock);
 	VARIABLE_MEM(StringBuffer<256>, g_logName);
 	VARIABLE_MEM(StringBuffer<512>, g_virtualWorkingDir);
 	VARIABLE_MEM(StringBuffer<256>, g_exeDir);
diff --git a/Engine/Source/Programs/UnrealBuildAccelerator/Detours/Private/UbaDetoursShared.h b/Engine/Source/Programs/UnrealBuildAccelerator/Detours/Private/UbaDetoursShared.h
--- a/Engine/Source/Programs/UnrealBuildAccelerator/Detours/Private/UbaDetoursShared.h
+++ b/Engine/Source/Programs/UnrealBuildAccelerator/Detours/Private/UbaDetoursShared.h
@@ -64,7 +64,7 @@ namespace uba
 	extern ProcessStats& g_stats;
 	extern KernelStats& g_kernelStats;
 	extern bool g_conEnabled[2]; // 0=stderr, 1=stdout
-	extern ReaderWriterLock& g_communicationLock;
+	extern Futex& g_communicationLock;
 	extern MemoryBlock& g_memoryBlock;
 	extern DirectoryTable& g_directoryTable;
 	extern MappedFileTable& g_mappedFileTable;
@@ -145,7 +145,7 @@ namespace uba
 	#define RPC_MESSAGE(messageName, timerName) \
 		DEBUG_LOG(TC("RPC_MESSAGE %s"), TC(#messageName)); \
 		TimerScope ts(g_stats.timerName); \
-		SCOPED_WRITE_LOCK(g_communicationLock, pcs); \
+		SCOPED_FUTEX(g_communicationLock, pcs); \
 		BinaryWriter writer; \
 		writer.WriteByte(MessageType_##messageName);
 
diff --git a/Engine/Source/Programs/UnrealBuildAccelerator/Detours/Private/Windows/UbaDetoursFunctionsKernelBase.inl b/Engine/Source/Programs/UnrealBuildAccelerator/Detours/Private/Windows/UbaDetoursFunctionsKernelBase.inl
--- a/Engine/Source/Programs/UnrealBuildAccelerator/Detours/Private/Windows/UbaDetoursFunctionsKernelBase.inl
+++ b/Engine/Source/Programs/UnrealBuildAccelerator/Detours/Private/Windows/UbaDetoursFunctionsKernelBase.inl
@@ -2594,7 +2594,7 @@ LPVOID Detoured_MapViewOfFileEx(HANDLE hFileMappingObject, DWORD dwDesiredAccess
 
 			DEBUG_LOG_DETOURED(L"MapViewOfFileEx", L"%llu Size %llu (%s) -> 0x%llx", uintptr_t(hFileMappingObject), dwNumberOfBytesToMap, HandleToName(hFileMappingObject), u64(mem));
 
-			SCOPED_WRITE_LOCK(g_mappedFileTable.m_memLookupLock, lock);
+			SCOPED_FUTEX(g_mappedFileTable.m_memLookupLock, lock);
 			auto& entry = g_mappedFileTable.m_memLookup[mem];
 			if (!entry.handle)
 			{
@@ -2628,7 +2628,7 @@ BOOL Detoured_UnmapViewOfFileEx(PVOID lpBaseAddress, ULONG UnmapFlags)
 	DETOURED_CALL(UnmapViewOfFileEx);
 
 	{
-		SCOPED_WRITE_LOCK(g_mappedFileTable.m_memLookupLock, lock);
+		SCOPED_FUTEX(g_mappedFileTable.m_memLookupLock, lock);
 		auto it = g_mappedFileTable.m_memLookup.find(lpBaseAddress);
 		if (it != g_mappedFileTable.m_memLookup.end())
 		{
diff --git a/Engine/Source/Programs/UnrealBuildAccelerator/Detours/Private/Windows/UbaDetoursFunctionsRpcrt4.inl b/Engine/Source/Programs/UnrealBuildAccelerator/Detours/Private/Windows/UbaDetoursFunctionsRpcrt4.inl
--- a/Engine/Source/Programs/UnrealBuildAccelerator/Detours/Private/Windows/UbaDetoursFunctionsRpcrt4.inl
+++ b/Engine/Source/Programs/UnrealBuildAccelerator/Detours/Private/Windows/UbaDetoursFunctionsRpcrt4.inl
@@ -17,7 +17,7 @@ RPC_STATUS Detoured_RpcBindingFromStringBindingW(RPC_WSTR StringBinding, RPC_BIN
 	/*
 	if (g_runningRemote && Contains((wchar_t*)StringBinding, L"mspdb"))
 	{
-		SCOPED_WRITE_LOCK(g_communicationLock, pcs);
+		SCOPED_FUTEX(g_communicationLock, pcs);
 		BinaryWriter writer;
 		writer.WriteByte(MessageType_RpcCommunication);
 		writer.WriteByte(1);
