[UBA]
* Implemented UbaRequestNextProcess2 for posix
* Enabled process reuse tests for mac

diff --git a/Engine/Source/Programs/UnrealBuildAccelerator/Detours/Private/PosixOS/UbaDetoursFunctionsPosix.cpp b/Engine/Source/Programs/UnrealBuildAccelerator/Detours/Private/PosixOS/UbaDetoursFunctionsPosix.cpp
--- a/Engine/Source/Programs/UnrealBuildAccelerator/Detours/Private/PosixOS/UbaDetoursFunctionsPosix.cpp
+++ b/Engine/Source/Programs/UnrealBuildAccelerator/Detours/Private/PosixOS/UbaDetoursFunctionsPosix.cpp
@@ -3083,7 +3083,7 @@ namespace uba
 
 extern "C"
 {
-	UBA_EXPORT bool UbaRequestNextProcess(u32 prevExitCode, char* outArguments, u32 outArgumentsCapacity)
+	UBA_EXPORT bool UbaRequestNextProcess2(u32 prevExitCode, char* outArguments, u32 outArgumentsCapacity, u32 timeOutMs, bool* outShouldExit)
 	{
 		#if UBA_DEBUG_LOG_ENABLED
 		FlushDebugLog();
@@ -3106,12 +3106,14 @@ extern "C"
 			BinaryWriter writer;
 			writer.WriteByte(MessageType_GetNextProcess);
 			writer.WriteU32(prevExitCode);
+			writer.WriteU32(timeOutMs);
 			g_stats.Write(writer);
 			g_kernelStats.Write(writer);
 
 			writer.Flush();
 			BinaryReader reader;
 			newProcess = reader.ReadBool();
+			*outShouldExit = reader.ReadBool();
 			if (newProcess)
 			{
 				reader.ReadString(outArguments, outArgumentsCapacity);
@@ -3123,11 +3125,11 @@ extern "C"
 			}
 		}
 
+		g_kernelStats = {};
+		g_stats = {};
+
 		if (newProcess)
 		{
-			g_kernelStats = {};
-			g_stats = {};
-
 			#if UBA_DEBUG_LOG_ENABLED
 			SuppressDetourScope scope;
 			int debugFile = g_debugFile;
@@ -3151,4 +3153,10 @@ extern "C"
 		Rpc_UpdateTables();
 		return newProcess;
 	}
+
+	UBA_EXPORT bool UbaRequestNextProcess(u32 prevExitCode, char* outArguments, u32 outArgumentsCapacity)
+	{
+		bool shouldExit;
+		return UbaRequestNextProcess2(prevExitCode, outArguments, outArgumentsCapacity, 0, &shouldExit);
+	}
 }
