[UBA]
* Fixed so visualizer for mac save/load config
* Added context menu

diff --git a/Engine/Source/Programs/UnrealBuildAccelerator/Visualizer/Private/Mac/UbaVisualizerApp.mm b/Engine/Source/Programs/UnrealBuildAccelerator/Visualizer/Private/Mac/UbaVisualizerApp.mm
--- a/Engine/Source/Programs/UnrealBuildAccelerator/Visualizer/Private/Mac/UbaVisualizerApp.mm
+++ b/Engine/Source/Programs/UnrealBuildAccelerator/Visualizer/Private/Mac/UbaVisualizerApp.mm
@@ -21,6 +21,30 @@ - (void)stopTimer;
 - (void)tick:(NSTimer *)timer;
 @end
 
+enum
+{
+	Popup_CopySessionInfo = 3,
+	Popup_CopyProcessInfo,
+	Popup_CopyProcessLog,
+	Popup_CopyProcessBreadcrumbs,
+	Popup_CopyWorkInfo,
+	Popup_Replay,
+	Popup_Pause,
+	Popup_Play,
+	Popup_JumpToEnd,
+
+	#define UBA_VISUALIZER_FLAG(name, defaultValue, desc) Popup_##name,
+	UBA_VISUALIZER_FLAGS2
+	#undef UBA_VISUALIZER_FLAG
+
+	Popup_IncreaseFontSize,
+	Popup_DecreaseFontSize,
+
+	Popup_SaveAs,
+	Popup_SaveSettings,
+	Popup_OpenSettings,
+	Popup_Quit,
+};
 
 namespace uba
 {
@@ -608,6 +632,52 @@ bool UpdateSelection() override
 		{
 			return Visualizer::UpdateSelection(((VisualizerBackendMac&)m_backend).m_mousePos);
 		}
+		void MenuClicked(int id)
+		{
+			switch (id)
+			{
+			case Popup_ShowProcessText:
+				m_config.ShowProcessText = !m_config.ShowProcessText;
+				Redraw(true);
+				break;
+			case Popup_ShowReadWriteColors:
+				m_config.ShowReadWriteColors = !m_config.ShowReadWriteColors;
+				DirtyBitmaps(false);
+				Redraw(true);
+				break;
+			case Popup_ScaleHorizontalWithScrollWheel:
+				m_config.ScaleHorizontalWithScrollWheel = !m_config.ScaleHorizontalWithScrollWheel;
+				break;
+			case Popup_ShowAllTraces:
+				m_config.ShowAllTraces = !m_config.ShowAllTraces;
+				break;
+			case Popup_SortActiveRemoteSessions:
+				m_config.SortActiveRemoteSessions = !m_config.SortActiveRemoteSessions;
+				Redraw(true);
+				break;
+			case Popup_AutoScaleHorizontal:
+				m_config.AutoScaleHorizontal = !m_config.AutoScaleHorizontal;
+				UpdateScrollbars(true);
+				Redraw(true);
+				break;
+			case Popup_LockTimelineToBottom:
+				m_config.LockTimelineToBottom = !m_config.LockTimelineToBottom;
+				Redraw(true);
+				break;
+			case Popup_DarkMode:
+			{
+				m_config.DarkMode = !m_config.DarkMode;
+				DirtyBitmaps(true);
+				//UpdateTheme();
+				Redraw(true);
+				break;
+			}
+			case Popup_AutoSaveSettings:
+				m_config.AutoSaveSettings = !m_config.AutoSaveSettings;
+				break;
+
+			}
+		}
 
 		void Paint(CGContextRef ctx, SIZE size)
 		{
@@ -683,6 +753,20 @@ void MouseDown(CGFloat x, CGFloat y)
 			}
 		}
 
+		void Scroll(float dx, float dy)
+		{
+			m_scrollPosY = Min(0.0f, m_scrollPosY + float(dy));
+			m_scrollPosX = Min(0.0f, m_scrollPosX + float(dx));
+			UpdateSelection();
+
+			int minScroll = m_clientWidth - m_contentWidth;
+			if (!m_traceView.finished && m_scrollPosX <= float(minScroll))
+				m_autoScroll = true;
+			else
+				m_autoScroll = false;
+			//if (_visualizer->m_scrollPosX > _visualizer->)
+		}
+
 		bool ThreadListen()
 		{
 			const tchar* channelName = TC("Default");
@@ -749,10 +833,13 @@ - (instancetype)initWithFrame:(NSRect)frame {
 	if (self)
 	{
 		using namespace uba;
-		VisualizerConfig config("~/UbaVisualizerConfig.toml");
+		char fixedPath[512];
+		GetFullPathNameW("~/.epic/UbaVisualizerConfig.toml", 512, fixedPath, nullptr);
+		VisualizerConfig config(fixedPath);
 		config.boxHeight = 14;
 		config.horizontalScaleValue = 1.0f;
 		config.DarkMode = true;
+		config.AutoScaleHorizontal = false;
 		auto writer = new MessageBoxLogWriter(self.window);
 
 		SetCustomAssertHandler([](const tchar* text, u32 textLen, void* userData)
@@ -763,6 +850,8 @@ - (instancetype)initWithFrame:(NSRect)frame {
 		auto logger = new LoggerWithWriter(*writer);
 		auto backend = new VisualizerBackendMac;
 
+		config.Load(*logger);
+
 		_visualizer = new VisualizerMac(config, *logger, *backend, self);
 		g_visualizerMac = _visualizer;
 	}
@@ -845,11 +934,7 @@ - (void)mouseDown:(NSEvent*)event
 - (void)scrollWheel:(NSEvent *)event
 {
 	using namespace uba;
-    CGFloat dx = event.scrollingDeltaX;
-    CGFloat dy = event.scrollingDeltaY;
-	_visualizer->m_scrollPosY = Min(0.0f, _visualizer->m_scrollPosY + float(dy));
-	_visualizer->m_scrollPosX = Min(0.0f, _visualizer->m_scrollPosX + float(dx));
-	_visualizer->UpdateSelection();
+	_visualizer->Scroll(float(event.scrollingDeltaX), float(event.scrollingDeltaY));
 	[self setNeedsDisplay:YES];
 }
 
@@ -894,6 +979,30 @@ - (void)magnifyWithEvent:(NSEvent *)event
     //CGFloat magnification = event.magnification;
 	//_visualizer->m_zoomValue += float(magnification);
 }
+
+- (NSMenu *)menuForEvent:(NSEvent *)event
+{
+    NSMenu *menu = [[NSMenu alloc] initWithTitle:@"Context"];
+	#define UBA_VISUALIZER_FLAG(name, defaultValue, desc) \
+		{ NSMenuItem* item = [menu addItemWithTitle:@desc action:@selector(onMenuItem:) keyEquivalent:@""]; \
+			item.tag = Popup_##name; \
+			item.state = _visualizer->m_config.name ? NSControlStateValueOn : NSControlStateValueOff; \
+		}
+				UBA_VISUALIZER_FLAGS2
+	#undef UBA_VISUALIZER_FLAG
+
+	menu.autoenablesItems = NO;
+	return menu;
+}
+
+- (void)onMenuItem:(NSMenuItem*)sender {
+	_visualizer->MenuClicked(sender.tag);
+}
+
+- (void)onAction2:(id)sender {
+
+}
+
 @end
 
 @interface AppDelegate : NSObject <NSApplicationDelegate, NSWindowDelegate>
@@ -921,6 +1030,15 @@ - (BOOL)application:(NSApplication *)sender openFile:(NSString *)filename
 	visualizer.NewTrace(0, false);
     return YES;
 }
+
+- (NSApplicationTerminateReply)applicationShouldTerminate:(NSApplication *)sender
+{
+	using namespace uba;
+	auto& visualizer = *g_visualizerMac;
+	LoggerWithWriter nullLogger(g_nullLogWriter);
+	visualizer.m_config.Save(nullLogger);
+    return NSTerminateNow;
+}
 @end
 
 int main(int argc, const char * argv[]) {
