[UBA]
* Fixed so when running processes on mac we try to load dylib files from same dir as binary.

diff --git a/Engine/Source/Programs/UnrealBuildAccelerator/Common/Private/UbaProcess.cpp b/Engine/Source/Programs/UnrealBuildAccelerator/Common/Private/UbaProcess.cpp
--- a/Engine/Source/Programs/UnrealBuildAccelerator/Common/Private/UbaProcess.cpp
+++ b/Engine/Source/Programs/UnrealBuildAccelerator/Common/Private/UbaProcess.cpp
@@ -2222,6 +2222,8 @@ namespace uba
 				detoursVar.Append("LD_PRELOAD=").Append(detoursLib);
 #else
 				detoursVar.Append("DYLD_INSERT_LIBRARIES=").Append(detoursLib);
+				if (m_runningRemote)
+					ldLibraryPath.Append(TCV("DYLD_LIBRARY_PATH=")).AppendDir(realApplication); // This is needed in case @rpath is "wrong" so it wants to resolve dylibs outside of own directory
 #endif
 
 				processVar.Append("UBA_SESSION_PROCESS=").AppendValue(getpid());
