[UBA]
* Fixed so mac session server is not sending minOsVersion if client is linux since there is no point sending it

diff --git a/Engine/Source/Programs/UnrealBuildAccelerator/Common/Private/UbaSessionServer.cpp b/Engine/Source/Programs/UnrealBuildAccelerator/Common/Private/UbaSessionServer.cpp
--- a/Engine/Source/Programs/UnrealBuildAccelerator/Common/Private/UbaSessionServer.cpp
+++ b/Engine/Source/Programs/UnrealBuildAccelerator/Common/Private/UbaSessionServer.cpp
@@ -1279,7 +1279,8 @@ namespace uba
 			writer.WriteBool(m.isSystem);
 			writer.WriteCasKey(casKey);
 			#if PLATFORM_MAC
-			writer.WriteU32(m.minOsVersion);
+			if (!isClientLinux)
+				writer.WriteU32(m.minOsVersion);
 			#endif
 		}
 
