[UBA]
* Fixed so SetFileInformationByHandle with FileRenameInfoEx works on wine (used by some mac shader executable)
* Fixed so child processes are properly tracked in trace file for GetNext process path
* Fixed so if dll is a known system dll, then we directly force load it from system32

diff --git a/Engine/Source/Programs/UnrealBuildAccelerator/Common/Private/UbaProcess.cpp b/Engine/Source/Programs/UnrealBuildAccelerator/Common/Private/UbaProcess.cpp
--- a/Engine/Source/Programs/UnrealBuildAccelerator/Common/Private/UbaProcess.cpp
+++ b/Engine/Source/Programs/UnrealBuildAccelerator/Common/Private/UbaProcess.cpp
@@ -1526,6 +1526,28 @@ namespace uba
 		m_processStats.Add(processStats);
 		m_kernelStats.Add(kernelStats);
 
+		// Transfer child times and clear
+		WaitForChildrenExit();
+
+		for (auto& child : m_childProcesses)
+		{
+			auto& process = *(ProcessImpl*)child.m_process;
+			m_sessionStats.Add(process.m_sessionStats);
+			m_storageStats.Add(process.m_storageStats);
+			m_processStats.Add(process.m_processStats);
+			m_kernelStats.Add(process.m_kernelStats);
+			
+			processStats.Add(process.m_processStats);
+			kernelStats.Add(process.m_kernelStats);
+			
+			process.m_sessionStats = {};
+			process.m_storageStats = {};
+			process.m_processStats = {};
+			process.m_kernelStats = {};
+		}
+		m_childProcesses.clear();
+
+
 		processStats.Write(statsWriter);
 		if (m_runningRemote)
 			m_sessionStats.Write(statsWriter);
diff --git a/Engine/Source/Programs/UnrealBuildAccelerator/Core/Public/UbaDetoursPayload.h b/Engine/Source/Programs/UnrealBuildAccelerator/Core/Public/UbaDetoursPayload.h
--- a/Engine/Source/Programs/UnrealBuildAccelerator/Core/Public/UbaDetoursPayload.h
+++ b/Engine/Source/Programs/UnrealBuildAccelerator/Core/Public/UbaDetoursPayload.h
@@ -30,7 +30,7 @@ namespace uba
 		bool readIntermediateFilesCompressed = false;
 		bool reportAllExceptions = false;
 		int uiLanguage = 0;
-		StringBuffer<256> logFile;
+		StringBuffer<512> logFile;
 	};
 
 	static const _GUID DetoursPayloadGuid = __uuidof(DetoursPayload);
diff --git a/Engine/Source/Programs/UnrealBuildAccelerator/Detours/Private/Windows/UbaDetoursFunctionsKernelBase.inl b/Engine/Source/Programs/UnrealBuildAccelerator/Detours/Private/Windows/UbaDetoursFunctionsKernelBase.inl
--- a/Engine/Source/Programs/UnrealBuildAccelerator/Detours/Private/Windows/UbaDetoursFunctionsKernelBase.inl
+++ b/Engine/Source/Programs/UnrealBuildAccelerator/Detours/Private/Windows/UbaDetoursFunctionsKernelBase.inl
@@ -2367,6 +2367,18 @@ BOOL Detoured_SetFileInformationByHandle(HANDLE hFile, FILE_INFO_BY_HANDLE_CLASS
 
 	if (!isDetouredHandle(hFile))
 	{
+		// DOES NOT EXIST IN WINE 10.7
+		if (g_isRunningWine && FileInformationClass == FileRenameInfoEx)
+		{
+			auto& oldInfo = *(FILE_RENAME_INFO*)lpFileInformation;
+			UBA_ASSERT(oldInfo.Flags == 0x00000001); // FILE_RENAME_REPLACE_IF_EXISTS
+			FILE_RENAME_INFO newInfo = oldInfo;
+			newInfo.ReplaceIfExists = true;
+
+			DEBUG_LOG_TRUE(L"SetFileInformationByHandle", L"%llu (FileRenameInfoEx)", uintptr_t(hFile));
+			return True_SetFileInformationByHandle(hFile, FileRenameInfo, &newInfo, dwBufferSize);
+		}
+
 		DEBUG_LOG_TRUE(L"SetFileInformationByHandle", L"%llu (%u)", uintptr_t(hFile), FileInformationClass);
 		return True_SetFileInformationByHandle(hFile, FileInformationClass, lpFileInformation, dwBufferSize);
 	}
@@ -2924,10 +2936,12 @@ HMODULE Shared_LoadLibrary(LPCWSTR lpLibFileName, HANDLE hFile, DWORD dwFlags, b
 HMODULE Detoured_LoadLibraryExW(LPCWSTR lpLibFileName, HANDLE hFile, DWORD dwFlags)
 {
 	DETOURED_CALL(LoadLibraryExW);
-	DEBUG_LOG_DETOURED(L"LoadLibraryExW", L"(%ls)", lpLibFileName);
 
 	if (!g_rules->AllowLoadLibrary(lpLibFileName))
+	{
+		DEBUG_LOG_DETOURED(L"LoadLibraryExW", L"SKIPPED (%ls)", lpLibFileName);
 		return 0;
+	}
 
 	if (IsKnownSystemFile(lpLibFileName))
 	{
@@ -2936,9 +2950,11 @@ HMODULE Detoured_LoadLibraryExW(LPCWSTR lpLibFileName, HANDLE hFile, DWORD dwFla
 		if (auto lastSeparator = TStrrchr(fileName, PathSeparator))
 			fileName = lastSeparator + 1;
 		DEBUG_LOG_TRUE(L"LoadLibraryExW", L"%ls (%s)", lpLibFileName, fileName);
-		return True_LoadLibraryExW(fileName, hFile, dwFlags);
+		return True_LoadLibraryExW(fileName, hFile, LOAD_LIBRARY_SEARCH_SYSTEM32);
 	}
 
+	DEBUG_LOG_DETOURED(L"LoadLibraryExW", L"(%ls)", lpLibFileName);
+
 	bool mightExist = true;
 	if (HMODULE res = Shared_LoadLibrary(lpLibFileName, hFile, dwFlags, mightExist))
 		return res;
diff --git a/Engine/Source/Programs/UnrealBuildAccelerator/Detours/Private/Windows/UbaDetoursFunctionsNtDll.inl b/Engine/Source/Programs/UnrealBuildAccelerator/Detours/Private/Windows/UbaDetoursFunctionsNtDll.inl
--- a/Engine/Source/Programs/UnrealBuildAccelerator/Detours/Private/Windows/UbaDetoursFunctionsNtDll.inl
+++ b/Engine/Source/Programs/UnrealBuildAccelerator/Detours/Private/Windows/UbaDetoursFunctionsNtDll.inl
@@ -1471,7 +1471,7 @@ thread_local bool t_ntOpenFileDisallowed;
 NTSTATUS NTAPI Detoured_NtOpenFile(PHANDLE FileHandle, ACCESS_MASK DesiredAccess, POBJECT_ATTRIBUTES ObjectAttributes, PIO_STATUS_BLOCK IoStatusBlock, ULONG ShareAccess, ULONG OpenOptions)
 {
 	DETOURED_CALL(NtOpenFile);
-	if (t_ntOpenFileDisallowed)
+	if (t_ntOpenFileDisallowed && !IsKnownSystemFile(ObjectAttributes->ObjectName->Buffer))
 	{
 		DEBUG_LOG_DETOURED(L"NtOpenFile", L"(DISALLOWED)(%.*s) -> STATUS_OBJECT_NAME_NOT_FOUND", ObjectAttributes->ObjectName->Length / 2, ObjectAttributes->ObjectName->Buffer);
 		return STATUS_OBJECT_NAME_NOT_FOUND;
