[UBA]
* Proper fix for FileRenameInfoEx

diff --git a/Engine/Source/Programs/UnrealBuildAccelerator/Detours/Private/Windows/UbaDetoursFunctionsKernelBase.inl b/Engine/Source/Programs/UnrealBuildAccelerator/Detours/Private/Windows/UbaDetoursFunctionsKernelBase.inl
--- a/Engine/Source/Programs/UnrealBuildAccelerator/Detours/Private/Windows/UbaDetoursFunctionsKernelBase.inl
+++ b/Engine/Source/Programs/UnrealBuildAccelerator/Detours/Private/Windows/UbaDetoursFunctionsKernelBase.inl
@@ -2370,13 +2370,15 @@ BOOL Detoured_SetFileInformationByHandle(HANDLE hFile, FILE_INFO_BY_HANDLE_CLASS
 		// DOES NOT EXIST IN WINE 10.7
 		if (g_isRunningWine && FileInformationClass == FileRenameInfoEx)
 		{
-			auto& oldInfo = *(FILE_RENAME_INFO*)lpFileInformation;
-			UBA_ASSERT(oldInfo.Flags == 0x00000001); // FILE_RENAME_REPLACE_IF_EXISTS
-			FILE_RENAME_INFO newInfo = oldInfo;
-			newInfo.ReplaceIfExists = true;
+			auto& info = *(FILE_RENAME_INFO*)lpFileInformation;
+			DWORD flags = info.Flags;
+			UBA_ASSERT(flags == 0x00000001); // FILE_RENAME_REPLACE_IF_EXISTS
+			info.ReplaceIfExists = true;
 
 			DEBUG_LOG_TRUE(L"SetFileInformationByHandle", L"%llu (FileRenameInfoEx)", uintptr_t(hFile));
-			return True_SetFileInformationByHandle(hFile, FileRenameInfo, &newInfo, dwBufferSize);
+			bool res = True_SetFileInformationByHandle(hFile, FileRenameInfo, lpFileInformation, dwBufferSize);
+			info.Flags = flags;
+			return res;
 		}
 
 		DEBUG_LOG_TRUE(L"SetFileInformationByHandle", L"%llu (%u)", uintptr_t(hFile), FileInformationClass);
