[UBA]
* Added Scheduler_GetProcessWeightThatCanRunRemotelyNow2 which provides number of processes that can run cross architecture and cross platform
* Hardening UbaVisualizer when using -host

diff --git a/Engine/Source/Programs/UnrealBuildAccelerator/Common/Private/UbaScheduler.cpp b/Engine/Source/Programs/UnrealBuildAccelerator/Common/Private/UbaScheduler.cpp
--- a/Engine/Source/Programs/UnrealBuildAccelerator/Common/Private/UbaScheduler.cpp
+++ b/Engine/Source/Programs/UnrealBuildAccelerator/Common/Private/UbaScheduler.cpp
@@ -481,12 +481,19 @@ namespace uba
 		return count;
 	}
 
-	float Scheduler::GetProcessWeightThatCanRunRemotelyNow()
+	void Scheduler::GetProcessWeightThatCanRunRemotelyNow(float& outTotal, float& outCrossArchitecture, float& outCrossPlatform)
 	{
+		outTotal = 0;
+		outCrossArchitecture = 0;
+		outCrossPlatform = 0;
+
 		if (m_session.IsRemoteExecutionDisabled())
-			return 0;
+			return;
+
+		float totalWeight = 0;
+		float crossArchitecture = 0;
+		float crossPlatform = 0;
 
-		float weight = 0;
 		SCOPED_FUTEX_READ(m_processEntriesLock, lock);
 		for (auto& entry : m_processEntries)
 		{
@@ -505,9 +512,17 @@ namespace uba
 			}
 			if (!canRun)
 				continue;
-			weight += entry.info->weight;
+			float weight = entry.info->weight;
+			totalWeight += weight;
+			if (entry.canCrossArchitecture)
+				crossArchitecture += weight;
+			if (entry.canCrossPlatform)
+				crossPlatform += weight;
 		}
-		return weight;
+
+		outTotal = totalWeight;
+		outCrossArchitecture = crossArchitecture;
+		outCrossPlatform = crossPlatform;
 	}
 
 	void Scheduler::ThreadLoop()
@@ -882,6 +897,9 @@ namespace uba
 
 			UBA_ASSERT(cacheStatus != CacheStatus_Success);
 
+			// TODO: If local it should first pick !canExecuteRemotely
+			// If not isCross, it should pick !canCross
+
 			if (isLocal)
 			{
 				if (m_forceRemote && entry.canExecuteRemotely)
diff --git a/Engine/Source/Programs/UnrealBuildAccelerator/Common/Public/UbaScheduler.h b/Engine/Source/Programs/UnrealBuildAccelerator/Common/Public/UbaScheduler.h
--- a/Engine/Source/Programs/UnrealBuildAccelerator/Common/Public/UbaScheduler.h
+++ b/Engine/Source/Programs/UnrealBuildAccelerator/Common/Public/UbaScheduler.h
@@ -102,7 +102,7 @@ namespace uba
 		SessionServer& GetSession() { return m_session; }
 		
 		u32 GetProcessCountThatCanRunRemotelyNow();
-		float GetProcessWeightThatCanRunRemotelyNow();
+		void GetProcessWeightThatCanRunRemotelyNow(float& outTotal, float& outCrossArchitecture, float& outCrossPlatform);
 
 	private:
 		struct CacheFetchInfo;
diff --git a/Engine/Source/Programs/UnrealBuildAccelerator/Host/Private/UbaExports.cpp b/Engine/Source/Programs/UnrealBuildAccelerator/Host/Private/UbaExports.cpp
--- a/Engine/Source/Programs/UnrealBuildAccelerator/Host/Private/UbaExports.cpp
+++ b/Engine/Source/Programs/UnrealBuildAccelerator/Host/Private/UbaExports.cpp
@@ -879,7 +879,16 @@ extern "C"
 
 	float Scheduler_GetProcessWeightThatCanRunRemotelyNow(uba::Scheduler* scheduler)
 	{
-		return scheduler->GetProcessWeightThatCanRunRemotelyNow();
+		float total;
+		float crossArchitecture;
+		float crossPlatform;
+		scheduler->GetProcessWeightThatCanRunRemotelyNow(total, crossArchitecture, crossPlatform);
+		return total;
+	}
+
+	void Scheduler_GetProcessWeightThatCanRunRemotelyNow2(uba::Scheduler* scheduler, float& outTotal, float& outCrossArchitecture, float& outCrossPlatform)
+	{
+		scheduler->GetProcessWeightThatCanRunRemotelyNow(outTotal, outCrossArchitecture, outCrossPlatform);
 	}
 
 	void Scheduler_SetAllowDisableRemoteExecution(uba::Scheduler* scheduler, bool allow)
diff --git a/Engine/Source/Programs/UnrealBuildAccelerator/Host/Public/UbaExports.h b/Engine/Source/Programs/UnrealBuildAccelerator/Host/Public/UbaExports.h
--- a/Engine/Source/Programs/UnrealBuildAccelerator/Host/Public/UbaExports.h
+++ b/Engine/Source/Programs/UnrealBuildAccelerator/Host/Public/UbaExports.h
@@ -173,6 +173,7 @@ extern "C"
 	UBA_API bool Scheduler_IsEmpty(uba::Scheduler* scheduler);
 	UBA_API void Scheduler_SetProcessFinishedCallback(uba::Scheduler* scheduler);
 	UBA_API float Scheduler_GetProcessWeightThatCanRunRemotelyNow(uba::Scheduler* scheduler);
+	UBA_API void Scheduler_GetProcessWeightThatCanRunRemotelyNow2(uba::Scheduler* scheduler, float& outTotal, float& outCrossArchitecture, float& outCrossPlatform);
 	UBA_API void Scheduler_SetAllowDisableRemoteExecution(uba::Scheduler* scheduler, bool allow);
 
 	// Cache
diff --git a/Engine/Source/Programs/UnrealBuildAccelerator/Visualizer/Private/UbaTraceReader.cpp b/Engine/Source/Programs/UnrealBuildAccelerator/Visualizer/Private/UbaTraceReader.cpp
--- a/Engine/Source/Programs/UnrealBuildAccelerator/Visualizer/Private/UbaTraceReader.cpp
+++ b/Engine/Source/Programs/UnrealBuildAccelerator/Visualizer/Private/UbaTraceReader.cpp
@@ -240,6 +240,9 @@ namespace uba
 
 	bool TraceReader::UpdateReceiveClient(NetworkClient& client)
 	{
+		if (!m_memoryBegin)
+			return true;
+
 		while (true)
 		{
 			u32 pos = u32(m_memoryEnd - m_memoryBegin);
diff --git a/Engine/Source/Programs/UnrealBuildAccelerator/Visualizer/Private/Windows/UbaVisualizerWin.cpp b/Engine/Source/Programs/UnrealBuildAccelerator/Visualizer/Private/Windows/UbaVisualizerWin.cpp
--- a/Engine/Source/Programs/UnrealBuildAccelerator/Visualizer/Private/Windows/UbaVisualizerWin.cpp
+++ b/Engine/Source/Programs/UnrealBuildAccelerator/Visualizer/Private/Windows/UbaVisualizerWin.cpp
@@ -215,6 +215,8 @@ namespace uba
 				m_trace.UpdateReceiveClient(*activeClient);
 			}
 
+			activeClient->Disconnect();
+
 			while (activeClient->IsConnected())
 			{
 				Sleep(500);
