[UBA]
* Fixed cross architecture/platform logic from scheduler to server session. things were not properly propagated

diff --git a/Engine/Source/Programs/UnrealBuildAccelerator/Common/Private/UbaScheduler.cpp b/Engine/Source/Programs/UnrealBuildAccelerator/Common/Private/UbaScheduler.cpp
--- a/Engine/Source/Programs/UnrealBuildAccelerator/Common/Private/UbaScheduler.cpp
+++ b/Engine/Source/Programs/UnrealBuildAccelerator/Common/Private/UbaScheduler.cpp
@@ -1139,7 +1139,7 @@ namespace uba
 			if (isLocal)
 				m_session.RunProcess(si, true, canDetour);
 			else
-				m_session.RunProcessRemote(si, info->weight, info->knownInputs, info->knownInputsCount);
+				m_session.RunProcessRemote(si, info->weight, info->knownInputs, info->knownInputsCount, processEntry.canCrossArchitecture, processEntry.canCrossPlatform);
 			return true;
 		}
 	}
diff --git a/Engine/Source/Programs/UnrealBuildAccelerator/Common/Private/UbaSessionServer.cpp b/Engine/Source/Programs/UnrealBuildAccelerator/Common/Private/UbaSessionServer.cpp
--- a/Engine/Source/Programs/UnrealBuildAccelerator/Common/Private/UbaSessionServer.cpp
+++ b/Engine/Source/Programs/UnrealBuildAccelerator/Common/Private/UbaSessionServer.cpp
@@ -123,7 +123,8 @@ namespace uba
 		Vector<u8> m_trackedInputs;
 		Vector<u8> m_trackedOutputs;
 		bool m_cancelled = false;
-		bool m_allowCrossArchitecture = false;
+		bool m_canCrossArchitecture = false;
+		bool m_canCrossPlatform = false;
 		bool m_isIdle = false;
 		u32 m_clientId = ~0u;
 		u32 m_sessionId = 0;
@@ -451,7 +452,7 @@ namespace uba
 		return racingProcess;
 	}
 
-	ProcessHandle SessionServer::RunProcessRemote(const ProcessStartInfo& startInfo, float weight, const void* knownInputs, u32 knownInputsCount, bool allowCrossArchitecture)
+	ProcessHandle SessionServer::RunProcessRemote(const ProcessStartInfo& startInfo, float weight, const void* knownInputs, u32 knownInputsCount, bool canCrossArchitecture, bool canCrossPlatform)
 	{
 		//TrackWorkScope tws(m_trace, AsView(TC("RunProcessRemote")), ColorWork);
 
@@ -467,11 +468,12 @@ namespace uba
 		#if UBA_TEST_LINUX_HELPING_WINDOWS
 		remoteProcess->m_startInfo.applicationStr = TC("d:\\temp\\dummy");
 		remoteProcess->m_startInfo.application = remoteProcess->m_startInfo.applicationStr.c_str();
-		allowCrossArchitecture = true;
+		m_canCrossPlatform = true;
 		#endif
 
 		remoteProcess->m_startInfo.rules = rules;
-		remoteProcess->m_allowCrossArchitecture = allowCrossArchitecture;
+		remoteProcess->m_canCrossArchitecture = canCrossArchitecture;
+		remoteProcess->m_canCrossPlatform = canCrossPlatform;
 
 		if (knownInputsCount)
 		{
@@ -1667,8 +1669,8 @@ namespace uba
 		auto processesLostRace = std::move(session.processesLostRace);
 		sessionsLock.Leave();
 
-		bool isCrossPlatform = IsLinux != session.isLinux;
-		bool isCrossArchitecture = IsArmBinary != session.isArm;
+		bool isClientCrossPlatform = IsLinux != session.isLinux;
+		bool isClientCrossArchitecture = IsArmBinary != session.isArm;
 
 		u32 weight32 = reader.ReadU32();
 		float availableWeight = *(float*)&weight32;
@@ -1691,7 +1693,7 @@ namespace uba
 
 			StringBuffer<> applicationOverride;
 
-			if (isCrossArchitecture || isCrossPlatform)
+			if (isClientCrossArchitecture || isClientCrossPlatform)
 			{
 				auto returnProcess = [&]()
 				{
@@ -1700,7 +1702,7 @@ namespace uba
 						m_remoteProcessReturnedEvent(*process);
 				};
 			
-				if (!process->m_allowCrossArchitecture)
+				if ((!process->m_canCrossArchitecture && isClientCrossArchitecture) || (!process->m_canCrossPlatform && isClientCrossPlatform))
 				{
 					// There is a risk cross architecture client dequeues another client's process that is not cross architecture
 					// and in that case we just have to return it and try again later.. (break out)
diff --git a/Engine/Source/Programs/UnrealBuildAccelerator/Common/Public/UbaSessionServer.h b/Engine/Source/Programs/UnrealBuildAccelerator/Common/Public/UbaSessionServer.h
--- a/Engine/Source/Programs/UnrealBuildAccelerator/Common/Public/UbaSessionServer.h
+++ b/Engine/Source/Programs/UnrealBuildAccelerator/Common/Public/UbaSessionServer.h
@@ -25,7 +25,7 @@ namespace uba
 		// weight is the expected core usage of the process. If processes are multithreaded it makes sense to increase weight. As an example, in UnrealBuildTool we see cl.exe as 1.5 and clang.exe as 1.0
 		// knownInputs is a memory block with null terminated tchar strings followed by an empty null terminated string to end. knownInputsCount is the number of strings in the memory block
 		// strings should be absolute or relative to working dir.
-		ProcessHandle RunProcessRemote(const ProcessStartInfo& startInfo, float weight = 1.0f, const void* knownInputs = nullptr, u32 knownInputsCount = 0, bool allowCrossArchitecture = false);
+		ProcessHandle RunProcessRemote(const ProcessStartInfo& startInfo, float weight = 1.0f, const void* knownInputs = nullptr, u32 knownInputsCount = 0, bool canCrossArchitecture = false, bool canCrossPlatform = false);
 
 		// Will kick off a local process with the same startInfo as the one provided to start the process with id matching raceAgainstRemoteProcessId
 		// This can be useful if there are free local cores and we know local machine is faster or network connection to remote is slow
