[UBA] Adding extra system libraries needed by newer LLVM clang to the Uba ninary parser for linux.

diff --git a/Engine/Source/Programs/UnrealBuildAccelerator/Core/Public/UbaLinuxBinDependencyParser.h b/Engine/Source/Programs/UnrealBuildAccelerator/Core/Public/UbaLinuxBinDependencyParser.h
--- a/Engine/Source/Programs/UnrealBuildAccelerator/Core/Public/UbaLinuxBinDependencyParser.h
+++ b/Engine/Source/Programs/UnrealBuildAccelerator/Core/Public/UbaLinuxBinDependencyParser.h
@@ -54,6 +54,9 @@ namespace uba
 		TCV("libXrandr.so"),
 		TCV("libXss.so"),
 		TCV("libudev.so"),
+		TCV("libz.so"),
+		TCV("libzstd.so"),
+		TCV("libtinfo.so"),
 	};
 
 	inline bool IsKnownSystemFile(const tchar* fileName)
