[UBA]
* The math is not mathing. Fixed a bug where m_activeLocalProcessWeight could be increased but then process is not popped up the queue.. which means that weight will forever increase causing fewer and fewer processes running

diff --git a/Engine/Source/Programs/UnrealBuildAccelerator/Common/Private/UbaScheduler.cpp b/Engine/Source/Programs/UnrealBuildAccelerator/Common/Private/UbaScheduler.cpp
--- a/Engine/Source/Programs/UnrealBuildAccelerator/Common/Private/UbaScheduler.cpp
+++ b/Engine/Source/Programs/UnrealBuildAccelerator/Common/Private/UbaScheduler.cpp
@@ -951,17 +951,18 @@ namespace uba
 					return ~0u;
 
 				if (!forReuse)
-					m_activeLocalProcessWeight += entry.info->weight;
-			}
-
+				{
+					// If there are queued getnext, then we try to hand over running queued processes
+					// Setting the event will wake up the getnext thread.. and If the getnext thread is timing out we will retry
+					if (m_getNextQueueCount)
+					{
+						m_updateThreadWakeup.Set();
+						Sleep(1);
+						return ~0u;
+					}
 
-			// If there are queued getnext, then we try to hand over running queued processes
-			// Setting the event will wake up the getnext thread.. and If the getnext thread is timing out we will retry
-			if (isLocal && !forReuse && m_getNextQueueCount)
-			{
-				m_updateThreadWakeup.Set();
-				Sleep(1);
-				return ~0u;
+					m_activeLocalProcessWeight += entry.info->weight;
+				}
 			}
 
 			outPrevStatus = entry.runStatus;
