[UBA]
* Added feature "DeferRemoteCapableProcesses" to scheduler. If this feature is enabled the scheduler will strive to prioritize processes that are only capable of running with constraints. So local session will first pickup processes that can only run local, then processes that can run remote but not cross system, and last processes that can run cross system. Non-cross system remote sessions will pick non-cross system processes first

diff --git a/Engine/Source/Programs/UnrealBuildAccelerator/Common/Private/UbaScheduler.cpp b/Engine/Source/Programs/UnrealBuildAccelerator/Common/Private/UbaScheduler.cpp
--- a/Engine/Source/Programs/UnrealBuildAccelerator/Common/Private/UbaScheduler.cpp
+++ b/Engine/Source/Programs/UnrealBuildAccelerator/Common/Private/UbaScheduler.cpp
@@ -166,6 +166,7 @@ namespace uba
 			table->GetValueAsBool(memWatchdog, TC("MemWatchdog"));
 
 			table->GetValueAsBool(writeTrackedInputs, TC("WriteTrackedInputs"));
+			table->GetValueAsBool(deferRemoteCapableProcesses, TC("DeferRemoteCapableProcesses"));
 
 			u32 temp = 0;
 			if (table->GetValueAsU32(temp, TC("MemTracingLevel")))
@@ -191,6 +192,7 @@ namespace uba
 	,	m_forceNative(info.forceNative)
 	,	m_useThreadToExitRemoteProcess(info.useThreadToExitRemoteProcess)
 	,	m_writeTrackedInputs(info.writeTrackedInputs)
+	,	m_deferRemoteCapableProcesses(info.deferRemoteCapableProcesses)
 	,	m_processConfigs(info.processConfigs)
 	,	m_writeToCache(info.writeToCache && info.cacheClientCount)
 	{
@@ -826,9 +828,9 @@ namespace uba
 			if (cacheStatus == CacheStatus_QueuedForTest)
 			{
 				bool canRun = true;
-				for (u32 j=0, je=entry.dependencyCount; j!=je; ++j)
+				while (entry.dependencyCount)
 				{
-					auto depIndex = entry.dependencies[j];
+					auto depIndex = *entry.dependencies;
 					if (depIndex >= m_processEntries.size())
 					{
 						m_session.GetLogger().Error(TC("Found dependency on index %u but there are only %u processes registered"), depIndex, u32(m_processEntries.size()));
@@ -846,6 +848,10 @@ namespace uba
 						canRun = false;
 						break;
 					}
+				
+					// Swap with last and reduce dependency count
+					*entry.dependencies = entry.dependencies[entry.dependencyCount - 1];
+					--entry.dependencyCount;
 				}
 				if (!canRun)
 					continue;
@@ -871,6 +877,10 @@ namespace uba
 		auto processEntries = m_processEntries.data();
 		bool allFinished = true;
 
+		u32 firstValidProcess = ~0u;
+		u32 firstValidRemoteCapableProcess = ~0u;
+		u32 firstValidCrossSystemCapableProcess = ~0u;
+
 		for (u32 i=m_processEntriesRunStart, e=u32(m_processEntries.size()); i!=e; ++i)
 		{
 			auto& entry = processEntries[i];
@@ -897,9 +907,6 @@ namespace uba
 
 			UBA_ASSERT(cacheStatus != CacheStatus_Success);
 
-			// TODO: If local it should first pick !canExecuteRemotely
-			// If not isCross, it should pick !canCross
-
 			if (isLocal)
 			{
 				if (m_forceRemote && entry.canExecuteRemotely)
@@ -921,9 +928,9 @@ namespace uba
 			}
 
 			bool canRun = true;
-			for (u32 j=0, je=entry.dependencyCount; j!=je; ++j)
+			while (entry.dependencyCount)
 			{
-				auto depIndex = entry.dependencies[j];
+				auto depIndex = *entry.dependencies;
 				if (depIndex >= m_processEntries.size())
 				{
 					m_session.GetLogger().Error(TC("Found dependency on index %u but there are only %u processes registered"), depIndex, u32(m_processEntries.size()));
@@ -940,37 +947,96 @@ namespace uba
 					canRun = false;
 					break;
 				}
+
+				// Swap with last and reduce dependency count
+				*entry.dependencies = entry.dependencies[entry.dependencyCount - 1];
+				--entry.dependencyCount;
 			}
 
 			if (!canRun)
 				continue;
 
-			if (isLocal)
-			{
-				if (!LocalProcessOkToSpawn())
-					return ~0u;
+			// Process can run, now we just need to check if we _should_ run this process (we might want to prioritize later processes)
 
-				if (!forReuse)
+			if (!isLocal)
+			{
+				if (m_deferRemoteCapableProcesses && m_hadCrossSystemRequest && (entry.canCrossPlatform || entry.canCrossArchitecture))
 				{
-					// If there are queued getnext, then we try to hand over running queued processes
-					// Setting the event will wake up the getnext thread.. and If the getnext thread is timing out we will retry
-					if (m_getNextQueueCount)
-					{
-						m_updateThreadWakeup.Set();
-						Sleep(1);
-						return ~0u;
-					}
-
-					m_activeLocalProcessWeight += entry.info->weight;
+					if (firstValidCrossSystemCapableProcess == ~0u)
+						firstValidCrossSystemCapableProcess = i;
+					continue;
 				}
+
+				firstValidRemoteCapableProcess = i;
+				break;
+			}
+
+			// Only local processes below here
+
+			// Process can only run locally, use it
+			if (!entry.canExecuteRemotely)
+			{
+				firstValidProcess = i;
+				break;
+			}
+
+			// If we don't want to defer remote-capable actions or haven't had any remote helpers, use this process
+			if (!m_deferRemoteCapableProcesses || !m_hadRemoteRequest)
+			{
+				firstValidProcess = i;
+				break;
+			}
+
+			// Track cross system and normal remote processes separately
+			if (entry.canCrossPlatform || entry.canCrossArchitecture)
+			{
+				if (firstValidCrossSystemCapableProcess == ~0u)
+					firstValidCrossSystemCapableProcess = i;
+			}
+			else
+			{
+				if (firstValidRemoteCapableProcess == ~0u)
+					firstValidRemoteCapableProcess = i;
 			}
 
-			outPrevStatus = entry.runStatus;
-			entry.runStatus = RunStatus_Running;
-			return i;
+			continue;
+		}
+
+		if (firstValidProcess == ~0u)
+		{
+			if (firstValidRemoteCapableProcess != ~0u)
+				firstValidProcess = firstValidRemoteCapableProcess;
+			else if (firstValidCrossSystemCapableProcess != ~0u)
+				firstValidProcess = firstValidCrossSystemCapableProcess;
+			else
+				return ~0u;
+		}
+
+		auto& entry = processEntries[firstValidProcess];
+
+		if (isLocal)
+		{
+			if (!LocalProcessOkToSpawn())
+				return ~0u;
+
+			if (!forReuse)
+			{
+				// If there are queued getnext, then we try to hand over running queued processes
+				// Setting the event will wake up the getnext thread.. and If the getnext thread is timing out we will retry
+				if (m_getNextQueueCount)
+				{
+					m_updateThreadWakeup.Set();
+					Sleep(1);
+					return ~0u;
+				}
+
+				m_activeLocalProcessWeight += entry.info->weight;
+			}
 		}
 
-		return ~0u;
+		outPrevStatus = entry.runStatus;
+		entry.runStatus = RunStatus_Running;
+		return firstValidProcess;
 	}
 
 	bool Scheduler::RunCacheQuery()
@@ -1092,6 +1158,12 @@ namespace uba
 		{
 			RunStatus prevStatus;
 			SCOPED_FUTEX(m_processEntriesLock, lock);
+
+			if (!isLocal)
+				m_hadRemoteRequest = true;
+			if (isCrossArchitecture || isCrossPlatform)
+				m_hadCrossSystemRequest = true;
+
 			u32 indexToRun = PopProcess(isLocal, isCrossArchitecture, isCrossPlatform, prevStatus, false);
 			if (indexToRun == ~0u)
 			{
diff --git a/Engine/Source/Programs/UnrealBuildAccelerator/Common/Public/UbaScheduler.h b/Engine/Source/Programs/UnrealBuildAccelerator/Common/Public/UbaScheduler.h
--- a/Engine/Source/Programs/UnrealBuildAccelerator/Common/Public/UbaScheduler.h
+++ b/Engine/Source/Programs/UnrealBuildAccelerator/Common/Public/UbaScheduler.h
@@ -39,6 +39,7 @@ namespace uba
 		bool useThreadToExitRemoteProcess = false; // Set this to true to move exit process call from job to thread. 
 		bool memWatchdog = false; // Will watch memory and throttle/kill processes if needed
 		bool writeTrackedInputs = false; // Set to true to write out log files with tracked inputs
+		bool deferRemoteCapableProcesses = false; // Set to true to have local session first pick local-only processes
 		u32 maxRacingPercent = 0u; // If processor weights go below this percent, then we will start filling up with racing processes
 		u8 memTracingLevel = 0;
 		u8 memStartWaitPercent = 85; // When memory usage goes above this percent, no new processes will be spawned until back below
@@ -195,7 +196,10 @@ namespace uba
 		bool m_useThreadToExitRemoteProcess;
 		bool m_allowDisableRemoteExecution = false;
 		bool m_writeTrackedInputs = false;
+		bool m_deferRemoteCapableProcesses = false;
 		bool m_cancelled = false;
+		bool m_hadRemoteRequest = false;
+		bool m_hadCrossSystemRequest = false;
 		ConfigTable* m_processConfigs = nullptr;
 
 		Atomic<float> m_activeLocalProcessWeight = 0.0f;
