[UBA]
* Fixed detoured symlink for posix

diff --git a/Engine/Source/Programs/UnrealBuildAccelerator/Detours/Private/PosixOS/UbaDetoursFunctionsPosix.cpp b/Engine/Source/Programs/UnrealBuildAccelerator/Detours/Private/PosixOS/UbaDetoursFunctionsPosix.cpp
--- a/Engine/Source/Programs/UnrealBuildAccelerator/Detours/Private/PosixOS/UbaDetoursFunctionsPosix.cpp
+++ b/Engine/Source/Programs/UnrealBuildAccelerator/Detours/Private/PosixOS/UbaDetoursFunctionsPosix.cpp
@@ -1865,22 +1865,55 @@ UBA_EXPORT int UBA_WRAPPER(symlink)(const char* path1, const char* path2)
 {
 	UBA_INIT_DETOUR(symlink, path1, path2);
 
-	StringBuffer<> fixedPath1;
-	FixPath(fixedPath1, path1);
-	DevirtualizePath(fixedPath1);
-	StringBuffer<> fixedPath2;
-	FixPath(fixedPath2, path2);
-	DevirtualizePath(fixedPath2);
+	StringBuffer<> fromPath;
+	FixPath(fromPath, path1);
+	DevirtualizePath(fromPath);
+	StringBuffer<> toPath;
+	FixPath(toPath, path2);
+	DevirtualizePath(toPath);
 
-	if (!CanDetour2(fixedPath1) && !CanDetour2(fixedPath2))
+	if (!CanDetour2(fromPath) && !CanDetour2(toPath))
 	{
 		DEBUG_LOG_TRUE("symlink", "(from %s to %s)", path1, path2);
 		return TRUE_WRAPPER(symlink)(path1, path2);
 	}
 
-	UBA_ASSERTF(false, "symlink not implemented (from %s to %s)", path1, path2);
+	StringKey fromKey = ToStringKeyLower(fromPath);
+	StringKey toKey = ToStringKeyLower(toPath);
+
+	StringBuffer<> newFromName;
+	StringBuffer<> newToName;
+	u32 closeId;
+	u32 lastError;
+	u32 directoryTableSize;
+	{
+		RPC_MESSAGE(CopyFile, copyFile)
+		writer.WriteStringKey(fromKey);
+		writer.WriteString(fromPath);
+		writer.WriteStringKey(toKey);
+		writer.WriteString(toPath);
+		writer.Flush();
+		BinaryReader reader;
+		reader.ReadString(newFromName);
+		reader.ReadString(newToName);
+		closeId = reader.ReadU32();
+		lastError = reader.ReadU32();
+		directoryTableSize = reader.ReadU32();
+	}
+
+	UBA_ASSERT(closeId != ~0u);
+
+	//UBA_ASSERTF(false, "symlink not implemented (from %s to %s)", path1, path2);
 	DEBUG_LOG_DETOURED("symlink", "(from %s to %s)", path1, path2);
-	return TRUE_WRAPPER(symlink)(path1, path2);
+	int res = TRUE_WRAPPER(symlink)(newFromName.data, newToName.data);
+
+	if (!closeId)
+		return res;
+
+	bool deleteOnClose = res != 0; // If failing to copy we set deleteOnClose
+	Rpc_UpdateCloseHandle(newToName.data, closeId, deleteOnClose, "", {}, 0, true);
+
+	return res;
 }
 
 UBA_EXPORT ssize_t UBA_WRAPPER(pread)(int __fd, void * __buf, size_t __nbyte, off_t __offset)
