[UBA]
* Fixed so UbaObjTool release obj file mappings after use

diff --git a/Engine/Source/Programs/UnrealBuildAccelerator/Core/Private/UbaApplicationRules.cpp b/Engine/Source/Programs/UnrealBuildAccelerator/Core/Private/UbaApplicationRules.cpp
--- a/Engine/Source/Programs/UnrealBuildAccelerator/Core/Private/UbaApplicationRules.cpp
+++ b/Engine/Source/Programs/UnrealBuildAccelerator/Core/Private/UbaApplicationRules.cpp
@@ -731,6 +731,11 @@ namespace uba
 
 	class ApplicationRulesUbaObjTool : public ApplicationRules
 	{
+		virtual bool AllowDetach() const override
+		{
+			return true;
+		}
+
 		virtual bool IsOutputFile(StringView fileName, StringView systemTemp) const override
 		{
 			return fileName.EndsWith(TCV(".obj"))
@@ -751,6 +756,11 @@ namespace uba
 		{
 			return true;
 		}
+
+		virtual bool CloseFileMappingAfterUse(StringView fileName) const
+		{
+			return true;
+		}
 	};
 
 	class ApplicationRulesUbaTestApp : public ApplicationRules
