[UBA]
* Improved error message when failing to parse binary dependency

diff --git a/Engine/Source/Programs/UnrealBuildAccelerator/Common/Private/UbaSession.cpp b/Engine/Source/Programs/UnrealBuildAccelerator/Common/Private/UbaSession.cpp
--- a/Engine/Source/Programs/UnrealBuildAccelerator/Common/Private/UbaSession.cpp
+++ b/Engine/Source/Programs/UnrealBuildAccelerator/Common/Private/UbaSession.cpp
@@ -392,7 +392,7 @@ namespace uba
 			}
 			else
 			{
-				// Free up unsude memory since we're keeping it around
+				// Free up unused memory since we're keeping it around
 				m_sharedMemory.TrimMemory(memoryHandle, file.memoryWritten);
 
 				bool isInvisible = rules.IsInvisible(file.name);
@@ -1757,7 +1757,7 @@ namespace uba
 			parseRes = ParseBinary(temp3, ToView(applicationDir), info, handleImport, errorStr);
 
 		if (!parseRes)
-			return m_logger.Error(TC("Failed to parse binary %s for imports"), temp3.data);
+			return m_logger.Error(TC("Failed to parse binary %s for imports%s"), temp3.data, isLinuxBinary ? TC(" (using linux binary parser)") : TC(""));
 
 		if (errorStr.count)
 			return m_logger.Error(errorStr.data);
