[UBA]
* Fixed stats bug where recvcas raw size was not properly updated on server side

diff --git a/Engine/Source/Programs/UnrealBuildAccelerator/Common/Private/UbaStorageServer.cpp b/Engine/Source/Programs/UnrealBuildAccelerator/Common/Private/UbaStorageServer.cpp
--- a/Engine/Source/Programs/UnrealBuildAccelerator/Common/Private/UbaStorageServer.cpp
+++ b/Engine/Source/Programs/UnrealBuildAccelerator/Common/Private/UbaStorageServer.cpp
@@ -1233,6 +1233,7 @@ namespace uba
 					activeLock.Leave();
 
 					s.fileSize = firstStore->fileSize;
+					s.actualSize = firstStore->actualSize;
 					s.mappedView = firstStore->mappedView;
 					s.fileAccessor = firstStore->fileAccessor;
 					s.casEntry = firstStore->casEntry;
