[UBA]
* Fixed so visualizer can do horizontal zoom with cmd+scroll horizontal

diff --git a/Engine/Source/Programs/UnrealBuildAccelerator/Visualizer/Private/Mac/UbaVisualizerApp.mm b/Engine/Source/Programs/UnrealBuildAccelerator/Visualizer/Private/Mac/UbaVisualizerApp.mm
--- a/Engine/Source/Programs/UnrealBuildAccelerator/Visualizer/Private/Mac/UbaVisualizerApp.mm
+++ b/Engine/Source/Programs/UnrealBuildAccelerator/Visualizer/Private/Mac/UbaVisualizerApp.mm
@@ -934,7 +934,16 @@ - (void)mouseDown:(NSEvent*)event
 - (void)scrollWheel:(NSEvent *)event
 {
 	using namespace uba;
-	_visualizer->Scroll(float(event.scrollingDeltaX), float(event.scrollingDeltaY));
+	if (event.modifierFlags & NSEventModifierFlagCommand)
+	{
+		_visualizer->m_config.horizontalScaleValue += float(event.scrollingDeltaX*0.001);
+		//_visualizer->m_zoomValue += float(event.scrollingDeltaX*0.001);
+	}
+	else
+	{
+		_visualizer->Scroll(float(event.scrollingDeltaX), float(event.scrollingDeltaY));
+	}
+
 	[self setNeedsDisplay:YES];
 }
 
@@ -975,9 +984,6 @@ - (void)magnifyWithEvent:(NSEvent *)event
 
 	if (event.phase == NSEventPhaseEnded || event.phase == NSEventPhaseCancelled)
 		_magnifyAccum = 0.0;
-
-    //CGFloat magnification = event.magnification;
-	//_visualizer->m_zoomValue += float(magnification);
 }
 
 - (NSMenu *)menuForEvent:(NSEvent *)event
