[UBA]
* Changed mac file system to be treated as case sensitive

diff --git a/Engine/Source/Programs/UnrealBuildAccelerator/Core/Public/UbaPlatform.h b/Engine/Source/Programs/UnrealBuildAccelerator/Core/Public/UbaPlatform.h
--- a/Engine/Source/Programs/UnrealBuildAccelerator/Core/Public/UbaPlatform.h
+++ b/Engine/Source/Programs/UnrealBuildAccelerator/Core/Public/UbaPlatform.h
@@ -180,7 +180,7 @@ namespace uba
 	inline constexpr bool CaseInsensitiveFs = false;
 	#define st_mtimespec st_mtim
 	#else
-	inline constexpr bool CaseInsensitiveFs = true;
+	inline constexpr bool CaseInsensitiveFs = false; // This can be true since mac can allow insensitive file system.
 	#endif
 	#define UBA_NOINLINE __attribute__((noinline))
     #define ERROR_SUCCESS 0
