[UBA]
* Detoured FlushViewOfFile and made it a no-op

diff --git a/Engine/Source/Programs/UnrealBuildAccelerator/Detours/Private/Windows/UbaDetoursFunctionsKernelBase.inl b/Engine/Source/Programs/UnrealBuildAccelerator/Detours/Private/Windows/UbaDetoursFunctionsKernelBase.inl
--- a/Engine/Source/Programs/UnrealBuildAccelerator/Detours/Private/Windows/UbaDetoursFunctionsKernelBase.inl
+++ b/Engine/Source/Programs/UnrealBuildAccelerator/Detours/Private/Windows/UbaDetoursFunctionsKernelBase.inl
@@ -2637,6 +2637,13 @@ LPVOID Detoured_MapViewOfFile(HANDLE hFileMappingObject, DWORD dwDesiredAccess,
 	return Detoured_MapViewOfFileEx(hFileMappingObject, dwDesiredAccess, dwFileOffsetHigh, dwFileOffsetLow, dwNumberOfBytesToMap, nullptr);
 }
 
+BOOL Detoured_FlushViewOfFile(LPCVOID lpBaseAddress, SIZE_T dwNumberOfBytesToFlush)
+{
+	DETOURED_CALL(FlushViewOfFile);
+	DEBUG_LOG_DETOURED(L"FlushViewOfFile", L"(NO-OP!!!!) %p (size %llu) -> true", lpBaseAddress, dwNumberOfBytesToFlush);
+	return true;
+}
+
 BOOL Detoured_UnmapViewOfFileEx(PVOID lpBaseAddress, ULONG UnmapFlags)
 {
 	DETOURED_CALL(UnmapViewOfFileEx);
diff --git a/Engine/Source/Programs/UnrealBuildAccelerator/Detours/Private/Windows/UbaDetoursFunctionsWin.h b/Engine/Source/Programs/UnrealBuildAccelerator/Detours/Private/Windows/UbaDetoursFunctionsWin.h
--- a/Engine/Source/Programs/UnrealBuildAccelerator/Detours/Private/Windows/UbaDetoursFunctionsWin.h
+++ b/Engine/Source/Programs/UnrealBuildAccelerator/Detours/Private/Windows/UbaDetoursFunctionsWin.h
@@ -92,6 +92,7 @@
 	DETOURED_FUNCTION(OpenFileMappingW) \
 	DETOURED_FUNCTION(MapViewOfFile) \
 	DETOURED_FUNCTION(MapViewOfFileEx) \
+	DETOURED_FUNCTION(FlushViewOfFile) \
 	DETOURED_FUNCTION(UnmapViewOfFile) \
 	DETOURED_FUNCTION(UnmapViewOfFileEx) \
 	DETOURED_FUNCTION(GetFinalPathNameByHandleW) \
