[UBA]
* Added "UbaCli -validate <binary>" to mac version that can be used to check if binary works with DYLD env variables. This is related to hardened binaries (which mac NDK clang binary is)
* Added hint in trace log when failing to run application in uba and getting 0 messages (likely SIP causing os to drop DYLD environment variables)

diff --git a/Engine/Source/Programs/UnrealBuildAccelerator/Cli/Private/UbaCli.cpp b/Engine/Source/Programs/UnrealBuildAccelerator/Cli/Private/UbaCli.cpp
--- a/Engine/Source/Programs/UnrealBuildAccelerator/Cli/Private/UbaCli.cpp
+++ b/Engine/Source/Programs/UnrealBuildAccelerator/Cli/Private/UbaCli.cpp
@@ -31,6 +31,15 @@
 #pragma comment (lib, "Dbghelp.lib")
 #endif
 
+#if PLATFORM_MAC
+#include <Security/Security.h>
+#include <CoreFoundation/CoreFoundation.h>
+#include <sys/sysctl.h>
+#define CS_RUNTIME        0x00010000  // hardened runtime
+#define CS_RESTRICT       0x00000800
+#define CSR_ALLOW_UNRESTRICTED_DYLD (1 << 7)
+#endif
+
 namespace uba
 {
 	const tchar*	Version = GetVersionString();
@@ -73,6 +82,9 @@ namespace uba
 		logger.Info(TC("   native                  Will run executable in a normal way"));
 		logger.Info(TC(""));
 		logger.Info(TC("  Options:"));
+		#if PLATFORM_MAC
+		logger.Info(TC("   -validate               Will validate binary to see if it can run dyld env needed by uba"));
+		#endif
 		logger.Info(TC("   -dir=<rootdir>          The directory used to store data. Defaults to \"%s\""), DefaultRootDir);
 		logger.Info(TC("   -port=[<host>:]<port>   The ip/name and port (default: %u) of the machine we want to help"), DefaultPort);
 		logger.Info(TC("   -log                    Log all processes detouring information to file (only works with debug builds)"));
@@ -192,6 +204,83 @@ namespace uba
 		}
 	}
 
+#if PLATFORM_MAC
+	bool ExplainDyldEnvSupport(Logger& logger, const char* path)
+	{
+		bool ok = true;
+		struct stat st;
+		if (stat(path, &st) != 0)
+			return logger.Error("FAIL: cannot stat binary %s", path);
+
+		if (st.st_mode & (S_ISUID | S_ISGID))
+			return logger.Error("FAIL: setuid/setgid binaries never allow DYLD injection");
+
+		CFURLRef url = CFURLCreateFromFileSystemRepresentation(kCFAllocatorDefault, (const UInt8*)path, strlen(path), true);
+		if (!url)
+			return logger.Error("FAIL: invalid path");
+
+		SecStaticCodeRef code = nullptr;
+		ok = SecStaticCodeCreateWithPath(url, kSecCSDefaultFlags, &code) == errSecSuccess;
+		CFRelease(url);
+		if (!ok)
+			return logger.Error("FAIL: cannot create SecStaticCode");
+
+		CFDictionaryRef info = nullptr;
+		ok = SecCodeCopySigningInformation(code, kSecCSSigningInformation, &info) == errSecSuccess;
+		CFRelease(code);
+		if (!ok)
+			return logger.Error("FAIL: cannot read code signature");
+
+		uint32_t flags = 0;
+		if (CFNumberRef flagsNum = (CFNumberRef)CFDictionaryGetValue(info, kSecCodeInfoFlags))
+			CFNumberGetValue(flagsNum, kCFNumberSInt32Type, &flags);
+
+		bool hardened = (flags & CS_RUNTIME) != 0;
+		bool restricted = (flags & CS_RESTRICT) != 0;
+
+		logger.Info("Hardened runtime: %s", hardened ? "YES" : "NO");
+		logger.Info("Restricted binary: %s", restricted ? "YES" : "NO");
+
+		bool disableLV = false;
+		if (CFDictionaryRef ent = (CFDictionaryRef)CFDictionaryGetValue(info, kSecCodeInfoEntitlementsDict))
+			if (CFDictionaryGetValue(ent, CFSTR("com.apple.security.cs.disable-library-validation")))
+				disableLV = true;
+
+		logger.Info("Library validation disabled: %s", disableLV ? "YES" : "NO");
+		CFRelease(info);
+
+		if (restricted)
+		{
+			logger.Error("FAIL: restricted binaries do not honor DYLD env vars");
+			ok = false;
+		}
+		if (hardened && !disableLV)
+		{
+			logger.Error("FAIL: hardened runtime with library validation blocks DYLD injection");
+			logger.Error("Re-sign binary with command line:");
+			logger.Error(" codesign --force --deep --sign - %s", path);
+			ok = false;
+		}
+#if 0 // This does not seem to ever work
+		int csr = 0; size_t csrSize = sizeof(csr);
+		if (sysctlbyname("kern.csr.active_config", &csr, &csrSize, nullptr, 0) == 0)
+		{
+			bool sipBlocksDyld = (csr & CSR_ALLOW_UNRESTRICTED_DYLD) == 0;
+			logger.Info("  SIP dyld unrestricted: %s", sipBlocksDyld ? "NO" : "YES");
+			if (sipBlocksDyld)
+				logger.Info("  NOTE: SIP may strip DYLD_* for protected system binaries");
+		}
+		else
+		{
+			logger.Info("  SIP state: unavailable (assuming enabled)");
+		}
+#endif
+		if (ok)
+			logger.Info("PASS: DYLD_* environment variables SHOULD work for this binary");
+		return ok;
+	}
+#endif
+
 	bool WrappedMain(int argc, tchar* argv[])
 	{
 		using namespace uba;
@@ -222,6 +311,7 @@ namespace uba
 		bool checkCloud = false;
 		bool getCas = false;
 		bool listImports = false;
+		bool validateBinary = false;
 		bool deleteCas = false;
 		bool fastDelete = false;
 		bool enableStdOut = true;
@@ -474,6 +564,12 @@ namespace uba
 				else if (name.Equals(TCV("-listimports")))
 				{
 					listImports = true;
+					commandType = CommandType_None;
+				}
+				else if (name.Equals(TCV("-validate")))
+				{
+					validateBinary = true;
+					commandType = CommandType_None;
 				}
 				else if (name.Equals(TCV("-summary")))
 				{
@@ -544,12 +640,16 @@ namespace uba
 				{
 					return PrintHelp(TC(""));
 				}
-				else
+				else if (name[0] == '-' || commandType != CommandType_None)
 				{
 					StringBuffer<> msg;
 					msg.Appendf(TC("Unknown argument '%s'"), name.data);
 					return PrintHelp(msg.data);
 				}
+				else if (commandType == CommandType_None)
+				{
+					application = name.data;
+				}
 				return true;
 			};
 
@@ -1106,35 +1206,45 @@ namespace uba
 					logger.Info(TC("  CasKey (uncompressed): %s"), CasKeyString(uncompressedKey).str);
 				return true;
 			}
+		}
 
-			if (listImports)
+		if (listImports)
+		{
+			StringBuffer<> error;
+			bool printImports = true;
+			BinaryInfo info;
+			if (!ParseBinary(application, StringView(application).GetPath(), info, [&](const tchar* import, bool isKnown, const tchar* const* loaderPaths)
 			{
-				StringBuffer<> error;
-				bool printImports = true;
-				BinaryInfo info;
-				if (!ParseBinary(application, StringView(application).GetPath(), info, [&](const tchar* import, bool isKnown, const tchar* const* loaderPaths)
+				if (printImports)
 				{
-					if (printImports)
+					if (loaderPaths && *loaderPaths)
 					{
-						if (loaderPaths && *loaderPaths)
-						{
-							logger.Info(TC("LoaderPaths:"));
-							for (auto it=loaderPaths;*it; ++it)
-								if (**it)
-									logger.Info(TC("  %s"), *it);
-						}
-						printImports = false;
-						logger.Info(TC("Imports:"));
+						logger.Info(TC("LoaderPaths:"));
+						for (auto it=loaderPaths;*it; ++it)
+							if (**it)
+								logger.Info(TC("  %s"), *it);
 					}
-					logger.Info(TC("  %s"), import);
+					printImports = false;
+					logger.Info(TC("Imports:"));
+				}
+				logger.Info(TC("  %s"), import);
+
+			}, error))
+				return logger.Error(TC("%s"), error.data);
+			return true;
+		}
 
-				}, error))
-					return logger.Error(TC("%s"), error.data);
-				#if PLATFORM_MAC
+		if (validateBinary)
+		{
+			#if PLATFORM_MAC
+			BinaryInfo info;
+			StringBuffer<> error;
+			if (ParseBinary(application, StringView(application).GetPath(), info, [&](const tchar* import, bool isKnown, const tchar* const* loaderPaths) {}, error))
 				logger.Info(TC("MinOsVersion: %u.%u.%u"), (info.minVersion >> 16) & 0xffff, (info.minVersion >> 8) & 0xff, info.minVersion & 0xff);
-				#endif
-				return true;
-			}
+			return ExplainDyldEnvSupport(logger, application.c_str());
+			#else
+			return true;
+			#endif
 		}
 
 		const tchar* dbgStr = TC("");
diff --git a/Engine/Source/Programs/UnrealBuildAccelerator/Cli/UbaCli.Build.cs b/Engine/Source/Programs/UnrealBuildAccelerator/Cli/UbaCli.Build.cs
--- a/Engine/Source/Programs/UnrealBuildAccelerator/Cli/UbaCli.Build.cs
+++ b/Engine/Source/Programs/UnrealBuildAccelerator/Cli/UbaCli.Build.cs
@@ -19,5 +19,10 @@ public class UbaCli : ModuleRules
 		PrivateDefinitions.AddRange(new string[] {
 			"_CONSOLE",
 		});
+
+		if (Target.Platform == UnrealTargetPlatform.Mac)
+		{
+			PublicFrameworks.AddRange(["CoreFoundation", "Security"]);
+		}
 	}
 }
diff --git a/Engine/Source/Programs/UnrealBuildAccelerator/Common/Private/UbaProcess.cpp b/Engine/Source/Programs/UnrealBuildAccelerator/Common/Private/UbaProcess.cpp
--- a/Engine/Source/Programs/UnrealBuildAccelerator/Common/Private/UbaProcess.cpp
+++ b/Engine/Source/Programs/UnrealBuildAccelerator/Common/Private/UbaProcess.cpp
@@ -243,6 +243,7 @@ namespace uba
 		{
 			m_exitCode = 44324;
 			CallExitedFunc();
+			m_hasExited = true;
 			return false;
 		}
 
@@ -272,6 +273,7 @@ namespace uba
 				sharedMemory.CloseHandle(m_session.m_logger, memHandle, TC("OverlayFiles"));
 				m_exitCode = 44325;
 				CallExitedFunc();
+				m_hasExited = true;
 				return false;
 			}
 
@@ -422,13 +424,27 @@ namespace uba
 		{
 			memset(&signalInfo, 0, sizeof(signalInfo));
 			int res = waitid(P_PID, (unsigned int)m_nativeProcessId, &signalInfo, WEXITED | WNOHANG | WNOWAIT);
+			StringBuffer<> err;
+
 			if (res)
 			{
 				UBA_ASSERT(res == -1);
 				if (errno == EINTR)
 					continue;
-				if (errno == ECHILD) // This should not happen, but let's return true on this since we can't use waitid on processes that are not our children
+				if (errno == ECHILD)
+				{
+					#if PLATFORM_MAC
+					if (m_messageCount == 0)
+					{
+						err.Appendf(TC("NO MESSAGES RECEIVED BEFORE EXIT. Use \"UbaCli -validate %s\" to verify DYLD env support"), m_startInfo.application);
+						LogLine(false, err.data, LogEntryType_Error);
+						m_nativeProcessExitCode = UBA_EXIT_CODE(667);
+						return false;
+					}
+					#endif
+					UBA_ASSERTF(m_parentProcess && m_messageCount == 0, "waitid got errno ECHILD when waiting for %s", m_startInfo.application);
 					return true;
+				}
 				UBA_ASSERTF(false, "waitid failed with error: %u (%s)", errno, strerror(errno));
 				break;
 			}
@@ -466,7 +482,6 @@ namespace uba
 				if (!codeType) // Is null if graceful exit (CLD_EXITED)
 					break;
 					
-				StringBuffer<> err;
 				err.Appendf(TC("Process %u (%s) %s by signal %i. Received %u messages. Execution time: %s."), nativeProcessId, m_startInfo.GetDescription(), codeType, signalInfo.si_status, m_messageCount, TimeToText(GetTime() - m_startTime).str);
 				LogLine(false, err.data, LogEntryType_Error);
 				m_nativeProcessExitCode = UBA_EXIT_CODE(666); // We do exit code 666 to trigger non-uba retry on the outside
@@ -489,6 +504,16 @@ namespace uba
 			err.Appendf(TC("ERROR: Process %u (%s) not active but did not get exit message. Received %u messages. Signal code: %i. Exit value or signal: %i. Execution time: %s."), m_nativeProcessId, m_startInfo.GetDescription(), m_messageCount, signalInfo.si_code, signalInfo.si_status, TimeToText(GetTime() - m_startTime).str);
 			LogLine(false, err.data, LogEntryType_Error);
 			m_nativeProcessExitCode = UBA_EXIT_CODE(666);
+
+			#if PLATFORM_MAC
+			if (m_messageCount == 0)
+			{
+				err.Clear().Appendf(TC("Use \"UbaCli -validate %s\" to verify DYLD env support"), m_startInfo.application);
+				ProcessImpl* p = m_parentProcess ? m_parentProcess : this;
+				p->LogLine(false, err.data, LogEntryType_Error);
+				m_nativeProcessExitCode = UBA_EXIT_CODE(667);
+			}
+			#endif
 		}
 
 		//m_session.m_logger.Info(TC("IsActive false (no parent)"), LastErrorToText().data);
