[UBA]
* Fixed assert that was wrong

diff --git a/Engine/Source/Programs/UnrealBuildAccelerator/Common/Private/UbaProcess.cpp b/Engine/Source/Programs/UnrealBuildAccelerator/Common/Private/UbaProcess.cpp
--- a/Engine/Source/Programs/UnrealBuildAccelerator/Common/Private/UbaProcess.cpp
+++ b/Engine/Source/Programs/UnrealBuildAccelerator/Common/Private/UbaProcess.cpp
@@ -442,7 +442,7 @@ namespace uba
 						return false;
 					}
 					#endif
-					UBA_ASSERTF(m_parentProcess && m_messageCount == 0, "waitid got errno ECHILD when waiting for %s", m_startInfo.application);
+					UBA_ASSERTF(!m_parentProcess || m_messageCount != 0, "waitid got errno ECHILD when waiting for %s", m_startInfo.application);
 					return true;
 				}
 				UBA_ASSERTF(false, "waitid failed with error: %u (%s)", errno, strerror(errno));
