[UBA]
* Fix assert dialog for mac uba visualizer

diff --git a/Engine/Source/Programs/UnrealBuildAccelerator/Visualizer/Private/Mac/UbaVisualizerApp.mm b/Engine/Source/Programs/UnrealBuildAccelerator/Visualizer/Private/Mac/UbaVisualizerApp.mm
--- a/Engine/Source/Programs/UnrealBuildAccelerator/Visualizer/Private/Mac/UbaVisualizerApp.mm
+++ b/Engine/Source/Programs/UnrealBuildAccelerator/Visualizer/Private/Mac/UbaVisualizerApp.mm
@@ -77,21 +77,25 @@ virtual void Log(LogEntryType type, const tchar* str, u32 strLen, const tchar* p
 			if (!text)
 				text = @"Boo";
 
-			NSAlert* alert = [[NSAlert alloc] init];
-			alert.messageText = @"UbaVisualizer";
-			alert.informativeText = text;
-			[alert addButtonWithTitle:@"OK"];
-			alert.alertStyle = NSAlertStyleInformational;
+			auto createAlert = [&]()
+			{
+				NSAlert* alert = [[NSAlert alloc] init];
+				alert.messageText = @"UbaVisualizer";
+				alert.informativeText = text;
+				[alert addButtonWithTitle:@"OK"];
+				alert.alertStyle = NSAlertStyleInformational;
+				return alert;
+			};
 
 			if ([NSThread isMainThread])
 			{
 				[NSApp activateIgnoringOtherApps:YES];
-				[alert runModal];
+				[createAlert() runModal];
 			}
 			else
 			{
 				dispatch_sync(dispatch_get_main_queue(), ^{
-					[alert beginSheetModalForWindow:_window completionHandler:nil];
+					[createAlert() beginSheetModalForWindow:_window completionHandler:nil];
 				});
 			}
 
@@ -936,7 +940,8 @@ - (void)scrollWheel:(NSEvent *)event
 	using namespace uba;
 	if (event.modifierFlags & NSEventModifierFlagCommand)
 	{
-		_visualizer->m_config.horizontalScaleValue += float(event.scrollingDeltaX*0.001);
+		float newValue = _visualizer->m_config.horizontalScaleValue + float(event.scrollingDeltaX*0.001);
+		_visualizer->m_config.horizontalScaleValue = Max(newValue, 0.001f);
 		//_visualizer->m_zoomValue += float(event.scrollingDeltaX*0.001);
 	}
 	else
