[UBA] Temporarily disabling invalid file roots warning

This is to mitigate a log spam in the build farm while we work on a fix.

diff --git a/Engine/Source/Programs/UnrealBuildAccelerator/Common/Private/UbaCacheClient.cpp b/Engine/Source/Programs/UnrealBuildAccelerator/Common/Private/UbaCacheClient.cpp
--- a/Engine/Source/Programs/UnrealBuildAccelerator/Common/Private/UbaCacheClient.cpp
+++ b/Engine/Source/Programs/UnrealBuildAccelerator/Common/Private/UbaCacheClient.cpp
@@ -391,7 +391,8 @@ namespace uba
 				auto root = rootPaths.FindRoot(path);
 				if (!root)
 				{
-					m_logger.Info(TC("FILE WITHOUT ROOT: %s (%s)"), path.data, info.GetDescription());
+					// Temporarily disabling invalid roots warnings to mitigate a log spam until an issue with them can be fixed
+					//m_logger.Info(TC("FILE WITHOUT ROOT: %s (%s)"), path.data, info.GetDescription());
 					success = false;
 					continue;
 				}
