[Uba] Fix uba CLI issue where the -cross= commandline argument does not work to test cross platform execution.

The API changed to support specify cross architecture and cross platform, but uba cli was not updated to use the new argument.

This CLI updates the Run Process API call to allow cross platform execution

#Synced-CL 50408994

diff --git a/Engine/Source/Programs/UnrealBuildAccelerator/Cli/Private/UbaCli.cpp b/Engine/Source/Programs/UnrealBuildAccelerator/Cli/Private/UbaCli.cpp
--- a/Engine/Source/Programs/UnrealBuildAccelerator/Cli/Private/UbaCli.cpp
+++ b/Engine/Source/Programs/UnrealBuildAccelerator/Cli/Private/UbaCli.cpp
@@ -1542,7 +1542,7 @@ namespace uba
 			if (enableStdOut)
 				pinfo.logLineFunc = [](void* userData, const tchar* line, u32 length, LogEntryType type) { ((Logger*)userData)->Log(type, line, length); };
 			logger.Info(TC("Running %s %s"), app.c_str(), arg.c_str());
-			ProcessHandle process = sessionServer.RunProcessRemote(pinfo, 1.0f, nullptr, 0, allowCross);
+			ProcessHandle process = sessionServer.RunProcessRemote(pinfo, 1.0f, nullptr, 0, allowCross, allowCross);
 			process.WaitForExit(~0u);
 			if (process.GetExitCode() != 0)
 				return logger.Error(TC("Error exit code: %u"), process.GetExitCode());
