[UBA]
* Improved some error descriptions

diff --git a/Engine/Source/Programs/UnrealBuildAccelerator/Common/Private/UbaProcess.cpp b/Engine/Source/Programs/UnrealBuildAccelerator/Common/Private/UbaProcess.cpp
--- a/Engine/Source/Programs/UnrealBuildAccelerator/Common/Private/UbaProcess.cpp
+++ b/Engine/Source/Programs/UnrealBuildAccelerator/Common/Private/UbaProcess.cpp
@@ -2439,8 +2439,11 @@ namespace uba
 		bool hadTimeout = false;
 		if (cancel)
 		{
-			TerminateProcess(handle, ProcessCancelExitCode);
-			WaitForSingleObject(handle, 1000); // On windows this should return in milliseconds. On wine it might take a little while. We need to wait to make sure process is reaped
+			if (!TerminateProcess(handle, ProcessCancelExitCode))
+				logger.Info(TC("TerminateProcess for process %llu (%s) failed when cancelling (%s)"), u64(handle), m_startInfo.GetDescription(), LastErrorToText().data);
+			DWORD res = WaitForSingleObject(handle, 1000); // On windows this should return in milliseconds. On wine it might take a little while. We need to wait to make sure process is reaped
+			if (res == WAIT_FAILED)
+				logger.Info(TC("WaitForSingleObject for process %llu (%s) failed when cancelling (%s)"), u64(handle), m_startInfo.GetDescription(), LastErrorToText().data);
 		}
 		else
 		{	
@@ -2458,11 +2461,12 @@ namespace uba
 					{
 						hadTimeout = true;
 						const tchar* gotMessage = m_gotExitMessage ? TC("Got") : TC("Did not get");
-						logger.Info(TC("WaitForSingleObject timed out after 120 seconds waiting for process %s to exit (Exit code %u, %s ExitMessage and wrote %u files. Runtime: %s). Will terminate and wait again"), m_startInfo.GetDescription(), m_nativeProcessExitCode, gotMessage, u32(m_shared.writtenFiles.size()), TimeToText(GetTime() - m_startTime).str);
-						TerminateProcess(handle, m_nativeProcessExitCode);
+						logger.Info(TC("WaitForSingleObject timed out after 120 seconds waiting for process %llu (%s) to exit (Exit code %u, %s ExitMessage and wrote %u files. Runtime: %s). Will terminate and wait again"), u64(handle), m_startInfo.GetDescription(), m_nativeProcessExitCode, gotMessage, u32(m_shared.writtenFiles.size()), TimeToText(GetTime() - m_startTime).str);
+						if (!TerminateProcess(handle, m_nativeProcessExitCode))
+							logger.Info(TC("TerminateProcess for process %llu (%s) failed when waiting timed out (%s)"), u64(handle), m_startInfo.GetDescription(), LastErrorToText().data);
 						continue;
 					}
-					logger.Error(TC("WaitForSingleObject failed while waiting for process %s to exit even after terminating it (%s)"), m_startInfo.GetDescription(), LastErrorToText().data);
+					logger.Error(TC("WaitForSingleObject failed while waiting for process %llu (%s) to exit even after terminating it (%s)"), u64(handle), m_startInfo.GetDescription(), LastErrorToText().data);
 				}
 				else if (res == WAIT_FAILED)
 					logger.Error(TC("WaitForSingleObject failed while waiting for process to exit (%s)"), LastErrorToText().data);
