[UBA]
* Fixed so pgo optimize linking works in uba

diff --git a/Engine/Source/Programs/UnrealBuildAccelerator/Core/Private/UbaApplicationRules.cpp b/Engine/Source/Programs/UnrealBuildAccelerator/Core/Private/UbaApplicationRules.cpp
--- a/Engine/Source/Programs/UnrealBuildAccelerator/Core/Private/UbaApplicationRules.cpp
+++ b/Engine/Source/Programs/UnrealBuildAccelerator/Core/Private/UbaApplicationRules.cpp
@@ -85,7 +85,7 @@ namespace uba
 		{
 			if (fileName.Contains(TCV("\\vctip_")))
 				return true;
-			if (fileName.Contains(systemTemp))
+			if (fileName.Contains(systemTemp) && !fileName.EndsWith(TCV(".pgd"))) // pgd is copied into temp folder and detoured CopyFileExW is missing implementation to copy to memory
 				return true;
 			return Super::KeepInMemory(fileName, systemTemp, isRunningRemote, isWrite);
 		}
@@ -843,6 +843,7 @@ namespace uba
 			{ ApplicationHash(TC("include-what-you-use.exe")),	new ApplicationRulesIWYU() },
 			{ ApplicationHash(TC("ubatestapp.exe")),			new ApplicationRulesUbaTestApp() },
 			{ ApplicationHash(TC("pgocvt.exe")),				new ApplicationRulesLinkExe() },
+			{ ApplicationHash(TC("pgomgr.exe")),				new ApplicationRulesVcLink() },
 			{ ApplicationHash(TC("icx.exe")),					new ApplicationRulesIcxExe() },
 			//{ L"MSBuild.dll"),				new ApplicationRules() },
 			//{ L"BreakpadSymbolEncoder.exe"),	new ApplicationRulesClang() },
diff --git a/Engine/Source/Programs/UnrealBuildAccelerator/Detours/Private/Windows/UbaDetoursFunctionsKernelBase.inl b/Engine/Source/Programs/UnrealBuildAccelerator/Detours/Private/Windows/UbaDetoursFunctionsKernelBase.inl
--- a/Engine/Source/Programs/UnrealBuildAccelerator/Detours/Private/Windows/UbaDetoursFunctionsKernelBase.inl
+++ b/Engine/Source/Programs/UnrealBuildAccelerator/Detours/Private/Windows/UbaDetoursFunctionsKernelBase.inl
@@ -1473,10 +1473,8 @@ BOOL Detoured_GetModuleHandleExW(DWORD dwFlags, LPCWSTR lpModuleName, HMODULE *p
 	return res;
 }
 
-BOOL Detoured_CopyFileExW(LPCWSTR lpExistingFileName, LPCWSTR lpNewFileName, LPPROGRESS_ROUTINE lpProgressRoutine, LPVOID lpData, LPBOOL pbCancel, DWORD dwCopyFlags)
+BOOL Shared_CopyFileExW(const tchar* func, LPCWSTR lpExistingFileName, LPCWSTR lpNewFileName, LPPROGRESS_ROUTINE lpProgressRoutine, LPVOID lpData, LPBOOL pbCancel, DWORD dwCopyFlags)
 {
-	DETOURED_CALL(CopyFileExW);
-
 	StringBuffer<> fromName;
 	FixPath(fromName, lpExistingFileName);
 	DevirtualizePath(fromName);
@@ -1512,17 +1510,27 @@ BOOL Detoured_CopyFileExW(LPCWSTR lpExistingFileName, LPCWSTR lpNewFileName, LPP
 		g_directoryTable.ParseDirectoryTable(directoryTableSize);
 		UBA_ASSERT(g_runningRemote);
 		SetLastError(lastError);
+		DEBUG_LOG_DETOURED(func, L"(File was copied server side)");
 		return lastError == ERROR_SUCCESS;
 	}
 
 	// TODO: This copy should probably be moved to session process instead.. to handle failing to copy better
 
+	bool keepToFileInMemory = g_allowKeepFilesInMemory && g_rules->KeepInMemory(newToName, g_systemTemp, g_runningRemote, true);
+
 	bool res;
+	if (keepToFileInMemory)
+	{
+		UBA_ASSERTF(false, TC("%s copying to memory file is not implemented (Used by link.exe -PGOOptimize). Either change application rules or implement code path (Destination file: %s)"), func, newToName.data);
+		SetLastError(ERROR_IO_DEVICE);
+		return false;
+	}
+	else
 	{
 		SuppressCreateFileDetourScope cfs;
 		res = True_CopyFileExW(newFromName.data, newToName.data, lpProgressRoutine, lpData, pbCancel, dwCopyFlags);
 	}
-	DEBUG_LOG_TRUE(L"CopyFileExW", L"%ls to %ls flags: %u (%ls to %ls) -> %ls", lpExistingFileName, lpNewFileName, dwCopyFlags, newFromName.data, newToName.data, ToString(res));
+	DEBUG_LOG_TRUE(func, L"%ls to %ls flags: %u (%ls to %ls) -> %ls", lpExistingFileName, lpNewFileName, dwCopyFlags, newFromName.data, newToName.data, ToString(res));
 
 	// We need to report the new file that has been added (and we must do it _after_ it has been copied
 	if (!closeId)
@@ -1534,11 +1542,16 @@ BOOL Detoured_CopyFileExW(LPCWSTR lpExistingFileName, LPCWSTR lpNewFileName, LPP
 	return res;
 }
 
+BOOL Detoured_CopyFileExW(LPCWSTR lpExistingFileName, LPCWSTR lpNewFileName, LPPROGRESS_ROUTINE lpProgressRoutine, LPVOID lpData, LPBOOL pbCancel, DWORD dwCopyFlags)
+{
+	DETOURED_CALL(CopyFileExW);
+	return Shared_CopyFileExW(TC("CopyFileExW"), lpExistingFileName, lpNewFileName, lpProgressRoutine, lpData, pbCancel, dwCopyFlags);
+}
+
 BOOL Detoured_CopyFileW(LPCWSTR lpExistingFileName, LPCWSTR lpNewFileName, BOOL bFailIfExists)
 {
 	DETOURED_CALL(CopyFileW);
-	DEBUG_LOG_TRUE(L"CopyFileW", L"");
-	return Detoured_CopyFileExW(lpExistingFileName, lpNewFileName, (LPPROGRESS_ROUTINE)NULL, (LPVOID)NULL, (LPBOOL)NULL, bFailIfExists ? (DWORD)COPY_FILE_FAIL_IF_EXISTS : 0);
+	return Shared_CopyFileExW(TC("CopyFileW"), lpExistingFileName, lpNewFileName, (LPPROGRESS_ROUTINE)NULL, (LPVOID)NULL, (LPBOOL)NULL, bFailIfExists ? (DWORD)COPY_FILE_FAIL_IF_EXISTS : 0);
 }
 
 
