[UBA]
* Added undocumented win32 code to get page file sizes because the current one is very non-accurate and gives wrong values when page files are controlled by system.

diff --git a/Engine/Source/Programs/UnrealBuildAccelerator/Common/Private/UbaEnvironment.cpp b/Engine/Source/Programs/UnrealBuildAccelerator/Common/Private/UbaEnvironment.cpp
--- a/Engine/Source/Programs/UnrealBuildAccelerator/Common/Private/UbaEnvironment.cpp
+++ b/Engine/Source/Programs/UnrealBuildAccelerator/Common/Private/UbaEnvironment.cpp
@@ -6,6 +6,7 @@
 #if PLATFORM_WINDOWS
 #include <powerbase.h>
 #pragma comment(lib, "Powrprof.lib")
+#include <winternl.h>
 #elif PLATFORM_MAC
 #include <mach/vm_statistics.h>
 #include <mach/mach_types.h>
@@ -137,6 +138,37 @@ namespace uba
 #if PLATFORM_WINDOWS
 	u64 GetMaxPageFileSize(Logger& logger, MEMORYSTATUSEX& memStatus)
 	{
+		struct SYSTEM_PAGEFILE_INFORMATION
+		{
+			ULONG NextEntryOffset;
+			ULONG TotalSize;      // in pages
+			ULONG TotalInUse;     // in pages
+			ULONG PeakUsage;      // in pages
+			UNICODE_STRING PageFileName;
+		};
+
+		Vector<u8> queryBuffer(64*1024);
+		ULONG returnedSize;
+		constexpr auto systemPageFileInformation = (SYSTEM_INFORMATION_CLASS)0x12;
+		if (NtQuerySystemInformation(systemPageFileInformation, queryBuffer.data(), (u32)queryBuffer.size(), &returnedSize) == STATUS_SUCCESS)
+		{
+			u64 totalSize = 0;
+			u8* entryIt = queryBuffer.data();
+			while (true)
+			{
+				auto& info = *reinterpret_cast<SYSTEM_PAGEFILE_INFORMATION*>(entryIt);
+				totalSize += u64(info.TotalSize)*4096;
+				if (info.NextEntryOffset == 0)
+					break;
+				entryIt += info.NextEntryOffset;
+				return 0;
+			}
+			return totalSize;
+		}
+
+		// Old fallback path.. doesn't really work well with system managed page files because they take into account lots of things
+		// There is some binary blob inside the registry that contains some sort of history etc.
+
 		tchar systemDrive = 'c';
 		{
 			tchar temp[32];
