[UBA]
* Removed accidental return 0 in GetMaxPageFile caused by previous submit

diff --git a/Engine/Source/Programs/UnrealBuildAccelerator/Common/Private/UbaEnvironment.cpp b/Engine/Source/Programs/UnrealBuildAccelerator/Common/Private/UbaEnvironment.cpp
--- a/Engine/Source/Programs/UnrealBuildAccelerator/Common/Private/UbaEnvironment.cpp
+++ b/Engine/Source/Programs/UnrealBuildAccelerator/Common/Private/UbaEnvironment.cpp
@@ -161,7 +161,6 @@ namespace uba
 				if (info.NextEntryOffset == 0)
 					break;
 				entryIt += info.NextEntryOffset;
-				return 0;
 			}
 			return totalSize;
 		}
