[UBA]
* Default all uba visualizers to always show all traces

diff --git a/Engine/Source/Programs/UnrealBuildAccelerator/Visualizer/Private/Windows/UbaVisualizerApp.cpp b/Engine/Source/Programs/UnrealBuildAccelerator/Visualizer/Private/Windows/UbaVisualizerApp.cpp
--- a/Engine/Source/Programs/UnrealBuildAccelerator/Visualizer/Private/Windows/UbaVisualizerApp.cpp
+++ b/Engine/Source/Programs/UnrealBuildAccelerator/Visualizer/Private/Windows/UbaVisualizerApp.cpp
@@ -313,7 +313,7 @@ int WINAPI wWinMain(_In_ HINSTANCE hInstance, _In_opt_ HINSTANCE hPrevInstance,
 		if (ownerInfo.pid)
 		{
 			configPath.Append('_').Append(ownerInfo.id);
-			showAllTraces = false;
+			//showAllTraces = false; // It seems like showAllTraces is actually what most people want regardless of where a trace is initiated.
 		}
 		configPath.Append(L".toml");
 	}
