Uba: Ensure trace name contains owner info if requested
Show some text when when the visualizer is embedded

#rnx

#Synced-CL 50718640

diff --git a/Engine/Source/Programs/UnrealBuildAccelerator/Host/Private/UbaExports.cpp b/Engine/Source/Programs/UnrealBuildAccelerator/Host/Private/UbaExports.cpp
--- a/Engine/Source/Programs/UnrealBuildAccelerator/Host/Private/UbaExports.cpp
+++ b/Engine/Source/Programs/UnrealBuildAccelerator/Host/Private/UbaExports.cpp
@@ -1102,6 +1102,10 @@ extern "C"
 		localtime_s(&ti, &rawtime);
 		fullName.Appendf(TC("%02i%02i%02i"), ti.tm_hour, ti.tm_min, ti.tm_sec);
 
+		OwnerInfo ownerInfo = GetOwnerInfo();
+		if (ownerInfo.pid)
+			fullName.Appendf(L"_%s%u", ownerInfo.id, ownerInfo.pid);
+
 		trace->StartWriteAndThread(fullName.data, 256*1024*1024, false, channel);
 		return trace;
 	}
diff --git a/Engine/Source/Programs/UnrealBuildAccelerator/Visualizer/Private/UbaVisualizer.cpp b/Engine/Source/Programs/UnrealBuildAccelerator/Visualizer/Private/UbaVisualizer.cpp
--- a/Engine/Source/Programs/UnrealBuildAccelerator/Visualizer/Private/UbaVisualizer.cpp
+++ b/Engine/Source/Programs/UnrealBuildAccelerator/Visualizer/Private/UbaVisualizer.cpp
@@ -116,7 +116,7 @@ namespace uba
 
 	////////////////////////////////////////////////////////////////////////////////////////////////////
 
-	void DrawCenteredText(VisualizerBackend::PaintContext& context, const StringBuffer<>* rows, u32 rowCount, BrushHandle backgroundBrush)
+	void Visualizer::DrawCenteredText(VisualizerBackend::PaintContext& context, const StringBuffer<>* rows, u32 rowCount, BrushHandle backgroundBrush)
 	{
 		SIZE clientSize = context.GetClientSize();
 		for (u32 i=0; i!=rowCount; ++i)
diff --git a/Engine/Source/Programs/UnrealBuildAccelerator/Visualizer/Private/UbaVisualizer.h b/Engine/Source/Programs/UnrealBuildAccelerator/Visualizer/Private/UbaVisualizer.h
--- a/Engine/Source/Programs/UnrealBuildAccelerator/Visualizer/Private/UbaVisualizer.h
+++ b/Engine/Source/Programs/UnrealBuildAccelerator/Visualizer/Private/UbaVisualizer.h
@@ -223,6 +223,7 @@ namespace uba
 		u32 m_workIndex = ~0u;
 
 		struct DrawTextLogger;
+		void DrawCenteredText(VisualizerBackend::PaintContext& context, const StringBuffer<>* rows, u32 rowCount, BrushHandle backgroundBrush);
 	};
 
 	////////////////////////////////////////////////////////////////////////////////////////////////////
diff --git a/Engine/Source/Programs/UnrealBuildAccelerator/Visualizer/Private/Windows/UbaVisualizerWin.cpp b/Engine/Source/Programs/UnrealBuildAccelerator/Visualizer/Private/Windows/UbaVisualizerWin.cpp
--- a/Engine/Source/Programs/UnrealBuildAccelerator/Visualizer/Private/Windows/UbaVisualizerWin.cpp
+++ b/Engine/Source/Programs/UnrealBuildAccelerator/Visualizer/Private/Windows/UbaVisualizerWin.cpp
@@ -528,6 +528,22 @@ namespace uba
 	{
 		if (!m_parentHwnd)
 			Visualizer::PaintEmptyScreen(context);
+		else
+		{
+			StringBuffer<> str[2];
+			if (m_listenChannel.count)
+			{
+				str[0].Appendf(TC("Listening for new sessions on channel '%s'"), m_listenChannel.data);
+				if (!m_config.ShowAllTraces)
+				{
+					OwnerInfo ownerInfo = GetOwnerInfo();
+					if (ownerInfo.pid)
+						str[1].Appendf(TC("Filtering by '_%s%u'"), ownerInfo.id, ownerInfo.pid);
+				}
+			}
+			context.SetFont(m_popupFont);
+			Visualizer::DrawCenteredText(context, str, 2, nullptr);
+		}
 	}
 
 	void VisualizerWin::FinishIncompatible()
