[UBA]
* Fixed linux compile error

diff --git a/Engine/Source/Programs/UnrealBuildAccelerator/Host/Private/UbaExports.cpp b/Engine/Source/Programs/UnrealBuildAccelerator/Host/Private/UbaExports.cpp
--- a/Engine/Source/Programs/UnrealBuildAccelerator/Host/Private/UbaExports.cpp
+++ b/Engine/Source/Programs/UnrealBuildAccelerator/Host/Private/UbaExports.cpp
@@ -1104,7 +1104,7 @@ extern "C"
 
 		OwnerInfo ownerInfo = GetOwnerInfo();
 		if (ownerInfo.pid)
-			fullName.Appendf(L"_%s%u", ownerInfo.id, ownerInfo.pid);
+			fullName.Appendf(TC("_%s%u"), ownerInfo.id, ownerInfo.pid);
 
 		trace->StartWriteAndThread(fullName.data, 256*1024*1024, false, channel);
 		return trace;
