[UBA]
* Fixed compile errors when compiling uba with clang

diff --git a/Engine/Source/Programs/UnrealBuildAccelerator/Visualizer/Private/Windows/UbaVisualizerBackendD2D.h b/Engine/Source/Programs/UnrealBuildAccelerator/Visualizer/Private/Windows/UbaVisualizerBackendD2D.h
--- a/Engine/Source/Programs/UnrealBuildAccelerator/Visualizer/Private/Windows/UbaVisualizerBackendD2D.h
+++ b/Engine/Source/Programs/UnrealBuildAccelerator/Visualizer/Private/Windows/UbaVisualizerBackendD2D.h
@@ -27,7 +27,7 @@ namespace uba
 			return  ((c & 0x000000FF) << 16) | (c & 0x0000FF00) | ((c & 0x00FF0000) >> 16);
 		}
 
-		struct PaintContextD2D : public PaintContext
+		struct PaintContextD2D final : public PaintContext
 		{
 			PaintContextD2D(VisualizerBackendD2D& o, ID2D1RenderTarget& rt, SIZE cs)
 			:	owner(o)
diff --git a/Engine/Source/Programs/UnrealBuildAccelerator/Visualizer/Private/Windows/UbaVisualizerBackendGDI.h b/Engine/Source/Programs/UnrealBuildAccelerator/Visualizer/Private/Windows/UbaVisualizerBackendGDI.h
--- a/Engine/Source/Programs/UnrealBuildAccelerator/Visualizer/Private/Windows/UbaVisualizerBackendGDI.h
+++ b/Engine/Source/Programs/UnrealBuildAccelerator/Visualizer/Private/Windows/UbaVisualizerBackendGDI.h
@@ -10,12 +10,12 @@ namespace uba
 
 	////////////////////////////////////////////////////////////////////////////////////////////////////
 
-	class VisualizerBackendGDI : public VisualizerBackend
+	class VisualizerBackendGDI final : public VisualizerBackend
 	{
 	public:
 		static constexpr int BitmapCacheHeight = 1024*1024;
 
-		struct PaintContextGDI : public PaintContext
+		struct PaintContextGDI final : public PaintContext
 		{
 			virtual SIZE GetClientSize() override { return clientSize; }
 			virtual void SetTextColor(COLORREF color) override { ::SetTextColor(hdc, color); }
