[UBA]
* Reimplemented ProcessImpl::PollStdPipes to handle one pipe being closed but not the other. Also exit the thread when both pipes are closed

diff --git a/Engine/Source/Programs/UnrealBuildAccelerator/Common/Private/UbaProcess.cpp b/Engine/Source/Programs/UnrealBuildAccelerator/Common/Private/UbaProcess.cpp
--- a/Engine/Source/Programs/UnrealBuildAccelerator/Common/Private/UbaProcess.cpp
+++ b/Engine/Source/Programs/UnrealBuildAccelerator/Common/Private/UbaProcess.cpp
@@ -646,7 +646,8 @@ namespace uba
 				Thread pollStdThread([&]()
 				{
 					while (!pollStdDone)
-						PollStdPipes(outReader, errReader, 500);
+						if (!PollStdPipes(outReader, errReader, 500))
+							return false;
 					return true;
 				}, TC("PollStd"));
 				#endif
@@ -2568,51 +2569,105 @@ namespace uba
 #if !PLATFORM_WINDOWS
 	bool ProcessImpl::PollStdPipes(PipeReader& outReader, PipeReader& errReader, int timeoutMs)
 	{
-		if (m_stdOutPipe == -1)
+		if (m_stdOutPipe == -1 && m_stdErrPipe == -1)
 			return false;
 
-		auto pipeGuard = MakeGuard([&]() { close(m_stdOutPipe); m_stdOutPipe = -1; close(m_stdErrPipe); m_stdErrPipe = -1; });
+		auto pipeGuard = MakeGuard([&]()
+		{
+			if (m_stdOutPipe != -1) { close(m_stdOutPipe); m_stdOutPipe = -1; }
+			if (m_stdErrPipe != -1) { close(m_stdErrPipe); m_stdErrPipe = -1; }
+		});
+
+		bool outClosed = (m_stdOutPipe == -1);
+		bool errClosed = (m_stdErrPipe == -1);
 
 		PipeReader* pipeReaders[] = { &outReader, &errReader };
 
-		while (true)
+		while (!outClosed || !errClosed)
 		{
-			pollfd plist[] = { {m_stdOutPipe,POLLIN, 0}, {m_stdErrPipe,POLLIN, 0} };
-			int rval = poll(plist, sizeof_array(plist), timeoutMs);
+			pollfd plist[2];
+			int nfds = 0;
+
+			if (!outClosed)
+			{
+				plist[nfds].fd = m_stdOutPipe;
+				plist[nfds].events = POLLIN;
+				plist[nfds].revents = 0;
+				++nfds;
+			}
+			if (!errClosed)
+			{
+				plist[nfds].fd = m_stdErrPipe;
+				plist[nfds].events = POLLIN;
+				plist[nfds].revents = 0;
+				++nfds;
+			}
+
+			int rval = poll(plist, nfds, timeoutMs);
 			if (rval < 0)
 			{
+				if (errno == EINTR)
+					continue;
 				#if PLATFORM_MAC
-				m_session.m_logger.Error(TC("pipe polling error with %i (%s)"), rval, strerror(errno));
+				m_session.m_logger.Error(TC("pipe polling error (%s)"), strerror(errno));
 				#endif
 				return false;
 			}
 
-			bool hasRead = false;
-			for (int i=0;i!=2;++i)
+			if (rval == 0)
+				break; // timeout
+
+			for (int i = 0; i != nfds; ++i)
 			{
-				if (plist[i].revents & POLLERR) // If there is an error on any of them we hang up
+				pollfd& p = plist[i];
+
+				if (p.revents & POLLERR)
 					return m_session.m_logger.Error(TC("pipe polling error"));
 
-				if (!(plist[i].revents & POLLIN))
+				if (!(p.revents & (POLLIN | POLLHUP)))
 					continue;
 
 				char buffer[1024];
-				if (int bytesRead = read(plist[i].fd, buffer, sizeof_array(buffer) - 1))
+				int bytesRead = read(p.fd, buffer, sizeof_array(buffer) - 1);
+
+				if (bytesRead > 0)
 				{
-					hasRead = true;
 					buffer[bytesRead] = 0;
-					pipeReaders[i]->ReadData(buffer, bytesRead);
+
+					if (p.fd == m_stdOutPipe)
+						outReader.ReadData(buffer, bytesRead);
+					else
+						errReader.ReadData(buffer, bytesRead);
+				}
+				else if (bytesRead == 0)
+				{
+					// EOF
+					if (p.fd == m_stdOutPipe)
+					{
+						close(m_stdOutPipe);
+						m_stdOutPipe = -1;
+						outClosed = true;
+					}
+					else
+					{
+						close(m_stdErrPipe);
+						m_stdErrPipe = -1;
+						errClosed = true;
+					}
+				}
+				else
+				{
+					if (errno != EINTR && errno != EAGAIN)
+					{
+						m_session.m_logger.Error(TC("PollStdPipes read error (%s)"), strerror(errno));
+						return false;
+					}
 				}
 			}
-
-			if (hasRead)
-				continue;
-			if (plist[0].revents & POLLHUP && plist[1].revents & POLLHUP)
-				return false;
-			break;
 		}
+
 		pipeGuard.Cancel();
-		return true;
+		return !outClosed || !errClosed;
 	}
 #endif
 
diff --git a/Engine/Source/Programs/UnrealBuildAccelerator/Core/Private/UbaEventPThread.inl b/Engine/Source/Programs/UnrealBuildAccelerator/Core/Private/UbaEventPThread.inl
--- a/Engine/Source/Programs/UnrealBuildAccelerator/Core/Private/UbaEventPThread.inl
+++ b/Engine/Source/Programs/UnrealBuildAccelerator/Core/Private/UbaEventPThread.inl
@@ -97,6 +97,7 @@ namespace uba
 			}
 			pthread_condattr_destroy(&attrcond);
 
+			m_triggered = TriggerType_None;
 			m_initialized = true;
 			return true;
 		}
