[UBA]
* Disabled ltcg for ubadetours for mac since it seems to produce bad code for 26.2.0+

diff --git a/Engine/Source/Programs/UnrealBuildAccelerator/UbaDetours.Target.cs b/Engine/Source/Programs/UnrealBuildAccelerator/UbaDetours.Target.cs
--- a/Engine/Source/Programs/UnrealBuildAccelerator/UbaDetours.Target.cs
+++ b/Engine/Source/Programs/UnrealBuildAccelerator/UbaDetours.Target.cs
@@ -11,5 +11,12 @@ public class UbaDetoursTarget : TargetRules
 		bShouldCompileAsDLL = true;
 		UbaAgentTarget.CommonUbaSettings(this, Target);
 		GlobalDefinitions.Add("_CRT_SECURE_CPP_OVERLOAD_SECURE_NAMES=0");
+
+		// For some reason xcode 26.2.0+ produce a bad dylib that fails to load when ltcg is enabled
+		if (Target.Platform.IsInGroup(UnrealPlatformGroup.Apple))
+		{
+			bAllowLTCG = false;
+		}
+
 	}
 }
