[UBA]
* Fixed bug in dependency crawler when parsing c code that has #include "nameoffile where the end quote is missing

diff --git a/Engine/Source/Programs/UnrealBuildAccelerator/Common/Public/UbaDependencyCrawler.h b/Engine/Source/Programs/UnrealBuildAccelerator/Common/Public/UbaDependencyCrawler.h
--- a/Engine/Source/Programs/UnrealBuildAccelerator/Common/Public/UbaDependencyCrawler.h
+++ b/Engine/Source/Programs/UnrealBuildAccelerator/Common/Public/UbaDependencyCrawler.h
@@ -662,22 +662,20 @@ namespace uba
 					StringBuffer<> include;
 					bool isQuote = false;
 
-					if (*it == '\"')
+					if (*it == '\"' || *it == '<')
 					{
+						isQuote = *it == '\"';
+						char endChar = isQuote ? '\"' : '>';
 						++it;
 						const char* includeStart = it;
-						const char* includeEnd = strchr(it, '\"');
-						it = includeEnd + 1;
-						include.Append(includeStart, u32(includeEnd - includeStart));
-						isQuote =true;
-					}
-					else if (*it == '<')
-					{
-						++it;
-						const char* includeStart = it;
-						const char* includeEnd = strchr(it, '>');
+						const char* includeEnd = strchr(it, endChar);
+						if (!includeEnd)
+							break;
+						u64 includeLen = includeEnd - includeStart;
+						if (includeLen >= include.capacity)
+							break;
 						it = includeEnd + 1;
-						include.Append(includeStart, u32(includeEnd - includeStart));
+						include.Append(includeStart, u32(includeLen));
 					}
 					else
 					{
