[UBA]
* Changed so network server backends connections are disconnected first in one loop and then waited on in another.. since disconnect can sometimes take some time this becomes faster when having lots of connections

diff --git a/Engine/Source/Programs/UnrealBuildAccelerator/Common/Private/UbaNetworkServer.cpp b/Engine/Source/Programs/UnrealBuildAccelerator/Common/Private/UbaNetworkServer.cpp
--- a/Engine/Source/Programs/UnrealBuildAccelerator/Common/Private/UbaNetworkServer.cpp
+++ b/Engine/Source/Programs/UnrealBuildAccelerator/Common/Private/UbaNetworkServer.cpp
@@ -166,7 +166,8 @@ namespace uba
 
 		~Connection()
 		{
-			Stop();
+			BeginStop();
+			EndStop();
 			if (m_cryptoKey)
 				Crypto::DestroyKey(m_cryptoKey);
 		}
@@ -183,10 +184,14 @@ namespace uba
 			//	m_server.m_logger.Detail(TC("Connection %u disconnected while it has %u active workers (%s)"), m_id, activeWorkerCount, reason);
 		}
 
-		bool Stop()
+		bool BeginStop()
 		{
 			Disconnect(TC("Stop"));
+			return true;
+		}
 
+		bool EndStop()
+		{
 			u64 startTime = GetTime();
 			bool waitReported = false;
 			while (m_activeWorkerCount)
@@ -215,16 +220,15 @@ namespace uba
 			if (waitReported)
 				m_server.m_logger.Info(TC("Connection %u wait for active workers ended after %s"), m_id, TimeToText(GetTime() - startTime).str);
 
-			if (!m_disconnectCallbackCalled.IsSet(30*1000)) // This should never time out!
-			{
-				m_backend.SetDisconnectCallback(m_backendConnection, nullptr, [](void* context, const Guid& connectionUid, void* connection) {});
-				m_backend.SetRecvCallbacks(m_backendConnection, nullptr, 0, [](void*, const Guid&, u8*, void*&, u8*&, u32&) { return false; }, nullptr, TC("Stopped"));
-				m_server.m_logger.Info(TC("Disconnect callback event timed out. This should never happen!!"));
-				#if UBA_DEBUG
-				PrintAllCallstacks(m_server.m_logger);
-				#endif
-			}
+			if (m_disconnectCallbackCalled.IsSet(30*1000)) // This should never time out!
+				return true;
 
+			m_backend.SetDisconnectCallback(m_backendConnection, nullptr, [](void* context, const Guid& connectionUid, void* connection) {});
+			m_backend.SetRecvCallbacks(m_backendConnection, nullptr, 0, [](void*, const Guid&, u8*, void*&, u8*&, u32&) { return false; }, nullptr, TC("Stopped"));
+			m_server.m_logger.Info(TC("Disconnect callback event timed out. This should never happen!!"));
+			#if UBA_DEBUG
+			PrintAllCallstacks(m_server.m_logger);
+			#endif
 			return true;
 		}
 
@@ -881,9 +885,12 @@ namespace uba
 			lock.Leave();
 
 			bool success = true;
+			for (auto& c : connections)
+				success = c.BeginStop() && success;
+
 			for (auto& c : connections)
 			{
-				success = c.Stop() && success;
+				success = c.EndStop() && success;
 				m_sendTimer += c.m_sendTimer;
 				m_encryptTimer += c.m_encryptTimer;
 				m_decryptTimer += c.m_decryptTimer;
