[UBA]
* Added comment to help investigate code path

diff --git a/Engine/Source/Programs/UnrealBuildAccelerator/Common/Private/UbaStorage.cpp b/Engine/Source/Programs/UnrealBuildAccelerator/Common/Private/UbaStorage.cpp
--- a/Engine/Source/Programs/UnrealBuildAccelerator/Common/Private/UbaStorage.cpp
+++ b/Engine/Source/Programs/UnrealBuildAccelerator/Common/Private/UbaStorage.cpp
@@ -735,6 +735,9 @@ namespace uba
 				//#if UBA_DEBUG
 				//return m_logger.Error(TC("This should not happen. (%s)"), fileName);
 				//#endif
+
+				// Saw this on d:\devel\fn\engine\plugins\runtime\audiowidgets\intermediate\build\win64\unrealeditor\inc\audiowidgets\uht\audiospectrogramviewport.generated.h
+				// but build finished properly with no issues (2026-02-12)
 				m_logger.Info(TC("Investigate how we get here. (%s)"), fileName);
 
 				SCOPED_FUTEX(dcc.lock, dccLock);
