[UBA]
* Some more harding of dependency crawler in case files have bad formatting

diff --git a/Engine/Source/Programs/UnrealBuildAccelerator/Common/Public/UbaDependencyCrawler.h b/Engine/Source/Programs/UnrealBuildAccelerator/Common/Public/UbaDependencyCrawler.h
--- a/Engine/Source/Programs/UnrealBuildAccelerator/Common/Public/UbaDependencyCrawler.h
+++ b/Engine/Source/Programs/UnrealBuildAccelerator/Common/Public/UbaDependencyCrawler.h
@@ -715,17 +715,21 @@ namespace uba
 						u32 defineLen = u32(defineEnd - defineBegin);
 						if (argBegin)
 						{
-							if (strncmp(defineBegin, "UE_INLINE_GENERATED_CPP_BY_NAME", defineLen) == 0)
+							u64 argLen = argEnd - argBegin;
+							if (argLen < include.capacity - 32)
 							{
-								include.Append(argBegin, u32(argEnd - argBegin)).Append(TCV(".gen.cpp"));
-							}
-							else if (strncmp(defineBegin, "COMPILED_PLATFORM_HEADER", defineLen) == 0)
-							{
-								if (!instance.platform.empty())
+								if (strncmp(defineBegin, "UE_INLINE_GENERATED_CPP_BY_NAME", defineLen) == 0)
+								{
+									include.Append(argBegin, u32(argLen)).Append(TCV(".gen.cpp"));
+								}
+								else if (strncmp(defineBegin, "COMPILED_PLATFORM_HEADER", defineLen) == 0)
 								{
-									if (!instance.platformIsExtension)
-										include.Append(instance.platform).Append(PathSeparator);
-									include.Append(instance.platform).Append(argBegin, u32(argEnd - argBegin));
+									if (!instance.platform.empty())
+									{
+										if (!instance.platformIsExtension)
+											include.Append(instance.platform).Append(PathSeparator);
+										include.Append(instance.platform).Append(argBegin, u32(argLen));
+									}
 								}
 							}
 						}
