[UBA]
* Fixed so logic handling std out/err logging  work properly when chunks are not terminated

diff --git a/Engine/Source/Programs/UnrealBuildAccelerator/Common/Private/UbaProcess.cpp b/Engine/Source/Programs/UnrealBuildAccelerator/Common/Private/UbaProcess.cpp
--- a/Engine/Source/Programs/UnrealBuildAccelerator/Common/Private/UbaProcess.cpp
+++ b/Engine/Source/Programs/UnrealBuildAccelerator/Common/Private/UbaProcess.cpp
@@ -72,21 +72,27 @@ namespace uba
 		void ReadData(char* buf, u32 readCount)
 		{
 			char* startPos = buf;
-			while (true)
+			char* end = buf + readCount;
+
+			while (startPos < end)
 			{
-				char* endOfLine = strchr(startPos, '\n');
+				char* endOfLine = (char*)memchr(startPos, '\n', end - startPos);
 				if (!endOfLine)
 				{
-					currentString.append(TString(startPos, startPos + strlen(startPos)));
+					currentString.append(TString(startPos, end));
 					return;
 				}
-				char* newStart = endOfLine + 1;
-				if (endOfLine > buf && endOfLine[-1] == '\r')
-					--endOfLine;
-				currentString.append(TString(startPos, endOfLine));
+
+				char* lineEnd = endOfLine;
+				if (lineEnd > startPos && lineEnd[-1] == '\r')
+					--lineEnd;
+
+				currentString.append(TString(startPos, lineEnd));
+
 				process.InternalLogLine(false, TString(currentString), logType);
 				currentString.clear();
-				startPos = newStart;
+
+				startPos = endOfLine + 1;
 			}
 		}
 
