Making UBA modules as requires platform sdk

#Synced-CL 50647103

diff --git a/Engine/Source/Programs/UnrealBuildAccelerator/Agent/UbaAgent.Build.cs b/Engine/Source/Programs/UnrealBuildAccelerator/Agent/UbaAgent.Build.cs
--- a/Engine/Source/Programs/UnrealBuildAccelerator/Agent/UbaAgent.Build.cs
+++ b/Engine/Source/Programs/UnrealBuildAccelerator/Agent/UbaAgent.Build.cs
@@ -9,6 +9,8 @@ public class UbaAgent : ModuleRules
 	{
 		PrivatePCHHeaderFile = "../Core/Public/UbaCorePch.h";
 
+		bRequiresPlatformSDK = true;
+
 		CppCompileWarningSettings.UnsafeTypeCastWarningLevel = WarningLevel.Error;
 		StaticAnalyzerDisabledCheckers.Clear();
 
diff --git a/Engine/Source/Programs/UnrealBuildAccelerator/CacheService/UbaCacheService.Build.cs b/Engine/Source/Programs/UnrealBuildAccelerator/CacheService/UbaCacheService.Build.cs
--- a/Engine/Source/Programs/UnrealBuildAccelerator/CacheService/UbaCacheService.Build.cs
+++ b/Engine/Source/Programs/UnrealBuildAccelerator/CacheService/UbaCacheService.Build.cs
@@ -8,6 +8,8 @@ public class UbaCacheService : ModuleRules
 	public UbaCacheService(ReadOnlyTargetRules Target) : base(Target)
 	{
 		PrivatePCHHeaderFile = "../Core/Public/UbaCorePch.h";
+		
+		bRequiresPlatformSDK = true;
 
 		CppCompileWarningSettings.UnsafeTypeCastWarningLevel = WarningLevel.Error;
 		StaticAnalyzerDisabledCheckers.Clear();
diff --git a/Engine/Source/Programs/UnrealBuildAccelerator/Cli/UbaCli.Build.cs b/Engine/Source/Programs/UnrealBuildAccelerator/Cli/UbaCli.Build.cs
--- a/Engine/Source/Programs/UnrealBuildAccelerator/Cli/UbaCli.Build.cs
+++ b/Engine/Source/Programs/UnrealBuildAccelerator/Cli/UbaCli.Build.cs
@@ -12,6 +12,8 @@ public class UbaCli : ModuleRules
 		CppCompileWarningSettings.UnsafeTypeCastWarningLevel = WarningLevel.Error;
 		StaticAnalyzerDisabledCheckers.Clear();
 
+		bRequiresPlatformSDK = true;
+
 		PrivateDependencyModuleNames.AddRange(new string[] {
 			"UbaCommon",
 		});
diff --git a/Engine/Source/Programs/UnrealBuildAccelerator/Common/UbaCommon.Build.cs b/Engine/Source/Programs/UnrealBuildAccelerator/Common/UbaCommon.Build.cs
--- a/Engine/Source/Programs/UnrealBuildAccelerator/Common/UbaCommon.Build.cs
+++ b/Engine/Source/Programs/UnrealBuildAccelerator/Common/UbaCommon.Build.cs
@@ -9,6 +9,8 @@ public class UbaCommon : ModuleRules
 	{
 		PrivatePCHHeaderFile = "../Core/Public/UbaCorePch.h";
 
+		bRequiresPlatformSDK = true;
+
 		CppCompileWarningSettings.UnsafeTypeCastWarningLevel = WarningLevel.Error;
 		StaticAnalyzerDisabledCheckers.Clear();
 
diff --git a/Engine/Source/Programs/UnrealBuildAccelerator/Core/UbaCore.Build.cs b/Engine/Source/Programs/UnrealBuildAccelerator/Core/UbaCore.Build.cs
--- a/Engine/Source/Programs/UnrealBuildAccelerator/Core/UbaCore.Build.cs
+++ b/Engine/Source/Programs/UnrealBuildAccelerator/Core/UbaCore.Build.cs
@@ -10,6 +10,8 @@ public class UbaCore : ModuleRules
 		CppCompileWarningSettings.UnsafeTypeCastWarningLevel = WarningLevel.Error;
 		StaticAnalyzerDisabledCheckers.Clear();
 
+		bRequiresPlatformSDK = true;
+
 		PrivatePCHHeaderFile = "Public/UbaCorePch.h";
 		SharedPCHHeaderFile = "Public/UbaCorePch.h";
 
diff --git a/Engine/Source/Programs/UnrealBuildAccelerator/Detours/UbaDetours.Build.cs b/Engine/Source/Programs/UnrealBuildAccelerator/Detours/UbaDetours.Build.cs
--- a/Engine/Source/Programs/UnrealBuildAccelerator/Detours/UbaDetours.Build.cs
+++ b/Engine/Source/Programs/UnrealBuildAccelerator/Detours/UbaDetours.Build.cs
@@ -12,6 +12,8 @@ public class UbaDetours : ModuleRules
 		StaticAnalyzerDisabledCheckers.Clear();
 		bUseUnity = false;
 
+		bRequiresPlatformSDK = true;
+
 		PrivateDependencyModuleNames.AddRange(new string[] {
 			"UbaCore",
 			"OodleDataCompression",
diff --git a/Engine/Source/Programs/UnrealBuildAccelerator/Host/UbaHost.Build.cs b/Engine/Source/Programs/UnrealBuildAccelerator/Host/UbaHost.Build.cs
--- a/Engine/Source/Programs/UnrealBuildAccelerator/Host/UbaHost.Build.cs
+++ b/Engine/Source/Programs/UnrealBuildAccelerator/Host/UbaHost.Build.cs
@@ -9,6 +9,8 @@ public class UbaHost : ModuleRules
 	{
 		PrivatePCHHeaderFile = "../Core/Public/UbaCorePch.h";
 
+		bRequiresPlatformSDK = true;
+
 		CppCompileWarningSettings.UnsafeTypeCastWarningLevel = WarningLevel.Error;
 		StaticAnalyzerDisabledCheckers.Clear();
 
diff --git a/Engine/Source/Programs/UnrealBuildAccelerator/ObjTool/UbaObjTool.Build.cs b/Engine/Source/Programs/UnrealBuildAccelerator/ObjTool/UbaObjTool.Build.cs
--- a/Engine/Source/Programs/UnrealBuildAccelerator/ObjTool/UbaObjTool.Build.cs
+++ b/Engine/Source/Programs/UnrealBuildAccelerator/ObjTool/UbaObjTool.Build.cs
@@ -9,6 +9,8 @@ public class UbaObjTool : ModuleRules
 	{
 		PrivatePCHHeaderFile = "../Core/Public/UbaCorePch.h";
 
+		bRequiresPlatformSDK = true;
+
 		CppCompileWarningSettings.UnsafeTypeCastWarningLevel = WarningLevel.Error;
 		StaticAnalyzerDisabledCheckers.Clear();
 
diff --git a/Engine/Source/Programs/UnrealBuildAccelerator/StorageProxy/UbaStorageProxy.Build.cs b/Engine/Source/Programs/UnrealBuildAccelerator/StorageProxy/UbaStorageProxy.Build.cs
--- a/Engine/Source/Programs/UnrealBuildAccelerator/StorageProxy/UbaStorageProxy.Build.cs
+++ b/Engine/Source/Programs/UnrealBuildAccelerator/StorageProxy/UbaStorageProxy.Build.cs
@@ -12,6 +12,8 @@ public class UbaStorageProxy : ModuleRules
 		CppCompileWarningSettings.UnsafeTypeCastWarningLevel = WarningLevel.Error;
 		StaticAnalyzerDisabledCheckers.Clear();
 
+		bRequiresPlatformSDK = true;
+
 		PrivateDependencyModuleNames.AddRange(new string[] {
 			"UbaCore",
 			"UbaCommon",
diff --git a/Engine/Source/Programs/UnrealBuildAccelerator/Test/UbaTest.Build.cs b/Engine/Source/Programs/UnrealBuildAccelerator/Test/UbaTest.Build.cs
--- a/Engine/Source/Programs/UnrealBuildAccelerator/Test/UbaTest.Build.cs
+++ b/Engine/Source/Programs/UnrealBuildAccelerator/Test/UbaTest.Build.cs
@@ -9,6 +9,8 @@ public class UbaTest : ModuleRules
 	{
 		PrivatePCHHeaderFile = "../Core/Public/UbaCorePch.h";
 		IWYUSupport = IWYUSupport.None;
+		
+		bRequiresPlatformSDK = true;
 
 		CppCompileWarningSettings.UnsafeTypeCastWarningLevel = WarningLevel.Error;
 
diff --git a/Engine/Source/Programs/UnrealBuildAccelerator/TestApp/UbaTestApp.Build.cs b/Engine/Source/Programs/UnrealBuildAccelerator/TestApp/UbaTestApp.Build.cs
--- a/Engine/Source/Programs/UnrealBuildAccelerator/TestApp/UbaTestApp.Build.cs
+++ b/Engine/Source/Programs/UnrealBuildAccelerator/TestApp/UbaTestApp.Build.cs
@@ -8,6 +8,8 @@ public class UbaTestApp : ModuleRules
 	public UbaTestApp(ReadOnlyTargetRules Target) : base(Target)
 	{
 		CppCompileWarningSettings.UnsafeTypeCastWarningLevel = WarningLevel.Error;
+		
+		bRequiresPlatformSDK = true;
 
 		PrivateDefinitions.AddRange(new string[] {
 			"_CONSOLE",
diff --git a/Engine/Source/Programs/UnrealBuildAccelerator/Visualizer/UbaVisualizer.Build.cs b/Engine/Source/Programs/UnrealBuildAccelerator/Visualizer/UbaVisualizer.Build.cs
--- a/Engine/Source/Programs/UnrealBuildAccelerator/Visualizer/UbaVisualizer.Build.cs
+++ b/Engine/Source/Programs/UnrealBuildAccelerator/Visualizer/UbaVisualizer.Build.cs
@@ -9,6 +9,8 @@ public class UbaVisualizer : ModuleRules
 	{
 		PrivatePCHHeaderFile = "../Core/Public/UbaCorePch.h";
 
+		bRequiresPlatformSDK = true;
+
 		CppCompileWarningSettings.UnsafeTypeCastWarningLevel = WarningLevel.Error;
 		StaticAnalyzerDisabledCheckers.Clear();
 
