[UBA]
* Fixed so radlink can run on wine

diff --git a/Engine/Source/Programs/UnrealBuildAccelerator/Detours/Private/Windows/UbaDetoursFunctionsKernelBase.inl b/Engine/Source/Programs/UnrealBuildAccelerator/Detours/Private/Windows/UbaDetoursFunctionsKernelBase.inl
--- a/Engine/Source/Programs/UnrealBuildAccelerator/Detours/Private/Windows/UbaDetoursFunctionsKernelBase.inl
+++ b/Engine/Source/Programs/UnrealBuildAccelerator/Detours/Private/Windows/UbaDetoursFunctionsKernelBase.inl
@@ -629,6 +629,17 @@ DWORD Detoured_GetSecurityInfo(HANDLE handle, SE_OBJECT_TYPE ObjectType, SECURIT
 	return True_GetSecurityInfo(handle, ObjectType, SecurityInfo, ppsidOwner, ppsidGroup, ppDacl, ppSacl, ppSecurityDescriptor);
 }
 
+BOOL Detoured_InitializeSynchronizationBarrier(LPSYNCHRONIZATION_BARRIER lpBarrier, LONG lTotalThreads, LONG lSpinCount)
+{
+	Rpc_WriteLogf(TC("You need a newer version of wine to be able to call InitializeSynchronizationBarrier (%s)"), g_virtualApplication.data);
+	return false;
+}
+
+BOOL Detoured_DeleteSynchronizationBarrier(LPSYNCHRONIZATION_BARRIER lpBarrier)
+{
+	return false;
+}
+
 StringBuffer<32 * 1024> g_stdFile;
 ReaderWriterLock& g_stdFileLock = *new ReaderWriterLock();
 
@@ -4205,6 +4216,51 @@ BOOL Detoured_UnmapViewOfFile2(HANDLE Process, PVOID BaseAddress, ULONG UnmapFla
 	return res;
 }
 
+BOOL Detoured_GetComputerNameA(LPSTR lpBuffer, LPDWORD nSize)
+{
+	DETOURED_CALL(GetComputerNameA);
+	auto res = True_GetComputerNameA(lpBuffer, nSize);
+	DEBUG_LOG_TRUE(L"GetComputerNameA", L"-> %s", ToString(res));
+	return res;
+}
+
+HMODULE Detoured_GetModuleHandleW(LPCWSTR lpModuleName)
+{
+	DETOURED_CALL(GetModuleHandleW);
+	auto res = True_GetModuleHandleW(lpModuleName);
+	DEBUG_LOG_TRUE(L"GetModuleHandleW", L"%s -> %llx", lpModuleName, u64(res));
+	return res;
+}
+
+VOID Detoured_GetStartupInfoW(LPSTARTUPINFOW lpStartupInfo)
+{
+	DETOURED_CALL(GetStartupInfoW);
+	DEBUG_LOG_TRUE(L"GetStartupInfoW", L"");
+	True_GetStartupInfoW(lpStartupInfo);
+}
+
+VOID Detoured_GetSystemInfo(LPSYSTEM_INFO lpSystemInfo)
+{
+	DETOURED_CALL(GetSystemInfo);
+	DEBUG_LOG_TRUE(L"GetSystemInfo", L"");
+	True_GetSystemInfo(lpSystemInfo);
+}
+
+VOID Detoured_GetSystemTimeAsFileTime(LPFILETIME lpSystemTimeAsFileTime)
+{
+	DETOURED_CALL(GetSystemTimeAsFileTime);
+	DEBUG_LOG_TRUE(L"GetSystemTimeAsFileTime", L"");
+	True_GetSystemTimeAsFileTime(lpSystemTimeAsFileTime);
+}
+
+SIZE_T Detoured_GetLargePageMinimum()
+{
+	DETOURED_CALL(GetLargePageMinimum);
+	auto res = True_GetLargePageMinimum();
+	DEBUG_LOG_TRUE(L"GetLargePageMinimum", L"-> %llu", u64(res));
+	return res;
+}
+
 // These functions require a very new version of win api (1-1-7).. so skip detouring these for now
 //HANDLE Detoured_CreateFileMapping2(HANDLE File, SECURITY_ATTRIBUTES* SecurityAttributes, ULONG DesiredAccess, ULONG PageProtection, ULONG AllocationAttributes, ULONG64 MaximumSize, PCWSTR Name, MEM_EXTENDED_PARAMETER* ExtendedParameters, ULONG ParameterCount)
 //{
diff --git a/Engine/Source/Programs/UnrealBuildAccelerator/Detours/Private/Windows/UbaDetoursFunctionsWin.cpp b/Engine/Source/Programs/UnrealBuildAccelerator/Detours/Private/Windows/UbaDetoursFunctionsWin.cpp
--- a/Engine/Source/Programs/UnrealBuildAccelerator/Detours/Private/Windows/UbaDetoursFunctionsWin.cpp
+++ b/Engine/Source/Programs/UnrealBuildAccelerator/Detours/Private/Windows/UbaDetoursFunctionsWin.cpp
@@ -582,10 +582,64 @@ void SendExitMessage(DWORD exitCode, u64 startTime)
 	writer.Flush(g_isChild);
 }
 
+void* FindIATEntryByName(HMODULE moduleBase, const char* importDllName, const char* importFuncName)
+{
+	u8* base = reinterpret_cast<u8*>(moduleBase);
+	auto* dos = reinterpret_cast<IMAGE_DOS_HEADER*>(base);
+	if (dos->e_magic != IMAGE_DOS_SIGNATURE)
+		return nullptr;
+	auto* nt = reinterpret_cast<IMAGE_NT_HEADERS64*>(base + dos->e_lfanew);
+	if (nt->Signature != IMAGE_NT_SIGNATURE)
+		return nullptr;
+	const auto& dir = nt->OptionalHeader.DataDirectory[IMAGE_DIRECTORY_ENTRY_IMPORT];
+	if (!dir.VirtualAddress || !dir.Size)
+		return nullptr;
+	for (auto* imp = reinterpret_cast<IMAGE_IMPORT_DESCRIPTOR*>(base + dir.VirtualAddress); imp->Name; ++imp)
+	{
+		const char* dllName = reinterpret_cast<const char*>(base + imp->Name);
+		if (_stricmp(dllName, importDllName) != 0)
+			continue;
+		auto* origThunk = reinterpret_cast<IMAGE_THUNK_DATA64*>(base + imp->OriginalFirstThunk);
+		auto* firstThunk = reinterpret_cast<IMAGE_THUNK_DATA64*>(base + imp->FirstThunk);
+		if (!imp->OriginalFirstThunk) // Some binaries may have OriginalFirstThunk == 0; then use FirstThunk for names too.
+			origThunk = firstThunk;
+		for (; origThunk->u1.AddressOfData; ++origThunk, ++firstThunk)
+		{
+			if (origThunk->u1.Ordinal & IMAGE_ORDINAL_FLAG64)
+				continue; // ordinal import
+			auto* importByName = reinterpret_cast<IMAGE_IMPORT_BY_NAME*>(base + origThunk->u1.AddressOfData);
+			const char* funcName = reinterpret_cast<const char*>(importByName->Name);
+			if (strcmp(funcName, importFuncName) == 0)
+				return &firstThunk->u1.Function; // Return address of the IAT slot (the pointer you must overwrite)
+		}
+	}
+	return nullptr;
+}
+
+bool PatchIAT(HMODULE exeBase, const char* importDllName, const char* importFuncName, void* newFunc, void** oldFuncOut = nullptr)
+{
+	void* slot = FindIATEntryByName(exeBase, importDllName, importFuncName);
+	if (!slot)
+		return false;
+	DWORD oldProt = 0;
+	if (!VirtualProtect(slot, sizeof(void*), PAGE_READWRITE, &oldProt))
+		return false;
+	void* old = *reinterpret_cast<void**>(slot);
+	*reinterpret_cast<void**>(slot) = newFunc;
+	VirtualProtect(slot, sizeof(void*), oldProt, &oldProt);
+	FlushInstructionCache(GetCurrentProcess(), slot, sizeof(void*));
+	if (oldFuncOut)
+		*oldFuncOut = old;
+	return true;
+}
+
 void DetourAttachFunction(void** trueFunc, void* detouredFunc, const char* funcName)
 {
 	if (!*trueFunc)
+	{
+		//Rpc_WriteLogf(L"GetProcAddress Failed to find %hs", funcName);
 		return;
+	}
 	auto error = DetourAttach(trueFunc, detouredFunc);
 	if (error == NO_ERROR)
 		return;
@@ -642,6 +696,11 @@ int DetourAttachFunctions(bool runningRemote)
 	if (HMODULE moduleHandle = GetModuleHandleW(L"kernel32.dll"))
 	{
 		DETOURED_FUNCTIONS_KERNEL32
+
+		// Old wine versions do not have this implemented. Call to another function that shows a message
+		if (IsRunningWine())
+			if (!GetProcAddress(moduleHandle, "InitializeSynchronizationBarrier"))
+				PatchIAT(GetModuleHandleW(nullptr), "KERNEL32.dll", "InitializeSynchronizationBarrier", (void*)&Detoured_InitializeSynchronizationBarrier);
 	}
 
 	if (HMODULE moduleHandle = GetModuleHandleW(L"ntdll.dll"))
diff --git a/Engine/Source/Programs/UnrealBuildAccelerator/Detours/Private/Windows/UbaDetoursFunctionsWin.h b/Engine/Source/Programs/UnrealBuildAccelerator/Detours/Private/Windows/UbaDetoursFunctionsWin.h
--- a/Engine/Source/Programs/UnrealBuildAccelerator/Detours/Private/Windows/UbaDetoursFunctionsWin.h
+++ b/Engine/Source/Programs/UnrealBuildAccelerator/Detours/Private/Windows/UbaDetoursFunctionsWin.h
@@ -321,6 +321,12 @@
 	DETOURED_FUNCTION(GetMappedFileNameW) \
 	DETOURED_FUNCTION(IsProcessorFeaturePresent) \
 	DETOURED_FUNCTION(UnmapViewOfFile2) \
+	DETOURED_FUNCTION(GetComputerNameA) \
+	DETOURED_FUNCTION(GetModuleHandleW) \
+	DETOURED_FUNCTION(GetStartupInfoW) \
+	DETOURED_FUNCTION(GetSystemInfo) \
+	DETOURED_FUNCTION(GetSystemTimeAsFileTime) \
+	DETOURED_FUNCTION(GetLargePageMinimum) \
 	//DETOURED_FUNCTION(VirtualFree) \
 	//DETOURED_FUNCTION(BaseThreadInitThunk) \
 	//DETOURED_FUNCTION(VirtualAllocEx) \
