[UBA]
* Added support for disabling remote execution on individual helpers. helpers are identified using the same string as provided when starting ubaagent with -name=<name>

diff --git a/Engine/Source/Programs/UnrealBuildAccelerator/Common/Private/UbaSessionServer.cpp b/Engine/Source/Programs/UnrealBuildAccelerator/Common/Private/UbaSessionServer.cpp
--- a/Engine/Source/Programs/UnrealBuildAccelerator/Common/Private/UbaSessionServer.cpp
+++ b/Engine/Source/Programs/UnrealBuildAccelerator/Common/Private/UbaSessionServer.cpp
@@ -611,6 +611,18 @@ namespace uba
 		//m_trace.RemoteExecutionDisabled();
 	}
 
+	void SessionServer::DisableRemoteExecutionOnAgent(const tchar* agentName)
+	{
+		SCOPED_CRITICAL_SECTION(m_remoteProcessAndSessionLock, queueLock);
+		for (auto s : m_clientSessions)
+		{
+			if (s->name != agentName)
+				continue;
+			s->remoteExecutionEnabled = false;
+			break;
+		}
+	}
+
 	void SessionServer::SetCustomCasKeyFromTrackedInputs(const tchar* fileName_, const tchar* workingDir_, const u8* trackedInputs, u32 trackedInputsBytes)
 	{
 		if (!m_customCasKeysEnabled)
@@ -1679,6 +1691,9 @@ namespace uba
 		for (auto& lostProcess : processesLostRace)
 			writer.WriteU32(lostProcess);
 
+		if ((!m_remoteExecutionEnabled || !session.remoteExecutionEnabled) && m_queuedRemoteProcesses.empty())
+			availableWeight = 0.0f;
+
 		Vector<RemoteProcess*> processesWithKnownInputsToSend;
 
 		float weightLeft = availableWeight;
@@ -1776,7 +1791,7 @@ namespace uba
 		//if (addCount)
 		//	m_logger.Debug(TC("Gave %u processes to %s using up %.1f weight out of %.1f available"), addCount, session.name.c_str(), availableWeight - weightLeft, availableWeight);
 
-		bool remoteExecutionEnabled = m_remoteExecutionEnabled || !m_queuedRemoteProcesses.empty();
+		bool remoteExecutionEnabled = (m_remoteExecutionEnabled && session.remoteExecutionEnabled) || !m_queuedRemoteProcesses.empty();
 		if (!remoteExecutionEnabled)
 		{
 			if (session.enabled)
diff --git a/Engine/Source/Programs/UnrealBuildAccelerator/Common/Public/UbaSessionServer.h b/Engine/Source/Programs/UnrealBuildAccelerator/Common/Public/UbaSessionServer.h
--- a/Engine/Source/Programs/UnrealBuildAccelerator/Common/Public/UbaSessionServer.h
+++ b/Engine/Source/Programs/UnrealBuildAccelerator/Common/Public/UbaSessionServer.h
@@ -36,6 +36,7 @@ namespace uba
 		void DisableRemoteExecution();
 		bool IsRemoteExecutionDisabled();
 		void ReenableRemoteExecution();
+		void DisableRemoteExecutionOnAgent(const tchar* agentName);
 
 		// This can be used to set a custom cas key based on a file and its inputs. Can be used for non-deterministic outputs to still be able to use cached cas content on clients
 		// For example, when a pch is built on the host we can use all the input involved to create the pch and use that as the cas key for the resulting huge file
@@ -193,6 +194,7 @@ namespace uba
 			bool hasNotification = false;
 			bool isArm = false;
 			bool isLinux = false; // Only used for non-linux
+			bool remoteExecutionEnabled = true;
 		};
 		Vector<ClientSession*> m_clientSessions;
 
diff --git a/Engine/Source/Programs/UnrealBuildAccelerator/Host/Private/UbaExports.cpp b/Engine/Source/Programs/UnrealBuildAccelerator/Host/Private/UbaExports.cpp
--- a/Engine/Source/Programs/UnrealBuildAccelerator/Host/Private/UbaExports.cpp
+++ b/Engine/Source/Programs/UnrealBuildAccelerator/Host/Private/UbaExports.cpp
@@ -686,6 +686,11 @@ extern "C"
 		server->GetTrace().TaskEnd(taskId, success);
 	}
 
+	void SessionServer_DisableRemoteExecutionOnAgent(uba::SessionServer* server, const uba::tchar* agentName)
+	{
+		server->DisableRemoteExecutionOnAgent(agentName);
+	}
+
 	void SessionServer_Destroy(uba::SessionServer* server)
 	{
 		if (server)
diff --git a/Engine/Source/Programs/UnrealBuildAccelerator/Host/Public/UbaExports.h b/Engine/Source/Programs/UnrealBuildAccelerator/Host/Public/UbaExports.h
--- a/Engine/Source/Programs/UnrealBuildAccelerator/Host/Public/UbaExports.h
+++ b/Engine/Source/Programs/UnrealBuildAccelerator/Host/Public/UbaExports.h
@@ -154,6 +154,7 @@ extern "C"
 	UBA_API uba::u32 SessionServer_TaskBegin(uba::SessionServer* server, const uba::tchar* description, const uba::tchar* details, uba::u32 color);
 	UBA_API void SessionServer_TaskHint(uba::SessionServer* server, uba::u32 taskId, const uba::tchar* hint);
 	UBA_API void SessionServer_TaskEnd(uba::SessionServer* server, uba::u32 taskId, bool success);
+	UBA_API void SessionServer_DisableRemoteExecutionOnAgent(uba::SessionServer* server, const uba::tchar* agentName);
 	UBA_API void SessionServer_Destroy(uba::SessionServer* server);
 
 	// Scheduler
