[UBA]
* Changed DisableRemoteExecutionOnAgent to return bool which is true if successful

diff --git a/Engine/Source/Programs/UnrealBuildAccelerator/Common/Private/UbaSessionServer.cpp b/Engine/Source/Programs/UnrealBuildAccelerator/Common/Private/UbaSessionServer.cpp
--- a/Engine/Source/Programs/UnrealBuildAccelerator/Common/Private/UbaSessionServer.cpp
+++ b/Engine/Source/Programs/UnrealBuildAccelerator/Common/Private/UbaSessionServer.cpp
@@ -611,7 +611,7 @@ namespace uba
 		//m_trace.RemoteExecutionDisabled();
 	}
 
-	void SessionServer::DisableRemoteExecutionOnAgent(const tchar* agentName)
+	bool SessionServer::DisableRemoteExecutionOnAgent(const tchar* agentName)
 	{
 		SCOPED_CRITICAL_SECTION(m_remoteProcessAndSessionLock, queueLock);
 		for (auto s : m_clientSessions)
@@ -619,8 +619,9 @@ namespace uba
 			if (s->name != agentName)
 				continue;
 			s->remoteExecutionEnabled = false;
-			break;
+			return true;
 		}
+		return false;
 	}
 
 	void SessionServer::SetCustomCasKeyFromTrackedInputs(const tchar* fileName_, const tchar* workingDir_, const u8* trackedInputs, u32 trackedInputsBytes)
diff --git a/Engine/Source/Programs/UnrealBuildAccelerator/Common/Public/UbaSessionServer.h b/Engine/Source/Programs/UnrealBuildAccelerator/Common/Public/UbaSessionServer.h
--- a/Engine/Source/Programs/UnrealBuildAccelerator/Common/Public/UbaSessionServer.h
+++ b/Engine/Source/Programs/UnrealBuildAccelerator/Common/Public/UbaSessionServer.h
@@ -36,7 +36,7 @@ namespace uba
 		void DisableRemoteExecution();
 		bool IsRemoteExecutionDisabled();
 		void ReenableRemoteExecution();
-		void DisableRemoteExecutionOnAgent(const tchar* agentName);
+		bool DisableRemoteExecutionOnAgent(const tchar* agentName);
 
 		// This can be used to set a custom cas key based on a file and its inputs. Can be used for non-deterministic outputs to still be able to use cached cas content on clients
 		// For example, when a pch is built on the host we can use all the input involved to create the pch and use that as the cas key for the resulting huge file
diff --git a/Engine/Source/Programs/UnrealBuildAccelerator/Host/Private/UbaExports.cpp b/Engine/Source/Programs/UnrealBuildAccelerator/Host/Private/UbaExports.cpp
--- a/Engine/Source/Programs/UnrealBuildAccelerator/Host/Private/UbaExports.cpp
+++ b/Engine/Source/Programs/UnrealBuildAccelerator/Host/Private/UbaExports.cpp
@@ -686,9 +686,9 @@ extern "C"
 		server->GetTrace().TaskEnd(taskId, success);
 	}
 
-	void SessionServer_DisableRemoteExecutionOnAgent(uba::SessionServer* server, const uba::tchar* agentName)
+	bool SessionServer_DisableRemoteExecutionOnAgent(uba::SessionServer* server, const uba::tchar* agentName)
 	{
-		server->DisableRemoteExecutionOnAgent(agentName);
+		return server->DisableRemoteExecutionOnAgent(agentName);
 	}
 
 	void SessionServer_Destroy(uba::SessionServer* server)
diff --git a/Engine/Source/Programs/UnrealBuildAccelerator/Host/Public/UbaExports.h b/Engine/Source/Programs/UnrealBuildAccelerator/Host/Public/UbaExports.h
--- a/Engine/Source/Programs/UnrealBuildAccelerator/Host/Public/UbaExports.h
+++ b/Engine/Source/Programs/UnrealBuildAccelerator/Host/Public/UbaExports.h
@@ -154,7 +154,7 @@ extern "C"
 	UBA_API uba::u32 SessionServer_TaskBegin(uba::SessionServer* server, const uba::tchar* description, const uba::tchar* details, uba::u32 color);
 	UBA_API void SessionServer_TaskHint(uba::SessionServer* server, uba::u32 taskId, const uba::tchar* hint);
 	UBA_API void SessionServer_TaskEnd(uba::SessionServer* server, uba::u32 taskId, bool success);
-	UBA_API void SessionServer_DisableRemoteExecutionOnAgent(uba::SessionServer* server, const uba::tchar* agentName);
+	UBA_API bool SessionServer_DisableRemoteExecutionOnAgent(uba::SessionServer* server, const uba::tchar* agentName);
 	UBA_API void SessionServer_Destroy(uba::SessionServer* server);
 
 	// Scheduler
