[UBA]
* Fixed so racing is disabled if forced remote is true

diff --git a/Engine/Source/Programs/UnrealBuildAccelerator/Common/Private/UbaScheduler.cpp b/Engine/Source/Programs/UnrealBuildAccelerator/Common/Private/UbaScheduler.cpp
--- a/Engine/Source/Programs/UnrealBuildAccelerator/Common/Private/UbaScheduler.cpp
+++ b/Engine/Source/Programs/UnrealBuildAccelerator/Common/Private/UbaScheduler.cpp
@@ -219,10 +219,15 @@ namespace uba
 
 		m_session.SetOuterScheduler(this);
 
-		if (m_maxRacingWeight != 0 && !m_session.AllowKeepFilesInMemory())
+		if (m_maxRacingWeight != 0)
 		{
-			m_maxRacingWeight = 0;
-			m_session.GetLogger().Warning(TC("Session has keepFilesInMemoryDisabled. Racing will not work and has been disabled"));
+			if (!m_session.AllowKeepFilesInMemory())
+			{
+				m_maxRacingWeight = 0;
+				m_session.GetLogger().Warning(TC("Session has keepFilesInMemoryDisabled. Racing will not work and has been disabled"));
+			}
+			if (m_forceRemote) // We don't want to race if we are also forcing remote
+				m_maxRacingWeight = 0;
 		}
 	}
 
