[UBA]
* Fixed hang on macos when running oom and operating system kills process

diff --git a/Engine/Source/Programs/UnrealBuildAccelerator/Common/Private/UbaProcess.cpp b/Engine/Source/Programs/UnrealBuildAccelerator/Common/Private/UbaProcess.cpp
--- a/Engine/Source/Programs/UnrealBuildAccelerator/Common/Private/UbaProcess.cpp
+++ b/Engine/Source/Programs/UnrealBuildAccelerator/Common/Private/UbaProcess.cpp
@@ -2597,14 +2597,14 @@ namespace uba
 			if (!outClosed)
 			{
 				plist[nfds].fd = m_stdOutPipe;
-				plist[nfds].events = POLLIN;
+				plist[nfds].events = POLLIN | POLLHUP;
 				plist[nfds].revents = 0;
 				++nfds;
 			}
 			if (!errClosed)
 			{
 				plist[nfds].fd = m_stdErrPipe;
-				plist[nfds].events = POLLIN;
+				plist[nfds].events = POLLIN | POLLHUP;
 				plist[nfds].revents = 0;
 				++nfds;
 			}
