[UBA]
* Implemented mem watchdog support for mac
* Enabled threaded tcp send on mac
* Set priority to 10 for BelowNormal processes on non-windows
* Bumped priority of network recv and send threads on mac

diff --git a/Engine/Source/Programs/UnrealBuildAccelerator/Common/Private/UbaEnvironment.cpp b/Engine/Source/Programs/UnrealBuildAccelerator/Common/Private/UbaEnvironment.cpp
--- a/Engine/Source/Programs/UnrealBuildAccelerator/Common/Private/UbaEnvironment.cpp
+++ b/Engine/Source/Programs/UnrealBuildAccelerator/Common/Private/UbaEnvironment.cpp
@@ -1,6 +1,7 @@
 // Copyright Epic Games, Inc. All Rights Reserved.
 
 #include "UbaEnvironment.h"
+#include "UbaPlatform.h"
 #include "UbaStringBuffer.h"
 
 #if PLATFORM_WINDOWS
@@ -132,6 +133,8 @@ namespace uba
 	{
 		#if PLATFORM_WINDOWS
 		//SetThreadPriority(GetCurrentThread(), THREAD_PRIORITY_HIGHEST);
+		#elif PLATFORM_MAC
+		pthread_set_qos_class_self_np(QOS_CLASS_USER_INITIATED, 0);
 		#endif
 	}
 
@@ -322,6 +325,27 @@ namespace uba
 			outTotal = memStatus.ullTotalPageFile;
 			outAvailable = memStatus.ullAvailPageFile;
 		}
+#elif PLATFORM_MAC
+		u64 memKb;
+		GetPhysicallyInstalledSystemMemory(memKb);
+		outTotal = memKb*1024*1024;
+
+		mach_msg_type_number_t count = HOST_VM_INFO64_COUNT;
+		vm_statistics64_data_t vm_stat;
+		if (host_statistics64(mach_host_self(), HOST_VM_INFO64, (host_info64_t)&vm_stat, &count) == KERN_SUCCESS)
+		{
+			vm_size_t page_size;
+			host_page_size(mach_host_self(), &page_size);
+			uint64_t active      = vm_stat.active_count * page_size;
+			uint64_t wired       = vm_stat.wire_count * page_size;
+			uint64_t compressed  = vm_stat.compressor_page_count * page_size;
+			uint64_t used = active + wired + compressed;
+			outAvailable = outTotal - used;
+		}
+		else
+		{
+			outAvailable = outTotal;
+		}
 #else
 		u64 memKb;
 		GetPhysicallyInstalledSystemMemory(memKb);
diff --git a/Engine/Source/Programs/UnrealBuildAccelerator/Common/Private/UbaNetworkBackendTcp.cpp b/Engine/Source/Programs/UnrealBuildAccelerator/Common/Private/UbaNetworkBackendTcp.cpp
--- a/Engine/Source/Programs/UnrealBuildAccelerator/Common/Private/UbaNetworkBackendTcp.cpp
+++ b/Engine/Source/Programs/UnrealBuildAccelerator/Common/Private/UbaNetworkBackendTcp.cpp
@@ -981,6 +981,8 @@ namespace uba
 		if (!connection.ready.IsSet(60000)) // This should never time out!
 			return;
 
+		ElevateCurrentThreadPriority();
+
 		SendBlock* first = nullptr;
 		SendBlock* last = nullptr;
 
diff --git a/Engine/Source/Programs/UnrealBuildAccelerator/Common/Private/UbaNetworkServer.cpp b/Engine/Source/Programs/UnrealBuildAccelerator/Common/Private/UbaNetworkServer.cpp
--- a/Engine/Source/Programs/UnrealBuildAccelerator/Common/Private/UbaNetworkServer.cpp
+++ b/Engine/Source/Programs/UnrealBuildAccelerator/Common/Private/UbaNetworkServer.cpp
@@ -730,8 +730,6 @@ namespace uba
 
 	void NetworkServer::Worker::ThreadWorker(NetworkServer& server)
 	{
-		ElevateCurrentThreadPriority();
-
 		t_worker = this;
 		while (m_context->workAvailable.IsSet(~0u) && m_loop)
 		{
@@ -1607,9 +1605,8 @@ namespace uba
 				backend.SetSendBufferSize(backendConnection, m_sendBufferSize);
 		}
 
-		if (IsWindows)
-			if (m_useThreadedSend)
-				backend.SetThreadedSend(backendConnection);
+		if (m_useThreadedSend)
+			backend.SetThreadedSend(backendConnection);
 
 		m_connections.emplace_back(*this, backend, backendConnection, remoteSocketAddr, requiresCrypto, cryptoKey, m_connectionIdCounter++);
 		m_maxActiveConnections = Max(m_maxActiveConnections, u32(m_connections.size()));
diff --git a/Engine/Source/Programs/UnrealBuildAccelerator/Common/Private/UbaProcess.cpp b/Engine/Source/Programs/UnrealBuildAccelerator/Common/Private/UbaProcess.cpp
--- a/Engine/Source/Programs/UnrealBuildAccelerator/Common/Private/UbaProcess.cpp
+++ b/Engine/Source/Programs/UnrealBuildAccelerator/Common/Private/UbaProcess.cpp
@@ -2372,11 +2372,12 @@ namespace uba
 			}
 			
 			errno = 0;
-			int prio = getpriority(PRIO_PROCESS, processID);
-			if (prio != -1)
+
+			if (m_startInfo.priorityClass != ProcessPriority_Normal)
 			{
-				errno = 0;
-				if (setpriority(PRIO_PROCESS, processID, prio + 2) == -1)
+				UBA_ASSERT(m_startInfo.priorityClass == ProcessPriority_BelowNormal);
+				int newPrio = 10;
+				if (setpriority(PRIO_PROCESS, processID, newPrio) == -1)
 				{
 					UBA_ASSERTF(errno == ESRCH || errno == EPERM, TC("setpriority failed: %s. pid: %i prio: %i (%s)"), m_realApplication.c_str(), processID, prio + 2, strerror(errno));
 				}
diff --git a/Engine/Source/Programs/UnrealBuildAccelerator/Host/Private/UbaExports.cpp b/Engine/Source/Programs/UnrealBuildAccelerator/Host/Private/UbaExports.cpp
--- a/Engine/Source/Programs/UnrealBuildAccelerator/Host/Private/UbaExports.cpp
+++ b/Engine/Source/Programs/UnrealBuildAccelerator/Host/Private/UbaExports.cpp
@@ -254,7 +254,7 @@ extern "C"
 		info.workerCount = workerCount;
 		info.sendSize = sendSize;
 		info.receiveTimeoutSeconds = receiveTimeoutSeconds;
-		info.useThreadedSend = true;
+		info.useThreadedSend = !IsLinux; // Threaded send is turned off for linux since it is not needed (it improves perf/throughput on win and mac)
 		info.Apply(GetConfig());
 
 		bool success = true;
