[UBA]
* Added timer scopes around a bunch of kernel calls in detoured code on posix

diff --git a/Engine/Source/Programs/UnrealBuildAccelerator/Detours/Private/PosixOS/UbaDetoursFunctionsPosix.cpp b/Engine/Source/Programs/UnrealBuildAccelerator/Detours/Private/PosixOS/UbaDetoursFunctionsPosix.cpp
--- a/Engine/Source/Programs/UnrealBuildAccelerator/Detours/Private/PosixOS/UbaDetoursFunctionsPosix.cpp
+++ b/Engine/Source/Programs/UnrealBuildAccelerator/Detours/Private/PosixOS/UbaDetoursFunctionsPosix.cpp
@@ -351,7 +351,10 @@ int Shared_open(const char* funcName, const char* file, int flags, int mode, con
 
 	if (!CanDetour2(fileName) || fileName.Equals("/"))
 	{
+		TimerScope ts(g_kernelStats.getFileInfo);
 		int res = trueOpen(file, flags, mode);
+		ts.Leave();
+
 		DEBUG_LOG_TRUE(funcName, "NODETOUR%s (%s) -> %i", isWriteStr, file, res);
 		if (res != -1)
 			TrackInput(fileName);
@@ -514,7 +517,9 @@ int Shared_open(const char* funcName, const char* file, int flags, int mode, con
 	else
 		tempFileName = info.name;
 
+	TimerScope ts(g_kernelStats.getFileInfo);
 	int fd = trueOpen(tempFileName, flags, mode);
+	ts.Leave();
 
 	DEBUG_LOG_TRUE(funcName, "%s%s (%s) %i %i -> %i (%s)", file, isWriteStr, tempFileName, flags, mode, fd, StrError(fd, errno));
 	if (fd == -1)
@@ -578,7 +583,11 @@ FILE* Shared_fopen(const char* funcName, const char* path, const char* mode, con
 		return nullptr;
 	}
 
+
+	TimerScope ts(g_kernelStats.createFile);
 	FILE* res = TRUE_WRAPPER(fdopen)(fd, mode);
+	ts.Leave();
+
 	DEBUG_LOG_TRUE(funcName, "%i (%s  %s) -> %p", fd, path, mode, res);
 	return res;
 }
@@ -588,6 +597,7 @@ void Shared_close(int fd, const TrueClose& trueClose)
 {
 	if (!g_isDetouring)
 	{
+		TimerScope ts(g_kernelStats.closeFile);
 		trueClose();
 		return;
 	}
@@ -596,6 +606,7 @@ void Shared_close(int fd, const TrueClose& trueClose)
 	auto findIt = g_fileHandles.find(fd);
 	if (findIt == g_fileHandles.end())
 	{
+		TimerScope ts(g_kernelStats.closeFile);
 		trueClose();
 		return;
 	}
@@ -606,7 +617,12 @@ void Shared_close(int fd, const TrueClose& trueClose)
 	g_fileHandles.erase(findIt);
 	lock.Leave();
 
-	auto closeGuard = MakeGuard([&]() { trueClose(); });
+	auto closeGuard = MakeGuard([&]()
+	{
+		TimerScope ts(g_kernelStats.closeFile);
+		trueClose();
+	});
+
 	if (--fo->refCount)
 		return;
 
@@ -637,7 +653,9 @@ int Shared_fstat(const char* funcName, int fd, struct stat* attr, const True_fst
 	auto findIt = g_fileHandles.find(fd);
 	if (findIt == g_fileHandles.end())
 	{
+		TimerScope ts(g_kernelStats.getFileInfo);
 		int res = trueFstat(fd, attr);
+		ts.Leave();
 		DEBUG_LOG_TRUE(funcName, "(%i) NODETOUR (size: %llu) -> %i (%s)", fd, attr->st_size, res, StrError(res, errno));
 		return res;
 	}
@@ -647,7 +665,9 @@ int Shared_fstat(const char* funcName, int fd, struct stat* attr, const True_fst
 
 	if (fo.desiredAccess & AccessFlag_Write)
 	{
+		TimerScope ts(g_kernelStats.getFileInfo);
 		int res = trueFstat(fd, attr);
+		ts.Leave();
 		DEBUG_LOG_TRUE(funcName, "(%i) AFW (%s size: %llu) -> %i (%s)", fd, fi.originalName, attr->st_size, res, StrError(res, errno));
 		return res;
 	}
@@ -658,7 +678,9 @@ int Shared_fstat(const char* funcName, int fd, struct stat* attr, const True_fst
 	if (!fileAttr.useCache)
 	{
 		lock.Leave();
+		TimerScope ts(g_kernelStats.getFileInfo);
 		int res = trueFstat(fd, attr);
+		ts.Leave();
 		DEBUG_LOG_TRUE(funcName, "(%i) NOCACHE (%s size: %llu) -> %i (%s)", fd, fi.originalName, fileAttr.data.st_size, res, StrError(res, errno));
 		return res;
 	}
@@ -704,6 +726,7 @@ int Shared_stat(const char* funcName, const char* file, struct stat* attr, const
 	StringBuffer<> fixedFile;
 	if (!FixPath(fixedFile, file) || fixedFile.Equals("/") || !CanDetour2(fixedFile))
 	{
+		TimerScope ts(g_kernelStats.getFileInfo);
 		int res = trueStat(file, attr);
 		return res;
 	}
@@ -724,6 +747,7 @@ int Shared_stat(const char* funcName, const char* file, struct stat* attr, const
 
 	if (!fileAttr.useCache)
 	{
+		TimerScope ts(g_kernelStats.getFileInfo);
 		int res = trueStat(realName, attr);
 		DEBUG_LOG_TRUE(funcName, "%s (%s) -> %i", file, realName, res);
 		return res;
@@ -786,7 +810,9 @@ int Shared_access(const char* funcName, const char* pathname, int mode)
 	StringBuffer<> fixedPath;
 	if (!FixPath(fixedPath, pathname) || fixedPath.StartsWith("/proc") || !CanDetour2(fixedPath))
 	{
+		TimerScope ts(g_kernelStats.getFileInfo);
 		auto res = TRUE_WRAPPER(access)(pathname, mode);
+		ts.Leave();
 		DEBUG_LOG_TRUE("access", "%s %i -> %i (%s)", pathname, mode, res, StrError(res, errno));
 		return res;
 	}
@@ -809,7 +835,9 @@ int Shared_access(const char* funcName, const char* pathname, int mode)
 
 	if (!CanDetour(pathname))
 	{
+		TimerScope ts(g_kernelStats.getFileInfo);
 		auto res = TRUE_WRAPPER(access)(pathname, mode);
+		ts.Leave();
 		DEBUG_LOG_TRUE(funcName, "%s %i -> %i (%s)", pathname, mode, res, StrError(res, errno));
 		return res;
 	}
@@ -819,7 +847,9 @@ int Shared_access(const char* funcName, const char* pathname, int mode)
 
 	if (!attr.useCache)
 	{
+		TimerScope ts(g_kernelStats.getFileInfo);
 		auto res = TRUE_WRAPPER(access)(realName, mode);
+		ts.Leave();
 		DEBUG_LOG_TRUE(funcName, "%s %i (%s) -> %i %s", pathname, mode, realName, res, StrError(res, errno));
 		return res;
 	}
@@ -1250,6 +1280,7 @@ UBA_EXPORT ssize_t UBA_WRAPPER(write)(int fd, const void* buf, size_t count)
 	//	return count;
 	//}
 	DEBUG_LOG_TRUE("write", "(%i size: %llu)", fd, count);
+	TimerScope ts(g_kernelStats.writeFile);
 	return TRUE_WRAPPER(write)(fd, buf, count);
 }
 #endif
@@ -1319,13 +1350,16 @@ UBA_EXPORT FILE* UBA_WRAPPER(fdopen)(int fd, const char* mode)
 {
 	UBA_INIT_DETOUR(fdopen, fd, mode);
 	DEBUG_LOG_TRUE("fdopen", "(%i)", fd);
+	TimerScope ts(g_kernelStats.createFile);
 	return TRUE_WRAPPER(fdopen)(fd, mode);
 }
 
 UBA_EXPORT int UBA_WRAPPER(fchmod)(int fd, mode_t mode)
 {
 	UBA_INIT_DETOUR(fchmod, fd, mode);
+	TimerScope ts(g_kernelStats.setFileInfo);
 	auto res = TRUE_WRAPPER(fchmod)(fd, mode);
+	ts.Leave();
 	DEBUG_LOG_TRUE("fchmod", "(%i) %i -> %i (%s)", fd, mode, res, StrError(res, errno));
 	return res;
 }
@@ -1338,7 +1372,9 @@ UBA_EXPORT int UBA_WRAPPER(fchmodat)(int dirfd, const char* pathname, mode_t mod
 	FixPath(fixedPath, pathname);
 	DevirtualizePath(fixedPath);
 
+	TimerScope ts(g_kernelStats.setFileInfo);
 	int res = TRUE_WRAPPER(fchmodat)(dirfd, fixedPath.data, mode, flags);
+	ts.Leave();
 	DEBUG_LOG_TRUE("fchmodat", "%i %s %i %i -> %i (%s)", dirfd, pathname, mode, flags, res, StrError(res, errno));
 	return res;
 }
@@ -1365,6 +1401,7 @@ UBA_EXPORT int UBA_WRAPPER(fstatat)(int dirfd, const char* pathname, struct stat
 		return Shared_stat("fstatat", pathname, statbuf, [](const char* file, struct stat* attr) { return TRUE_WRAPPER(stat)(file, attr); });
 	UBA_ASSERT(!g_runningRemote);
 	DEBUG_LOG_TRUE("fstatat", "%s", pathname);
+	TimerScope ts(g_kernelStats.getFileInfo);
 	return TRUE_WRAPPER(fstatat)(dirfd, pathname, statbuf, flags);
 }
 
