[UBA]
* Fixed debug compile error

diff --git a/Engine/Source/Programs/UnrealBuildAccelerator/Common/Private/UbaProcess.cpp b/Engine/Source/Programs/UnrealBuildAccelerator/Common/Private/UbaProcess.cpp
--- a/Engine/Source/Programs/UnrealBuildAccelerator/Common/Private/UbaProcess.cpp
+++ b/Engine/Source/Programs/UnrealBuildAccelerator/Common/Private/UbaProcess.cpp
@@ -2379,7 +2379,7 @@ namespace uba
 				int newPrio = 10;
 				if (setpriority(PRIO_PROCESS, processID, newPrio) == -1)
 				{
-					UBA_ASSERTF(errno == ESRCH || errno == EPERM, TC("setpriority failed: %s. pid: %i prio: %i (%s)"), m_realApplication.c_str(), processID, prio + 2, strerror(errno));
+					UBA_ASSERTF(errno == ESRCH || errno == EPERM, TC("setpriority failed: %s. pid: %i prio: %i (%s)"), m_realApplication.c_str(), processID, newPrio, strerror(errno));
 				}
 			}
 
