[UBA]
* Changed so force remote is changing to only force cross system remotes if any cross platform helpers are connected

diff --git a/Engine/Source/Programs/UnrealBuildAccelerator/Common/Private/UbaScheduler.cpp b/Engine/Source/Programs/UnrealBuildAccelerator/Common/Private/UbaScheduler.cpp
--- a/Engine/Source/Programs/UnrealBuildAccelerator/Common/Private/UbaScheduler.cpp
+++ b/Engine/Source/Programs/UnrealBuildAccelerator/Common/Private/UbaScheduler.cpp
@@ -205,7 +205,7 @@ namespace uba
 				return false;
 			});
 
-		if (info.memWatchdog && IsWindows && !IsRunningWine()) // Only implemented for normal windows for now
+		if (info.memWatchdog && !IsLinux && !IsRunningWine()) // Only implemented for normal windows for now
 		{
 			m_memWatchdog = true;
 			m_memTracingLevel = info.memTracingLevel;
@@ -914,8 +914,18 @@ namespace uba
 
 			if (isLocal)
 			{
-				if (m_forceRemote && entry.canExecuteRemotely)
-					continue;
+				if (m_forceRemote)
+				{
+					if (entry.canExecuteRemotely)
+					{
+						if (!m_hadRemoteRequest)
+							continue;
+						if (!m_hadCrossSystemRequest)
+							continue;
+						if (entry.canCrossArchitecture || entry.canCrossPlatform)
+							continue;
+					}
+				}
 				if (!forReuse)
 				{
 					if (runStatus == RunStatus_QueuedForRun && atMaxLocalWeight)
@@ -1580,6 +1590,9 @@ namespace uba
 		if (!m_memWatchdog)
 			return;
 
+		if (!IsWindows) // Memory prediction is not implemented for non-windows
+			return;
+
 		auto& impl = (ProcessImpl&)process;
 		auto ei = (ExitProcessInfo*)impl.m_startInfo.userData;
 		
@@ -1607,9 +1620,12 @@ namespace uba
 
 	void Scheduler::LocalProcessStart(ProcessStartInfo2& info)
 	{
+		if (!IsWindows) // Memory prediction is not implemented for non-windows
+			return;
+			
 		if (!m_memWatchdog || !info.memoryGroupId)
 			return;
-
+		
 		SCOPED_FUTEX(m_memGroupLookupLock, statsLock);
 		auto insres = m_memGroupLookup.try_emplace(info.memoryGroupId);
 		MemGroupStats& stats = insres.first->second;
@@ -1679,6 +1695,9 @@ namespace uba
 		if (!m_memWatchdog)
 			return;
 
+		if (!IsWindows) // Not implemented for non windows
+			return;
+
 		auto& impl = *(ProcessImpl*)handle.m_process;
 
 		if (info.memoryGroupId)
@@ -1745,6 +1764,8 @@ namespace uba
 
 	void Scheduler::LoadMemTrackTable()
 	{
+		if (!IsWindows) // Not implemented for non windows
+			return;
 		StringBuffer<> fileName;
 		fileName.Append(m_session.GetRootDir()).Append("memgroups");
 		FileAccessor file(m_session.GetLogger(), fileName.data);
@@ -1767,6 +1788,9 @@ namespace uba
 
 	void Scheduler::SaveMemTrackTable()
 	{
+		if (!IsWindows) // Not implemented for non windows
+			return;
+
 		for (auto& kv : m_memGroupLookup)
 			m_memGroupDb[kv.first] = MemGroupEntry{ kv.second.baseline, kv.second.average };
 
