[UBA]
* Fixes to PollStdPipes.. now it should work with a timeout of -1

diff --git a/Engine/Source/Programs/UnrealBuildAccelerator/Common/Private/UbaProcess.cpp b/Engine/Source/Programs/UnrealBuildAccelerator/Common/Private/UbaProcess.cpp
--- a/Engine/Source/Programs/UnrealBuildAccelerator/Common/Private/UbaProcess.cpp
+++ b/Engine/Source/Programs/UnrealBuildAccelerator/Common/Private/UbaProcess.cpp
@@ -64,9 +64,15 @@ namespace uba
 	{
 		PipeReader(ProcessImpl& p, LogEntryType lt) : process(p), logType(lt) {}
 		~PipeReader()
+		{
+			Flush();
+		}
+
+		void Flush()
 		{
 			if (!currentString.empty())
 				process.InternalLogLine(false, TString(currentString), logType);
+			currentString.clear();
 		}
 
 		void ReadData(char* buf, u32 readCount)
@@ -641,23 +647,23 @@ namespace uba
 		{
 			exitCode = InternalCreateProcess(environment, m_comMemory.handle, m_comMemory.offset);
 
+			#if !PLATFORM_WINDOWS
+			PipeReader outReader(*this, LogEntryType_Info);
+			PipeReader errReader(*this, LogEntryType_Error);
+			Atomic<bool> pollStdDone = false;
+			Thread pollStdThread([&]()
+			{
+				while (!pollStdDone)
+					if (!PollStdPipes(outReader, errReader, -1))
+						return false;
+				return true;
+			}, TC("PollStd"));
+			#endif
+			
 			if (exitCode == 0)
 			{
 				m_session.ProcessNativeCreated(*this);
 
-				#if !PLATFORM_WINDOWS
-				PipeReader outReader(*this, LogEntryType_Info);
-				PipeReader errReader(*this, LogEntryType_Error);
-				Atomic<bool> pollStdDone = false;
-				Thread pollStdThread([&]()
-				{
-					while (!pollStdDone)
-						if (!PollStdPipes(outReader, errReader, 500))
-							return false;
-					return true;
-				}, TC("PollStd"));
-				#endif
-
 				bool loop = m_detourEnabled;
 				while (loop && WaitForRead())
 				{
@@ -669,13 +675,6 @@ namespace uba
 					m_processStats.hostTotalTime += GetTime() - startTime;
 					++m_messageCount;
 				}
-
-				#if !PLATFORM_WINDOWS
-				pollStdDone = true;
-				pollStdThread.Wait();
-				while (PollStdPipes(outReader, errReader, 500) && !IsCancelled())
-					continue;
-				#endif
 			}
 
 			m_processStats.exitTime = GetTime();
@@ -693,6 +692,14 @@ namespace uba
 			// We need to wait also because we need all written files to be added in shared file system
 			WaitForChildrenExit();
 
+			#if !PLATFORM_WINDOWS
+			pollStdDone = true;
+			pollStdThread.Wait();
+			PollStdPipes(outReader, errReader, -1);
+			outReader.Flush();
+			errReader.Flush();
+			#endif
+			
 			if (cancelled)
 			{
 				m_session.ProcessCancelled(*this);
@@ -2326,17 +2333,27 @@ namespace uba
 
 			if (true)
 			{
+				#if PLATFORM_MAC
 				if (pipe(outPipe) || pipe(errPipe))
+				#else
+				if (pipe2(outPipe, O_CLOEXEC) || pipe2(errPipe, O_CLOEXEC))
+				#endif
 				{
 					logger.Error("pipe failed (%s)", strerror(errno));
 					return UBA_EXIT_CODE(18);
 				}
+				
+				#if PLATFORM_MAC
+				auto set_cloexec = [](int fd){ int f = fcntl(fd, F_GETFD); if (f != -1) fcntl(fd, F_SETFD, f | FD_CLOEXEC); };
+				set_cloexec(outPipe[0]); set_cloexec(outPipe[1]);
+				set_cloexec(errPipe[0]); set_cloexec(errPipe[1]);
+				#endif
 
 				res = posix_spawn_file_actions_addclose(&fileActions, outPipe[0]);
 				UBA_ASSERTF(!res, "posix_spawn_file_actions_addclose outPipe[0] failed: %i", res);
 
 				res = posix_spawn_file_actions_addclose(&fileActions, errPipe[0]);
-				UBA_ASSERTF(!res, "posix_spawn_file_actions_addclose errPipe[1] failed: %i", res);
+				UBA_ASSERTF(!res, "posix_spawn_file_actions_addclose errPipe[0] failed: %i", res);
 
 				res = posix_spawn_file_actions_adddup2(&fileActions, outPipe[1], 1);
 				UBA_ASSERTF(!res, "posix_spawn_file_actions_adddup2 outPipe[1] failed: %i", res);
@@ -2503,11 +2520,6 @@ namespace uba
 		return UBA_EXIT_CODE(14);
 #else
 
-		if (m_stdOutPipe != -1)
-			close(m_stdOutPipe);
-		if (m_stdErrPipe != -1)
-			close(m_stdErrPipe);
-
 		auto g = MakeGuard([this]() { m_nativeProcessId = 0; });
 
 		if (cancel)
diff --git a/Engine/Source/Programs/UnrealBuildAccelerator/Detours/Private/PosixOS/UbaDetoursFunctionsPosix.cpp b/Engine/Source/Programs/UnrealBuildAccelerator/Detours/Private/PosixOS/UbaDetoursFunctionsPosix.cpp
--- a/Engine/Source/Programs/UnrealBuildAccelerator/Detours/Private/PosixOS/UbaDetoursFunctionsPosix.cpp
+++ b/Engine/Source/Programs/UnrealBuildAccelerator/Detours/Private/PosixOS/UbaDetoursFunctionsPosix.cpp
@@ -2447,7 +2447,7 @@ UBA_EXPORT pid_t UBA_WRAPPER(wait4)(pid_t pid, int* status, int options, struct
 	pid_t res = TRUE_WRAPPER(wait4)(pid, status, options, rusage);
 	if (WIFEXITED(*status))
 		Rpc_UpdateTables();
-	//DEBUG_LOG_TRUE("wait4", "(pid %i) -> %i (%i)", pid, res, *status);
+	DEBUG_LOG_TRUE("wait4", "(pid %i) -> %i (%i)", pid, res, *status);
 	return res;
 }
 
