[UBA]
* Fixed bug where the code did not handle creating directories in temp folder when application rule had set all files in temp to be memory files. Fixes errors when using latest intel toolchain

diff --git a/Engine/Source/Programs/UnrealBuildAccelerator/Detours/Private/Windows/UbaDetoursFunctionsKernelBase.inl b/Engine/Source/Programs/UnrealBuildAccelerator/Detours/Private/Windows/UbaDetoursFunctionsKernelBase.inl
--- a/Engine/Source/Programs/UnrealBuildAccelerator/Detours/Private/Windows/UbaDetoursFunctionsKernelBase.inl
+++ b/Engine/Source/Programs/UnrealBuildAccelerator/Detours/Private/Windows/UbaDetoursFunctionsKernelBase.inl
@@ -219,19 +219,29 @@ BOOL Detoured_CreateDirectoryW(LPCWSTR lpPathName, LPSECURITY_ATTRIBUTES lpSecur
 	StringBuffer<> pathName;
 	FixPath(pathName, lpPathName);
 
+	StringBuffer<> lowerName(pathName);
+
 	if (pathName.StartsWith(g_systemTemp.data))
 	{
 		SuppressCreateFileDetourScope s;
 		BOOL res = True_CreateDirectoryW(lpPathName, lpSecurityAttributes);
 		DEBUG_LOG_TRUE(L"CreateDirectoryW", L"%ls -> %ls", lpPathName, ToString(res));
+
+		// Some applications (icx.exe) creates directory where outputfiles are.. so we need to create an entry with '$' (directory) in the mapping table
+		if (res && g_allowOutputFiles && g_rules->IsOutputFile(pathName, g_systemTemp))
+		{
+			SCOPED_WRITE_LOCK(g_mappedFileTable.m_lookupLock, lookupLock);
+			auto insres = g_mappedFileTable.m_lookup.try_emplace(ToStringKey(lowerName.MakeLower()));
+			FileInfo& info = insres.first->second;
+			info.originalName = g_memoryBlock.Strdup(pathName).data;
+			info.name = TC("$");
+		}
 		return res;
 	}
 
 	DevirtualizePath(pathName);
 	
-	StringBuffer<> lowerName(pathName);
 	lowerName.MakeLower();
-
 	StringKey pathNameKey = ToStringKey(lowerName);
 
 	bool checkIfDir = true;
@@ -2098,7 +2108,7 @@ BOOL Detoured_GetFileInformationByHandleEx(HANDLE hFile, FILE_INFO_BY_HANDLE_CLA
 		if (entryOffset == ~u32(0) && trueHandle == INVALID_HANDLE_VALUE)
 		{
 			MemoryFile* mf = fi.memoryFile;
-			UBA_ASSERTF(mf, L"GetFileInformationByHandleEx called on file %s which has no entry offset or real handle", HandleToName(hFile));
+			UBA_ASSERTF(mf || fi.name[0] == '$', L"GetFileInformationByHandleEx called on file %s which has no entry offset or real handle", HandleToName(hFile));
 
 			DEBUG_LOG_DETOURED(L"GetFileInformationByHandleEx", L"(MEMORY) (%u) %llu (%ls)", fileInformationClass, uintptr_t(hFile), HandleToName(hFile));
 
@@ -2121,6 +2131,20 @@ BOOL Detoured_GetFileInformationByHandleEx(HANDLE hFile, FILE_INFO_BY_HANDLE_CLA
 				data.Directory = false;
 				return TRUE;
 			}
+			else if (fileInformationClass == FileRemoteProtocolInfo)
+			{
+				UBA_ASSERT(fi.name[0] == '$');
+				auto& data = *(FILE_REMOTE_PROTOCOL_INFO*)lpFileInformation;
+				memset(&data, 0, sizeof(FILE_REMOTE_PROTOCOL_INFO));
+				data.StructureVersion = 1;
+				data.StructureSize = sizeof(FILE_REMOTE_PROTOCOL_INFO);
+				data.Protocol = 0;  // means "not remote"
+				data.ProtocolMajorVersion = 0;
+				data.ProtocolMinorVersion = 0;
+				data.ProtocolRevision = 0;
+				data.Flags = 0;
+				return TRUE;
+			}
 			else
 			{
 				UBA_ASSERTF(!mf, L"GetFileInformationByHandleEx called for memory file using class %u which is not implemented (%s)", fileInformationClass, HandleToName(hFile));
diff --git a/Engine/Source/Programs/UnrealBuildAccelerator/Detours/Private/Windows/UbaDetoursFunctionsNtDll.inl b/Engine/Source/Programs/UnrealBuildAccelerator/Detours/Private/Windows/UbaDetoursFunctionsNtDll.inl
--- a/Engine/Source/Programs/UnrealBuildAccelerator/Detours/Private/Windows/UbaDetoursFunctionsNtDll.inl
+++ b/Engine/Source/Programs/UnrealBuildAccelerator/Detours/Private/Windows/UbaDetoursFunctionsNtDll.inl
@@ -1315,6 +1315,7 @@ NTSTATUS NTAPI Shared_NtCreateFile(bool IsCreateFunc, PHANDLE hFileHandle, ACCES
 			bool isOutput = (isWrite || isDeleteOnClose) && g_rules->IsOutputFile(fileName, g_systemTemp);
 			bool isLocal = !isOutput;
 			MemoryFile* mf = new MemoryFile(isLocal, FileTypeMaxSize(fileName, isSystemOrTempFile), isThrowAway, FileTypeCommitSize(fileName), fileName.data);
+			auto mfg = MakeGuard([mf](){ delete mf; });
 
 			if (!isThrowAway && CreateDisposition == FILE_OPEN)
 			{
@@ -1390,6 +1391,8 @@ NTSTATUS NTAPI Shared_NtCreateFile(bool IsCreateFunc, PHANDLE hFileHandle, ACCES
 				mf->fileIndex = InterlockedDecrement(&g_memoryFileIndexCounter);
 			}
 
+			mfg.Cancel();
+
 			info.created = true;
 			info.memoryFile = mf;
 		}
