[UBA]
* Changed "but can't find neither compressed or uncompressed version" from error to warning

diff --git a/Engine/Source/Programs/UnrealBuildAccelerator/Common/Private/UbaStorage.cpp b/Engine/Source/Programs/UnrealBuildAccelerator/Common/Private/UbaStorage.cpp
--- a/Engine/Source/Programs/UnrealBuildAccelerator/Common/Private/UbaStorage.cpp
+++ b/Engine/Source/Programs/UnrealBuildAccelerator/Common/Private/UbaStorage.cpp
@@ -2472,7 +2472,7 @@ namespace uba
 			if (!HasCasFile(actualKey, &casEntryPtr))
 			{
 				if (!testCompressed)
-					return m_logger.Error(TC("[CopyOrLink] Trying to copy cas %s to %s but can't find neither compressed or uncompressed version"), CasKeyString(actualKey).str, destination);
+					return m_logger.Warning(TC("[CopyOrLink] Trying to copy cas %s to %s but can't find neither compressed or uncompressed version"), CasKeyString(actualKey).str, destination);
 
 				actualKey = AsCompressed(casKey, true);
 				testCompressed = false;
