[UBA]
Downgraded "Failed to copy cas from" error to warning

diff --git a/Engine/Source/Programs/UnrealBuildAccelerator/Common/Private/UbaSessionServer.cpp b/Engine/Source/Programs/UnrealBuildAccelerator/Common/Private/UbaSessionServer.cpp
--- a/Engine/Source/Programs/UnrealBuildAccelerator/Common/Private/UbaSessionServer.cpp
+++ b/Engine/Source/Programs/UnrealBuildAccelerator/Common/Private/UbaSessionServer.cpp
@@ -1454,7 +1454,7 @@ namespace uba
 			else
 			{
 				if (!m_storage.CopyOrLink(casKey, destination.data, file.attributes, writeCompressed, clientId, formattingFunc))
-					return m_logger.Error(TC("Failed to copy cas from %s to %s (%s)"), CasKeyString(casKey).str, destination.data, startInfo.GetDescription());
+					return m_logger.Warning(TC("Failed to copy cas from %s to %s (%s)"), CasKeyString(casKey).str, destination.data, startInfo.GetDescription());
 			}
 			TraceWrittenFile(file.processId, destination);
 		}
