[UBA]
* Changed wait loop in StorageServer::WaitForCasFile to first check time to calculate timeout before running additional work in case the additional work takes a very long time (It shouldn't but who knows)
* Added more logging to try to find issues seen on the farm

diff --git a/Engine/Source/Programs/UnrealBuildAccelerator/Common/Private/UbaProcess.cpp b/Engine/Source/Programs/UnrealBuildAccelerator/Common/Private/UbaProcess.cpp
--- a/Engine/Source/Programs/UnrealBuildAccelerator/Common/Private/UbaProcess.cpp
+++ b/Engine/Source/Programs/UnrealBuildAccelerator/Common/Private/UbaProcess.cpp
@@ -972,7 +972,7 @@ namespace uba
 			UBA_PROCESS_MESSAGES
 			#undef UBA_PROCESS_MESSAGE
 		}
-		return CancelWithError(m_session.m_logger.Error(TC("Unknown message type %u"), messageType));
+		return CancelWithError(m_session.m_logger.Error(TC("Unknown message type %u when communicating with detoured process %s. (Message index %u after %s)"), messageType, m_startInfo.GetDescription(), m_messageCount, TimeToText(GetTime() - m_startTime).str));
 	}
 
 	void ProcessImpl::CallSession(MessageType type, bool success)
diff --git a/Engine/Source/Programs/UnrealBuildAccelerator/Common/Private/UbaStorage.cpp b/Engine/Source/Programs/UnrealBuildAccelerator/Common/Private/UbaStorage.cpp
--- a/Engine/Source/Programs/UnrealBuildAccelerator/Common/Private/UbaStorage.cpp
+++ b/Engine/Source/Programs/UnrealBuildAccelerator/Common/Private/UbaStorage.cpp
@@ -2472,8 +2472,21 @@ namespace uba
 			if (!HasCasFile(actualKey, &casEntryPtr))
 			{
 				if (!testCompressed)
-					return m_logger.Warning(TC("[CopyOrLink] Trying to copy cas %s to %s but can't find neither compressed or uncompressed version"), CasKeyString(actualKey).str, destination);
+				{
+					SCOPED_READ_LOCK(m_casLookupLock, lookupLock);
+					auto it = m_casLookup.find(casKey);
+					bool found = it == m_casLookup.end();
+					lookupLock.Leave();
+					const tchar* extraInfo = TC("NeverExisted");
+					if (found)
+					{
+						CasEntry& casEntry = it->second;
+						CasEntryWriteLock casEntryLock(casEntry);
+						extraInfo = casEntry.dropped ? TC("Dropped") : TC("InLookup");
+					}
 
+					return m_logger.Warning(TC("[CopyOrLink] Trying to copy cas %s to %s but can't find neither compressed or uncompressed version (%s)"), CasKeyString(actualKey).str, destination, extraInfo);
+				}
 				actualKey = AsCompressed(casKey, true);
 				testCompressed = false;
 				continue;
diff --git a/Engine/Source/Programs/UnrealBuildAccelerator/Common/Private/UbaStorageServer.cpp b/Engine/Source/Programs/UnrealBuildAccelerator/Common/Private/UbaStorageServer.cpp
--- a/Engine/Source/Programs/UnrealBuildAccelerator/Common/Private/UbaStorageServer.cpp
+++ b/Engine/Source/Programs/UnrealBuildAccelerator/Common/Private/UbaStorageServer.cpp
@@ -115,8 +115,6 @@ namespace uba
 		u32 timeout = 0;
 		while (!waitEntry.done.IsSet(timeout)) // TODO WaitMultipleObjects (additional work and Done)
 		{
-			timeout = m_server.DoAdditionalWork() ? 0 : 50;
-
 			u64 waited = GetTime() - startTime;
 			u64 waitedMs = TimeToMs(waited);
 			if (waitedMs > 500)
@@ -124,8 +122,10 @@ namespace uba
 				if (!m_server.IsConnectedSlow(clientId))
 					return m_logger.Info(TC("Client with id %u disconnected while waiting for cas %s (%s)"), clientId, CasKeyString(casKey).str, hint).ToFalse();
 				if (waitedMs > 4 * 60 * 1000) // 4 minutes timeout
-					return m_logger.Info(TC("Timed out waiting %s for cas %s to be transferred from remote to storage (%s)"), TimeToText(waited).str, CasKeyString(casKey).str, hint).ToFalse();
+					return m_logger.Info(TC("Timed out waiting %s for cas %s to be transferred from remote to storage (%s)"), TimeToText(waited, true).str, CasKeyString(casKey).str, hint).ToFalse();
 			}
+
+			timeout = m_server.DoAdditionalWork() ? 0 : 50;
 		}
 		return waitEntry.success;
 	}
