[UBA]
* Added shutdown socket to send thread when exiting thread. Maybe this can solve some of the longer timeouts we've seen

diff --git a/Engine/Source/Programs/UnrealBuildAccelerator/Common/Private/UbaNetworkBackendTcp.cpp b/Engine/Source/Programs/UnrealBuildAccelerator/Common/Private/UbaNetworkBackendTcp.cpp
--- a/Engine/Source/Programs/UnrealBuildAccelerator/Common/Private/UbaNetworkBackendTcp.cpp
+++ b/Engine/Source/Programs/UnrealBuildAccelerator/Common/Private/UbaNetworkBackendTcp.cpp
@@ -1070,6 +1070,12 @@ namespace uba
 			if (!HandleSendError(connection, logger, hint))
 				return;
 		}
+
+		SCOPED_CRITICAL_SECTION(connection.shutdownLock, lock2);
+		SCOPED_FUTEX(connection.sendLock, lock);
+		Socket s = connection.socket;
+		if (s != InvalidSocket)
+			ShutdownSocket(logger, s, TC("ThreadSend"));
 	}
 
 	bool NetworkBackendTcp::HandleSendError(Connection& connection, Logger& logger, const tchar* hint)
