[UBA]
* Improved assert when NtCreateSection is called with detoured handle that is not backed by real handle

diff --git a/Engine/Source/Programs/UnrealBuildAccelerator/Detours/Private/Windows/UbaDetoursFunctionsNtDll.inl b/Engine/Source/Programs/UnrealBuildAccelerator/Detours/Private/Windows/UbaDetoursFunctionsNtDll.inl
--- a/Engine/Source/Programs/UnrealBuildAccelerator/Detours/Private/Windows/UbaDetoursFunctionsNtDll.inl
+++ b/Engine/Source/Programs/UnrealBuildAccelerator/Detours/Private/Windows/UbaDetoursFunctionsNtDll.inl
@@ -625,7 +625,7 @@ NTSTATUS NTAPI Detoured_NtCreateSection(PHANDLE SectionHandle, ACCESS_MASK Desir
 	{
 		DetouredHandle& h = asDetouredHandle(FileHandle);
 		FileHandle = h.trueHandle;
-		UBA_ASSERT(FileHandle != INVALID_HANDLE_VALUE);
+		UBA_ASSERTF(FileHandle != INVALID_HANDLE_VALUE, L"NtCreateSection does not support detoured handles without backing file (%s)", HandleToName(FileHandle));
 	}
 	return True_NtCreateSection(SectionHandle, DesiredAccess, ObjectAttributes, MaximumSize, SectionPageProtection, AllocationAttributes, FileHandle);
 }
