[UBA]
* Fixed so m_messageThreadDone is always initialized and in ctor since it is used by WaitForExit from the outside.. fixed so it is always set
* Added retry for ERROR_NOT_ENOUGH_MEMORY when trying to spawn process
* Disable log file when failing to create process to prevent errors when agents try to send non-existing log files back

diff --git a/Engine/Source/Programs/UnrealBuildAccelerator/Common/Private/UbaProcess.cpp b/Engine/Source/Programs/UnrealBuildAccelerator/Common/Private/UbaProcess.cpp
--- a/Engine/Source/Programs/UnrealBuildAccelerator/Common/Private/UbaProcess.cpp
+++ b/Engine/Source/Programs/UnrealBuildAccelerator/Common/Private/UbaProcess.cpp
@@ -172,6 +172,8 @@ namespace uba
 		if (!m_comMemory.memory)
 			return;
 
+		m_messageThreadDone.Create(EventResetType_Manual);
+
 		#if !PLATFORM_WINDOWS
 		if (UseSemaphoreForSharedEvents())
 		{
@@ -256,6 +258,7 @@ namespace uba
 			m_exitCode = 44324;
 			CallExitedFunc();
 			m_hasExited = true;
+			m_messageThreadDone.Set();
 			return false;
 		}
 
@@ -286,6 +289,7 @@ namespace uba
 				m_exitCode = 44325;
 				CallExitedFunc();
 				m_hasExited = true;
+				m_messageThreadDone.Set();
 				return false;
 			}
 
@@ -310,7 +314,6 @@ namespace uba
 
 		if (async)
 		{
-			m_messageThreadDone.Create(EventResetType_Manual);
 			m_session.m_processMessageThreads.AddWork([this, environment](const WorkContext& context)
 				{
 					m_messageThread = context.thread;
@@ -336,6 +339,8 @@ namespace uba
 			#if !PLATFORM_WINDOWS
 			pthread_sigmask(SIG_BLOCK, &oldMask, nullptr);
 			#endif
+
+			m_messageThreadDone.Set();
 		}
 		return true;
 	}
@@ -2016,6 +2021,16 @@ namespace uba
 					m_cancelEvent.Set();
 					return ProcessCancelExitCode;
 				}
+				else if (error == ERROR_NOT_ENOUGH_MEMORY) // Have no explanation why this can happen but have seen it on own machine when helping out compressing staging data
+				{
+					bool retry = retryCount++ < 5;
+					if (retry)
+					{
+						Sleep(100 + (rand() % 200));
+						ZeroMemory(&processInfo, sizeof(processInfo));
+						continue;
+					}
+				}
 
 				LastErrorToText lett(error);
 				const tchar* errorText = lett.data;
@@ -2027,6 +2042,8 @@ namespace uba
 					if (error == ERROR_DIRECTORY)
 						logger.Error(TC("HOW CAN THIS HAPPEN? '%s'"), workingDir);
 
+					m_startInfo.logFile = TC("");
+
 					logger.Error(TC("DetourCreateProcessWithDllEx failed: %s (Working dir: %s). Exit code: %u - %s"), commandLine.c_str(), workingDir, error, errorText);
 				}
 				return UBA_EXIT_CODE(6);
