[UBA]
* Fixed so scheduler will not try to spawn new processes until a while after killing process due to memory pressure
* Fixed bug with process counter when getting cache misses

diff --git a/Engine/Source/Programs/UnrealBuildAccelerator/Common/Private/UbaScheduler.cpp b/Engine/Source/Programs/UnrealBuildAccelerator/Common/Private/UbaScheduler.cpp
--- a/Engine/Source/Programs/UnrealBuildAccelerator/Common/Private/UbaScheduler.cpp
+++ b/Engine/Source/Programs/UnrealBuildAccelerator/Common/Private/UbaScheduler.cpp
@@ -675,7 +675,6 @@ namespace uba
 		lock.Leave();
 
 		UpdateQueueCounter(1);
-		UpdateActiveProcessCounter(false, -1);
 		m_updateThreadWakeup.Set();
 	}
 
@@ -1649,7 +1648,7 @@ namespace uba
 		if (!m_memWaitThreshold)
 			return true;
 
-		if (m_memEstimatedUsage < m_memWaitThreshold)
+		if (m_memEstimatedUsage < m_memWaitThreshold && !m_lastKillTime)
 		{
 			if (m_memDelaySpawning)
 			{
@@ -2020,7 +2019,7 @@ namespace uba
 					}
 				}
 				estimatedMemoryUsage = memBaseLine + untrackedMemory + estimatedProcessMemoryUsage;
-				m_memEstimatedUsage = estimatedMemoryUsage < memUsed ? memUsed : estimatedMemoryUsage;
+				m_memEstimatedUsage = Max(memUsed, estimatedMemoryUsage);
 			}
 
 			if (m_memTracingLevel >= 1)
@@ -2037,11 +2036,19 @@ namespace uba
 
 			if (memKillThreshold)
 			{
+				if (m_lastKillTime && m_lastKillTime + MsToTime(3000) < GetTime())
+				{
+					m_lastKillTime = 0;
+					m_updateThreadWakeup.Set();
+				}
+
 				while (memUsed > memKillThreshold)
 				{
-					str.Clear().Appendf(TC("Low on memory (%s/%s). Kill threshold is %s"), BytesToText(memUsed).str, BytesToText(memTotal).str, BytesToText(memKillThreshold).str);
+					if (m_activeLocalProcesses <= 1)
+						break;
+					m_lastKillTime = GetTime();
 
-					// kill racing processes first
+					str.Clear().Appendf(TC("Low on memory (%s/%s). Kill threshold is %s"), BytesToText(memUsed).str, BytesToText(memTotal).str, BytesToText(memKillThreshold).str);
 
 					KillNewestLocalProcess(str);
 					if (!GetMemoryInfo(logger, memAvail, memTotal, &maxPageFile))
diff --git a/Engine/Source/Programs/UnrealBuildAccelerator/Common/Public/UbaScheduler.h b/Engine/Source/Programs/UnrealBuildAccelerator/Common/Public/UbaScheduler.h
--- a/Engine/Source/Programs/UnrealBuildAccelerator/Common/Public/UbaScheduler.h
+++ b/Engine/Source/Programs/UnrealBuildAccelerator/Common/Public/UbaScheduler.h
@@ -232,6 +232,7 @@ namespace uba
 		bool m_memDelaySpawning = false;
 		Atomic<u64> m_memEstimatedUsage = 0;
 		Atomic<u64> m_memWaitThreshold = 0;
+		Atomic<u64> m_lastKillTime = 0;
 		u8 m_memStartWaitPercent = 0;
 		u8 m_memStartKillPercent = 0;
 		u8 m_memTracingLevel = 0;
