[UBA]
* Implemented (repaired) support for not using casdb on agent side

diff --git a/Engine/Source/Programs/UnrealBuildAccelerator/Common/Private/UbaSessionClient.cpp b/Engine/Source/Programs/UnrealBuildAccelerator/Common/Private/UbaSessionClient.cpp
--- a/Engine/Source/Programs/UnrealBuildAccelerator/Common/Private/UbaSessionClient.cpp
+++ b/Engine/Source/Programs/UnrealBuildAccelerator/Common/Private/UbaSessionClient.cpp
@@ -269,7 +269,7 @@ namespace uba
 		//	storeUncompressed = !IsCompressed(casKey);
 
 		Storage::RetrieveResult result;
-		bool res = m_storage.RetrieveCasFile(result, tempKey, keyHint, hint, nullptr, 1, allowProxy);
+		bool res = m_storage.RetrieveCasFile(result, tempKey, keyHint, hint, nullptr, false, 1, allowProxy);
 		outNewKey = result.casKey;
 		outSize = result.size;
 		return res;
@@ -669,10 +669,13 @@ namespace uba
 					SCOPED_FUTEX(entry.lock, entryCs);
 					ts.Leave();
 
+					bool canBeFreed = true; // For now this code path is really only used for one thing (stage compression on console target)
+
 					if (entry.handled)
 					{
 						if (!entry.success)
 							return false;
+						UBA_ASSERT(canBeFreed == entry.canBeFreed);
 					}
 					else
 					{
@@ -681,12 +684,21 @@ namespace uba
 						entry.handled = true;
 						Storage::RetrieveResult result;
 						bool allowProxy = rules.AllowStorageProxy(fileName);
-						if (!m_storage.RetrieveCasFile(result, casKey, fileNameKey, fileName.data, &m_fileMappingBuffer, memoryMapAlignment, allowProxy))
+						if (!m_storage.RetrieveCasFile(result, casKey, fileNameKey, fileName.data, &m_fileMappingBuffer, canBeFreed, memoryMapAlignment, allowProxy))
 							return m_logger.Error(TC("Error retrieving cas entry %s (%s) (2)"), CasKeyString(casKey).str, fileName.data);
 						entry.success = true;
 						entry.contentSize = result.size;
 						entry.memoryHandle = result.view.handle;
 						entry.memoryOffset = result.view.offset;
+						entry.canBeFreed = canBeFreed;
+					}
+
+					if (canBeFreed)
+					{
+						ProcessImpl::UsedFileMapping usedFileMapping { true };
+						SCOPED_FUTEX(process.m_usedFileMappingsLock, lock);
+						if (process.m_usedFileMappings.insert({fileNameKey, usedFileMapping}).second)
+							++entry.refCount;
 					}
 
 					fileSize = entry.contentSize;
@@ -1574,6 +1586,10 @@ namespace uba
 				CasKey knownInputKey = reader.ReadCasKey();
 				u32 mappingAlignment = reader.ReadU32();
 				bool allowProxy = reader.ReadBool();
+
+				if (!m_useStorage) // If storage is not used these will just be downloaded and forgotten since there is no tracking
+					continue;
+
 				bool storeUncompressed = !m_allowMemoryMaps || mappingAlignment == 0;
 				if (storeUncompressed)
 					knownInputKey = AsCompressed(knownInputKey, false);
@@ -1582,7 +1598,7 @@ namespace uba
 					{
 						TrackWorkScope tws(m_client, StringBuffer<>().Appendf(TC("KnownInput")), ColorWork);
 						Storage::RetrieveResult result;
-						bool res = m_storage.RetrieveCasFile(result, knownInputKey, knownFile, TC("KnownInput"), nullptr, 1, allowProxy);
+						bool res = m_storage.RetrieveCasFile(result, knownInputKey, knownFile, TC("KnownInput"), nullptr, false, 1, allowProxy);
 						(void)res;
 					}, 1, TC("KnownInput"));
 			}
@@ -2002,6 +2018,12 @@ namespace uba
 				rec.isKilled = true;
 				SendReturnProcess(rec.handle.GetId(), reason);
 				++m_stats.killCount;
+
+				// We must always wait for cancelled processes to really finish since we might get new processes from host
+				// which depend on this process and then outputs might collide.
+				lock.Leave();
+				if (!rec.handle.WaitForExit(60*1000))
+					m_logger.Warning(TC("Timed out waiting for killed process to exit"));
 				return true;
 			};
 
@@ -2108,11 +2130,6 @@ namespace uba
 						if (rec.handle.GetId() != cancelled)
 							continue;
 						TryKill(rec, TC("Killed by server"));
-
-						// We must always wait for cancelled processes to really finish since we might get new processes from host
-						// which depend on this process and then outputs might collide.
-						if (!rec.handle.WaitForExit(60*1000))
-							m_logger.Error(TC("Failed to wait for killed process to exit"));
 					}
 				}
 
@@ -3107,7 +3124,10 @@ namespace uba
 		// send and recv are swapped on purpose because that is how visualizer is visualizing 
 		m_trace.SessionUpdate(0, 0, send, recv, m_lastPing, m_memAvail, m_memTotal, m_cpuUsage);
 
-		//m_sharedMemory.TraceStats(m_trace, 30);
+		#if UBA_TRACK_SHARED_MEMORY_ALLOCATIONS
+		m_sharedMemory.TraceStats(m_trace, 30);
+		#endif
+
 		//m_fileMappingBuffer.TraceStatus(m_trace, 20, TC("Session"), 1);
 		//m_storage.TraceUpdate(m_trace, 300);
 	}
diff --git a/Engine/Source/Programs/UnrealBuildAccelerator/Common/Private/UbaSessionServer.cpp b/Engine/Source/Programs/UnrealBuildAccelerator/Common/Private/UbaSessionServer.cpp
--- a/Engine/Source/Programs/UnrealBuildAccelerator/Common/Private/UbaSessionServer.cpp
+++ b/Engine/Source/Programs/UnrealBuildAccelerator/Common/Private/UbaSessionServer.cpp
@@ -857,6 +857,8 @@ namespace uba
 				continue;
 			if (p.m_parentProcess) // Only get root processes
 				continue;
+			if (p.IsCancelled()) // Already cancelled
+				continue;
 			if (p.WaitForExit(0)) // Process already exited, let it finish its work
 				continue;
 
@@ -1533,7 +1535,7 @@ namespace uba
 		CasKey casKey = reader.ReadCasKey();
 
 		Storage::RetrieveResult res;
-		bool success = m_storage.RetrieveCasFile(res, casKey, file.destinationKey, file.destination.c_str(), nullptr, 1, true, clientId);
+		bool success = m_storage.RetrieveCasFile(res, casKey, file.destinationKey, file.destination.c_str(), nullptr, false, 1, true, clientId);
 		auto writeResponse = MakeGuard([&]() { writer.WriteBool(success); });
 
 		if (!success)
diff --git a/Engine/Source/Programs/UnrealBuildAccelerator/Common/Private/UbaStorage.cpp b/Engine/Source/Programs/UnrealBuildAccelerator/Common/Private/UbaStorage.cpp
--- a/Engine/Source/Programs/UnrealBuildAccelerator/Common/Private/UbaStorage.cpp
+++ b/Engine/Source/Programs/UnrealBuildAccelerator/Common/Private/UbaStorage.cpp
@@ -39,6 +39,7 @@ namespace uba
 			return;
 		const ConfigTable& table = *tablePtr;
 		table.GetValueAsString(rootDir, TC("RootDir"));
+		table.GetValueAsBool(writeToDisk, TC("WriteToDisk"));
 		table.GetValueAsBool(asyncUnmapViewOfFile, TC("AsyncUnmapViewOfFile"));
 		table.GetValueAsU64(casCapacityBytes, TC("CasCapacityBytes"));
 		table.GetValueAsBool(storeCompressed, TC("StoreCompressed"));
@@ -1076,12 +1077,12 @@ namespace uba
 
 		m_rootDir.Append(TCV("cas")).EnsureEndsWithSlash();
 
-		m_casDataBuffer.Init(TC("CasData"), 32ull*1024*1024*1024, 1ull*1024*1024);
-
 		m_casCompressor = info.casCompressor;
 		m_casCompressionLevel = info.casCompressionLevel;
 
 		m_sharedMemory.Init(info.sharedMemoryCapacity, info.sharedMemoryCommitSize, info.sharedMemoryTempFile, info.sharedMemoryTempFileSparse);
+
+		m_casDataBuffer.Init(TC("CasData"), 32ull*1024*1024*1024, 1ull*1024*1024);
 	}
 
 	StorageImpl::~StorageImpl()
@@ -1734,7 +1735,7 @@ namespace uba
 		return res;
 	}
 
-	bool StorageImpl::RetrieveCasFile(RetrieveResult& out, const CasKey& casKey, StringKey keyHint, const tchar* hint, FileMappingBuffer* mappingBuffer, u64 memoryMapAlignment, bool allowProxy, u32 clientId)
+	bool StorageImpl::RetrieveCasFile(RetrieveResult& out, const CasKey& casKey, StringKey keyHint, const tchar* hint, FileMappingBuffer* mappingBuffer, bool createIndependentMapping, u64 memoryMapAlignment, bool allowProxy, u32 clientId)
 	{
 		UBA_ASSERT(false);
 		return false;
diff --git a/Engine/Source/Programs/UnrealBuildAccelerator/Common/Private/UbaStorageClient.cpp b/Engine/Source/Programs/UnrealBuildAccelerator/Common/Private/UbaStorageClient.cpp
--- a/Engine/Source/Programs/UnrealBuildAccelerator/Common/Private/UbaStorageClient.cpp
+++ b/Engine/Source/Programs/UnrealBuildAccelerator/Common/Private/UbaStorageClient.cpp
@@ -213,7 +213,7 @@ namespace uba
 		return true;
 	}
 
-	bool StorageClient::RetrieveCasFile(RetrieveResult& out, const CasKey& casKeyTmp, StringKey keyHint, const tchar* hint, FileMappingBuffer* mappingBuffer, u64 memoryMapAlignment, bool allowProxy, u32 clientId)
+	bool StorageClient::RetrieveCasFile(RetrieveResult& out, const CasKey& casKeyTmp, StringKey keyHint, const tchar* hint, FileMappingBuffer* mappingBuffer, bool createIndependentMapping, u64 memoryMapAlignment, bool allowProxy, u32 clientId)
 	{
 		CasKey casKey = casKeyTmp;
 
@@ -578,7 +578,7 @@ namespace uba
 						UBA_ASSERT(!writeMem || mappedView.size == sizeOnDisk);
 						if (!writeMem)
 						{
-							mappedView = mappingBuffer->AllocAndMapView(sizeOnDisk, memoryMapAlignment, hint);
+							mappedView = mappingBuffer->AllocAndMapView(sizeOnDisk, memoryMapAlignment, hint, createIndependentMapping);
 							writeMem = mappedView.memory;
 							if (!writeMem)
 								return false;
diff --git a/Engine/Source/Programs/UnrealBuildAccelerator/Common/Private/UbaStorageServer.cpp b/Engine/Source/Programs/UnrealBuildAccelerator/Common/Private/UbaStorageServer.cpp
--- a/Engine/Source/Programs/UnrealBuildAccelerator/Common/Private/UbaStorageServer.cpp
+++ b/Engine/Source/Programs/UnrealBuildAccelerator/Common/Private/UbaStorageServer.cpp
@@ -138,7 +138,7 @@ namespace uba
 		return true;
 	}
 
-	bool StorageServer::RetrieveCasFile(RetrieveResult& out, const CasKey& casKey, StringKey keyHint, const tchar* hint, FileMappingBuffer* mappingBuffer, u64 memoryMapAlignment, bool allowProxy, u32 clientId)
+	bool StorageServer::RetrieveCasFile(RetrieveResult& out, const CasKey& casKey, StringKey keyHint, const tchar* hint, FileMappingBuffer* mappingBuffer, bool createIndependentMapping, u64 memoryMapAlignment, bool allowProxy, u32 clientId)
 	{
 		UBA_ASSERT(!mappingBuffer);
 		UBA_ASSERT(casKey != CasKeyZero);
diff --git a/Engine/Source/Programs/UnrealBuildAccelerator/Common/Public/UbaStorage.h b/Engine/Source/Programs/UnrealBuildAccelerator/Common/Public/UbaStorage.h
--- a/Engine/Source/Programs/UnrealBuildAccelerator/Common/Public/UbaStorage.h
+++ b/Engine/Source/Programs/UnrealBuildAccelerator/Common/Public/UbaStorage.h
@@ -52,7 +52,7 @@ namespace uba
 		virtual bool GetFileData(void* outData, const CasKey& casKey, const tchar* hint) = 0;
 
 		struct RetrieveResult { CasKey casKey; u64 size = 0; MappedView2 view; };
-		virtual bool RetrieveCasFile(RetrieveResult& out, const CasKey& casKey, StringKey keyHint, const tchar* hint, FileMappingBuffer* mappingBuffer = nullptr, u64 memoryMapAlignment = 1, bool allowProxy = true, u32 clientId = 0) = 0;
+		virtual bool RetrieveCasFile(RetrieveResult& out, const CasKey& casKey, StringKey keyHint, const tchar* hint, FileMappingBuffer* mappingBuffer = nullptr, bool createIndependentMapping = false, u64 memoryMapAlignment = 1, bool allowProxy = true, u32 clientId = 0) = 0;
 
 		virtual bool AddRef(const CasKey& casKey) = 0;
 		virtual bool ReleaseRef(const CasKey& casKey, bool dropIfZero) = 0;
@@ -173,7 +173,7 @@ namespace uba
 		virtual bool DeleteCasForFile(const tchar* file) override;
 		virtual bool GetFileSize(u64& outSize, const CasKey& casKey, const tchar* hint) override;
 		virtual bool GetFileData(void* outData, const CasKey& casKey, const tchar* hint) override;
-		virtual bool RetrieveCasFile(RetrieveResult& out, const CasKey& casKey, StringKey keyHint, const tchar* hint, FileMappingBuffer* mappingBuffer = nullptr, u64 memoryMapAlignment = 1, bool allowProxy = true, u32 clientId = 0) override;
+		virtual bool RetrieveCasFile(RetrieveResult& out, const CasKey& casKey, StringKey keyHint, const tchar* hint, FileMappingBuffer* mappingBuffer = nullptr, bool createIndependentMapping = false, u64 memoryMapAlignment = 1, bool allowProxy = true, u32 clientId = 0) override;
 		virtual bool AddRef(const CasKey& casKey) override;
 		virtual bool ReleaseRef(const CasKey& casKey, bool dropIfZero) override;
 		virtual bool VerifyAndGetCachedFileInfo(CachedFileInfo& out, StringKey fileNameKey, u64 verifiedLastWriteTime, u64 verifiedSize) override;
diff --git a/Engine/Source/Programs/UnrealBuildAccelerator/Common/Public/UbaStorageClient.h b/Engine/Source/Programs/UnrealBuildAccelerator/Common/Public/UbaStorageClient.h
--- a/Engine/Source/Programs/UnrealBuildAccelerator/Common/Public/UbaStorageClient.h
+++ b/Engine/Source/Programs/UnrealBuildAccelerator/Common/Public/UbaStorageClient.h
@@ -56,7 +56,7 @@ namespace uba
 		virtual MappedView2 MapView(const CasKey& casKey, const tchar* hint) override;
 
 		virtual bool GetZone(StringBufferBase& out) override;
-		virtual bool RetrieveCasFile(RetrieveResult& out, const CasKey& casKey, StringKey keyHint, const tchar* hint, FileMappingBuffer* mappingBuffer = nullptr, u64 memoryMapAlignment = 1, bool allowProxy = true, u32 clientId = 0) override;
+		virtual bool RetrieveCasFile(RetrieveResult& out, const CasKey& casKey, StringKey keyHint, const tchar* hint, FileMappingBuffer* mappingBuffer = nullptr, bool createIndependentMapping = false, u64 memoryMapAlignment = 1, bool allowProxy = true, u32 clientId = 0) override;
 		virtual bool StoreCasFile(CasKey& out, const tchar* fileName, const CasKey& casKeyOverride, bool deferCreation) override;
 		virtual bool StoreCasFileClient(CasKey& out, StringKey fileNameKey, const tchar* fileName, SharedMemoryHandle mappingFile, u64 mappingOffset, u64 fileSize, const tchar* hint, bool keepMappingInMemory = false, bool storeCompressed = true) override;
 		virtual bool HasCasFile(const CasKey& casKey, CasEntry** out = nullptr, bool useMapping = false) override;
diff --git a/Engine/Source/Programs/UnrealBuildAccelerator/Common/Public/UbaStorageServer.h b/Engine/Source/Programs/UnrealBuildAccelerator/Common/Public/UbaStorageServer.h
--- a/Engine/Source/Programs/UnrealBuildAccelerator/Common/Public/UbaStorageServer.h
+++ b/Engine/Source/Programs/UnrealBuildAccelerator/Common/Public/UbaStorageServer.h
@@ -46,7 +46,7 @@ namespace uba
 	protected:
 
 		virtual bool GetZone(StringBufferBase& out) override;
-		virtual bool RetrieveCasFile(RetrieveResult& out, const CasKey& casKey, StringKey keyHint, const tchar* hint, FileMappingBuffer* mappingBuffer = nullptr, u64 memoryMapAlignment = 1, bool allowProxy = true, u32 clientId = 0) override;
+		virtual bool RetrieveCasFile(RetrieveResult& out, const CasKey& casKey, StringKey keyHint, const tchar* hint, FileMappingBuffer* mappingBuffer = nullptr, bool createIndependentMapping = false, u64 memoryMapAlignment = 1, bool allowProxy = true, u32 clientId = 0) override;
 		virtual bool IsDisallowedPath(const tchar* fileName) override;
 		virtual void SetTrace(Trace* trace, bool detailed) override;
 		virtual bool HasProxy(u32 clientId) override;
