[UBA]
* Fixed TestLogLinesNoDetour unit test

diff --git a/Engine/Source/Programs/UnrealBuildAccelerator/Common/Private/UbaProcess.cpp b/Engine/Source/Programs/UnrealBuildAccelerator/Common/Private/UbaProcess.cpp
--- a/Engine/Source/Programs/UnrealBuildAccelerator/Common/Private/UbaProcess.cpp
+++ b/Engine/Source/Programs/UnrealBuildAccelerator/Common/Private/UbaProcess.cpp
@@ -168,12 +168,11 @@ namespace uba
 		#endif
 	,	m_detourEnabled(detourEnabled)
 	,	m_shared(parent ? parent->m_shared : *new Shared)
+	,	m_messageThreadDone(EventResetType_Manual)
 	{
 		if (!m_comMemory.memory)
 			return;
 
-		m_messageThreadDone.Create(EventResetType_Manual);
-
 		#if !PLATFORM_WINDOWS
 		if (UseSemaphoreForSharedEvents())
 		{
