[UBA]
* Added name to MemoryBlock to be able to get better errors

diff --git a/Engine/Source/Programs/UnrealBuildAccelerator/Common/Private/UbaCacheClient.cpp b/Engine/Source/Programs/UnrealBuildAccelerator/Common/Private/UbaCacheClient.cpp
--- a/Engine/Source/Programs/UnrealBuildAccelerator/Common/Private/UbaCacheClient.cpp
+++ b/Engine/Source/Programs/UnrealBuildAccelerator/Common/Private/UbaCacheClient.cpp
@@ -1045,12 +1045,12 @@ namespace uba
 		{
 			DowngradedLogger logger(m_connected, m_logger.m_writer, TC("UbaCacheClientNormalizedDownload"));
 			// Fetch into memory, file is in special format without absolute paths
-			MemoryBlock normalizedBlock(4*1024*1024);
+			MemoryBlock normalizedBlock(TC("FetchFileNormalized"), 4*1024*1024);
 			bool destinationIsCompressed = false;
 			if (!fetcher.RetrieveFile(logger, m_client, casKey, path.data, destinationIsCompressed, &normalizedBlock))
 				return logger.Debug(TC("Failed to download cache output for %s"), info.GetDescription()).ToFalse();
 
-			MemoryBlock localBlock(4*1024*1024);
+			MemoryBlock localBlock(TC("FetchFileDenormalize"), 4*1024*1024);
 
 			u32 rootOffsets = *(u32*)(normalizedBlock.memory);
 			char* fileStart = (char*)(normalizedBlock.memory + sizeof(u32));
@@ -1200,7 +1200,7 @@ namespace uba
 		}
 		else
 		{
-			MemoryBlock block(512*1024*1024);
+			MemoryBlock block(TC("ExecuteCommand"), 512*1024*1024);
 			if (!fetcher.RetrieveFile(m_logger, m_client, statusFileCasKey, TC("CommandString"), destinationIsCompressed, &block))
 				return false;
 			BinaryReader reader(block.memory, 3, block.writtenSize); // Skipping bom
@@ -1426,7 +1426,7 @@ namespace uba
 				FileAccessor file(m_logger, path.data);
 				if (!file.OpenMemoryRead())
 					return false;
-				MemoryBlock block(AlignUp(file.GetSize() + 16, PageSize));
+				MemoryBlock block(TC("SendNormalizedFile"), AlignUp(file.GetSize() + 16, PageSize));
 				u32& rootOffsetsStart = *(u32*)block.Allocate(sizeof(u32), 1, TC(""));
 				rootOffsetsStart = 0;
 				Vector<u32> rootOffsets;
diff --git a/Engine/Source/Programs/UnrealBuildAccelerator/Common/Private/UbaCacheServer.cpp b/Engine/Source/Programs/UnrealBuildAccelerator/Common/Private/UbaCacheServer.cpp
--- a/Engine/Source/Programs/UnrealBuildAccelerator/Common/Private/UbaCacheServer.cpp
+++ b/Engine/Source/Programs/UnrealBuildAccelerator/Common/Private/UbaCacheServer.cpp
@@ -769,7 +769,7 @@ namespace uba
 
 		struct CasFileInfo { CasFileInfo(u32 s = 0) : size(s) {} u32 size; Atomic<bool> isUsed = false; }; // These are compressed cas, should never be over 4gb
 
-		MemoryBlock existingCasMemoryBlock;
+		MemoryBlock existingCasMemoryBlock(TC("ExistingCasMemoryBlock"));
 		HashMap<CasKey, CasFileInfo> existingCas;
 
 
@@ -882,7 +882,7 @@ namespace uba
 				Bucket& bucket = **it;
 				auto context = bucket.m_maintenanceContext;
 				if (!context)
-					context = bucket.m_maintenanceContext = new Bucket::MaintenanceContext;
+					context = bucket.m_maintenanceContext = new Bucket::MaintenanceContext(TC("MaintenanceContext"));
 
 				// Traverse all deleted cas files and create a bit table which has the bit of caskeytable offset set to true if deleted
 				// This is created to make fast lookup further down
@@ -2496,7 +2496,7 @@ namespace uba
 		else if (command.Equals(TCV("updateentries")))
 		{
 			HashMap2<u32, u32> oldToNewCasKeyOffset;
-			MemoryBlock memory(64*1024);
+			MemoryBlock memory(TC("UpdateEntries"), 64*1024);
 			oldToNewCasKeyOffset.Init(memory, 1);
 			Vector<u32> temp;
 			Vector<u8> temp2;
diff --git a/Engine/Source/Programs/UnrealBuildAccelerator/Common/Private/UbaCompactTables.cpp b/Engine/Source/Programs/UnrealBuildAccelerator/Common/Private/UbaCompactTables.cpp
--- a/Engine/Source/Programs/UnrealBuildAccelerator/Common/Private/UbaCompactTables.cpp
+++ b/Engine/Source/Programs/UnrealBuildAccelerator/Common/Private/UbaCompactTables.cpp
@@ -11,7 +11,8 @@ namespace uba
 
 
 	CompactPathTable::CompactPathTable(bool caseInsensitive, u64 reservePathCount, u64 reserveSegmentCount, u32 version)
-	:	m_version(version)
+	:	m_pathTableMem(TC("CompactPathTable"))
+	,	m_version(version)
 	,	m_caseInsensitive(caseInsensitive)
 	{
 		m_offsets.Init(reservePathCount);
@@ -23,7 +24,7 @@ namespace uba
 		if (!m_pathTableMem.memory)
 			m_pathTableMem.Init(PathTableMaxSize);
 		if (!m_pathTableMem.writtenSize)
-			m_pathTableMem.AllocateNoLock(1, 1, TC("CompactPathTable"));
+			m_pathTableMem.AllocateNoLock(1, 1, TC("Init"));
 	}
 
 	u32 CompactPathTable::Add(const tchar* str, u64 strLen)
@@ -257,7 +258,7 @@ namespace uba
 					const tchar* seg = *it;
 					u32 segLen = TStrlen(seg);
 					u64 writtenSize = m_pathTableMem.writtenSize;
-					BinaryWriter writer((u8*)m_pathTableMem.AllocateNoLock(segLen+1, 1, TC("CompactPathTable")), 0, segLen+1);
+					BinaryWriter writer((u8*)m_pathTableMem.AllocateNoLock(segLen+1, 1, TC("AddSegments")), 0, segLen+1);
 					StringKey segmentKey = ToStringKeyNoCheck(seg, segLen);
 					u32& strOffset = m_segmentOffsets.Insert(segmentKey);
 					strOffset = u32(writtenSize);
@@ -355,7 +356,7 @@ namespace uba
 			// Put string directly after current element and set segment offset to 0
 			u64 bytesForString = GetStringWriteSize(seg, segLen);
 			u64 memSize = bytesForParent + 1 + bytesForString;
-			u8* mem = (u8*)m_pathTableMem.AllocateNoLock(memSize, 1, TC("CompactPathTable"));
+			u8* mem = (u8*)m_pathTableMem.AllocateNoLock(memSize, 1, TC("Add1"));
 			BinaryWriter writer(mem, 0, memSize);
 			writer.Write7BitEncoded(parentOffset);
 			writer.Write7BitEncoded(0);
@@ -373,7 +374,7 @@ namespace uba
 		#endif
 
 		u64 memSize = bytesForParent + Get7BitEncodedCount(strOffset);
-		u8* mem = (u8*)m_pathTableMem.AllocateNoLock(memSize, 1, TC("CompactPathTable"));
+		u8* mem = (u8*)m_pathTableMem.AllocateNoLock(memSize, 1, TC("Add2"));
 		BinaryWriter writer(mem, 0, memSize);
 		writer.Write7BitEncoded(parentOffset);
 		writer.Write7BitEncoded(strOffset);
@@ -390,7 +391,7 @@ namespace uba
 		u8 bytesForParent = Get7BitEncodedCount(parentRelativeOffset);
 		u64 memSize = bytesForParent + 1 + bytesForString;
 	
-		BinaryWriter writer((u8*)m_pathTableMem.AllocateNoLock(memSize, 1, TC("CompactPathTable")), 0, memSize);
+		BinaryWriter writer((u8*)m_pathTableMem.AllocateNoLock(memSize, 1, TC("Add3")), 0, memSize);
 		writer.Write7BitEncoded(parentRelativeOffset);
 		writer.Write7BitEncoded(0);
 		writer.WriteString(segment);
@@ -403,7 +404,7 @@ namespace uba
 
 		u8 bytesForParent = Get7BitEncodedCount(parentRelativeOffset);
 		u64 memSize = bytesForParent + Get7BitEncodedCount(segmentOffset);
-		BinaryWriter writer((u8*)m_pathTableMem.AllocateNoLock(memSize, 1, TC("CompactPathTable")), 0, memSize);
+		BinaryWriter writer((u8*)m_pathTableMem.AllocateNoLock(memSize, 1, TC("Add4")), 0, memSize);
 		writer.Write7BitEncoded(parentRelativeOffset);
 		writer.Write7BitEncoded(segmentOffset);
 	}
@@ -563,7 +564,7 @@ namespace uba
 
 		u64 writtenSize = m_pathTableMem.writtenSize;
 		u64 left = reader.GetLeft();
-		void* mem = m_pathTableMem.AllocateNoLock(left, 1, TC("CompactPathTable"));
+		void* mem = m_pathTableMem.AllocateNoLock(left, 1, TC("ReadMem"));
 		reader.ReadBytes(mem, left);
 
 		if (!populateLookup)
@@ -599,12 +600,12 @@ namespace uba
 	{
 		if (!m_pathTableMem.memory)
 			m_pathTableMem.Init(PathTableMaxSize);
-		return (u8*)m_pathTableMem.CommitNoLock(size, TC("CompactPathTable::BeginCommit"));
+		return (u8*)m_pathTableMem.CommitNoLock(size, TC("BeginCommit"));
 	}
 
 	void CompactPathTable::EndCommit(u8* data, u64 written)
 	{
-		m_pathTableMem.AllocateNoLock(written, 1, TC("CompactPathTable::EndCommit"));
+		m_pathTableMem.AllocateNoLock(written, 1, TC("EndCommit"));
 	}
 
 	void CompactPathTable::Swap(CompactPathTable& other)
@@ -618,6 +619,7 @@ namespace uba
 	}
 
 	CompactCasKeyTable::CompactCasKeyTable(u64 reserveOffsetsCount)
+	:	m_casKeyTableMem(TC("CompactCasKeyTable"))
 	{
 		if (reserveOffsetsCount)
 			m_offsets.Init(reserveOffsetsCount);
@@ -652,7 +654,7 @@ namespace uba
 		if (added)
 		{
 			u8 bytesForStringOffset = Get7BitEncodedCount(stringOffset);
-			u8* mem = (u8*)m_casKeyTableMem.AllocateNoLock(bytesForStringOffset + sizeof(CasKey), 1, TC("CompactCasKeyTable"));
+			u8* mem = (u8*)m_casKeyTableMem.AllocateNoLock(bytesForStringOffset + sizeof(CasKey), 1, TC("Add"));
 			BinaryWriter writer(mem, 0, 1000);
 			writer.Write7BitEncoded(stringOffset);
 			writer.WriteCasKey(casKey);
@@ -680,7 +682,7 @@ namespace uba
 		if (added)
 		{
 			u8 bytesForStringOffset = Get7BitEncodedCount(stringOffset);
-			u8* mem = (u8*)m_casKeyTableMem.AllocateNoLock(bytesForStringOffset + sizeof(CasKey), 1, TC("CompactCasKeyTable"));
+			u8* mem = (u8*)m_casKeyTableMem.AllocateNoLock(bytesForStringOffset + sizeof(CasKey), 1, TC("AddNoLock"));
 			BinaryWriter writer(mem, 0, 1000);
 			writer.Write7BitEncoded(stringOffset);
 			writer.WriteCasKey(casKey);
diff --git a/Engine/Source/Programs/UnrealBuildAccelerator/Common/Private/UbaNetworkBackendTcp.cpp b/Engine/Source/Programs/UnrealBuildAccelerator/Common/Private/UbaNetworkBackendTcp.cpp
--- a/Engine/Source/Programs/UnrealBuildAccelerator/Common/Private/UbaNetworkBackendTcp.cpp
+++ b/Engine/Source/Programs/UnrealBuildAccelerator/Common/Private/UbaNetworkBackendTcp.cpp
@@ -356,6 +356,8 @@ namespace uba
 		SCOPED_CRITICAL_SECTION(conn.shutdownLock, lock);
 		if (conn.socket == InvalidSocket)
 			return;
+		if (conn.shutdownCalled)
+			return;
 		ShutdownSocket(conn.logger, conn.socket, TC("Shutdown"));
 		conn.shutdownCalled = true;
 	}
@@ -773,7 +775,9 @@ namespace uba
 
 		if (s == InvalidSocket)
 			return;
-		ShutdownSocket(logger, s, TC("ThreadRecv"));
+		if (!connection.shutdownCalled)
+			ShutdownSocket(logger, s, TC("ThreadRecv"));
+		connection.shutdownCalled = true;
 		CloseSocket(logger, s, TC("ThreadRecv"));
 	}
 
@@ -1074,8 +1078,12 @@ namespace uba
 		SCOPED_CRITICAL_SECTION(connection.shutdownLock, lock2);
 		SCOPED_FUTEX(connection.sendLock, lock);
 		Socket s = connection.socket;
-		if (s != InvalidSocket)
-			ShutdownSocket(logger, s, TC("ThreadSend"));
+		if (s == InvalidSocket)
+			return;
+		if (connection.shutdownCalled)
+			return;
+		ShutdownSocket(logger, s, TC("ThreadSend"));
+		connection.shutdownCalled = true;
 	}
 
 	bool NetworkBackendTcp::HandleSendError(Connection& connection, Logger& logger, const tchar* hint)
diff --git a/Engine/Source/Programs/UnrealBuildAccelerator/Common/Private/UbaObjectFile.cpp b/Engine/Source/Programs/UnrealBuildAccelerator/Common/Private/UbaObjectFile.cpp
--- a/Engine/Source/Programs/UnrealBuildAccelerator/Common/Private/UbaObjectFile.cpp
+++ b/Engine/Source/Programs/UnrealBuildAccelerator/Common/Private/UbaObjectFile.cpp
@@ -252,7 +252,7 @@ namespace uba
 		ObjectFileCoff objectFileCoff;
 		ObjectFileElf objectFileElf;
 		
-		MemoryBlock memoryBlock(64*1024*1024);
+		MemoryBlock memoryBlock(TC("CreateExtraFile"), 64*1024*1024);
 
 		bool res;
 		if (extraObjFilename.EndsWith(TCV(".obj")))
diff --git a/Engine/Source/Programs/UnrealBuildAccelerator/Common/Private/UbaSession.cpp b/Engine/Source/Programs/UnrealBuildAccelerator/Common/Private/UbaSession.cpp
--- a/Engine/Source/Programs/UnrealBuildAccelerator/Common/Private/UbaSession.cpp
+++ b/Engine/Source/Programs/UnrealBuildAccelerator/Common/Private/UbaSession.cpp
@@ -325,7 +325,7 @@ namespace uba
 					if (!destinationFile.CreateWrite(false, DefaultAttributes(), 0, m_tempPath.data))
 						return false;
 
-					MemoryBlock block(5*1024*1024);
+					MemoryBlock block(TC("WriteFileToDiskDevirtualize"), 5*1024*1024);
 					if (!DevirtualizeDepsFile(rootsHandle, block, mem, fileSize, escapeSpaces, file.name.c_str()))
 						return false;
 
@@ -1788,8 +1788,10 @@ namespace uba
 	,	m_logger(info.logWriter, logPrefix)
 	,	m_workManager(workManager)
 	,	m_ownsTrace(info.trace == nullptr)
+	,	m_directoryTableMemory(TC("SessionDirectoryTable"))
 	,	m_directoryTable(m_directoryTableMemory)
 	,	m_fileMappingBuffer(m_logger, m_sharedMemory)
+	,	m_nameToHashTableMem(TC("SessionNameToHashTable"))
 	,	m_processCommunicationAllocator(m_logger, TC("CommunicationAllocator"))
 	,	m_processMessageThreads(*new WorkManagerImpl(~0u, TC("UbaProcMsg")))
 	,	m_trace(info.trace ? *info.trace : *new Trace(info.logWriter))
@@ -3365,7 +3367,7 @@ namespace uba
 						return false;
 					void* mem = readFile.GetData();
 					fileSize = readFile.GetSize();
-					MemoryBlock block(5*1024*1024);
+					MemoryBlock block(TC("DevirtualizeDepsFile"), 5*1024*1024);
 					RootsHandle rootsHandle = startInfo.rootsHandle;
 					if (!DevirtualizeDepsFile(rootsHandle, block, mem, fileSize, escapeSpaces, name.data))
 						return false;
@@ -4099,7 +4101,7 @@ namespace uba
 		#if UBA_DEBUG
 		verbose = true;
 		#endif
-		MemoryBlock memoryBlock(64*1024*1024); // Seen files that are over 32mb when verbose = true
+		MemoryBlock memoryBlock(TC("ExtractSymbols"), 64*1024*1024); // Seen files that are over 32mb when verbose = true
 		if (!objectFile->WriteImportsAndExports(m_logger, memoryBlock, verbose))
 			return false;
 
diff --git a/Engine/Source/Programs/UnrealBuildAccelerator/Common/Private/UbaStorage.cpp b/Engine/Source/Programs/UnrealBuildAccelerator/Common/Private/UbaStorage.cpp
--- a/Engine/Source/Programs/UnrealBuildAccelerator/Common/Private/UbaStorage.cpp
+++ b/Engine/Source/Programs/UnrealBuildAccelerator/Common/Private/UbaStorage.cpp
@@ -2591,7 +2591,7 @@ namespace uba
 					if (!decompressToMemory(slot))
 						return false;
 
-					MemoryBlock block(5*1024*1024);
+					MemoryBlock block(TC("Formatting"), 5*1024*1024);
 
 					if (!formattingFunc(block, slot, decompressedSize, destination))
 						return false;
diff --git a/Engine/Source/Programs/UnrealBuildAccelerator/Common/Private/UbaStorageUtils.cpp b/Engine/Source/Programs/UnrealBuildAccelerator/Common/Private/UbaStorageUtils.cpp
--- a/Engine/Source/Programs/UnrealBuildAccelerator/Common/Private/UbaStorageUtils.cpp
+++ b/Engine/Source/Programs/UnrealBuildAccelerator/Common/Private/UbaStorageUtils.cpp
@@ -390,7 +390,7 @@ namespace uba
 		u64 messageHeader = client.GetMessageHeaderSize();
 		u64 messageHeaderMaxSize = messageHeader + firstMessageOverHead;
 
-		MemoryBlock memoryBlock(sourceSize + messageHeaderMaxSize + 1024);
+		MemoryBlock memoryBlock(TC("FileSender"), sourceSize + messageHeaderMaxSize + 1024);
 		{
 			const u8* uncompressedData = sourceMem;
 			u8* compressBufferStart = memoryBlock.memory + messageHeaderMaxSize;
diff --git a/Engine/Source/Programs/UnrealBuildAccelerator/Common/Public/UbaHashMap.h b/Engine/Source/Programs/UnrealBuildAccelerator/Common/Public/UbaHashMap.h
--- a/Engine/Source/Programs/UnrealBuildAccelerator/Common/Public/UbaHashMap.h
+++ b/Engine/Source/Programs/UnrealBuildAccelerator/Common/Public/UbaHashMap.h
@@ -43,7 +43,7 @@ namespace uba
 			u64 lookupSize = lookupCount*sizeof(u32);
 			u64 entriesSize = sizeof(Entry)*(reserveCount+1);
 			u64 valuesSize = sizeof(Value)*(reserveCount+1);
-			m_memory = new MemoryBlock();
+			m_memory = new MemoryBlock(hint);
 			m_memory->Init(lookupSize + entriesSize + valuesSize);
 			m_ownsMemory = true;
 
diff --git a/Engine/Source/Programs/UnrealBuildAccelerator/Core/Private/UbaMemory.cpp b/Engine/Source/Programs/UnrealBuildAccelerator/Core/Private/UbaMemory.cpp
--- a/Engine/Source/Programs/UnrealBuildAccelerator/Core/Private/UbaMemory.cpp
+++ b/Engine/Source/Programs/UnrealBuildAccelerator/Core/Private/UbaMemory.cpp
@@ -20,13 +20,15 @@ namespace uba
 	constexpr u64 MemoryBlock_ReserveAlign = 1 * 1024 * 1024;
 
 
-	MemoryBlock::MemoryBlock(u64 reserveSize_, void* baseAddress_)
+	MemoryBlock::MemoryBlock(const tchar* name_, u64 reserveSize_, void* baseAddress_)
 	{
+		name = name_;
 		Init(reserveSize_, baseAddress_, false);
 	}
 
-	MemoryBlock::MemoryBlock(u8* baseAddress_)
+	MemoryBlock::MemoryBlock(const tchar* name_, u8* baseAddress_)
 	{
+		name = name_;
 		memory = baseAddress_;
 	}
 
@@ -41,7 +43,7 @@ namespace uba
 		reserveSize = AlignUp(reserveSize_, MemoryBlock_ReserveAlign);
 		memory = (u8*)VirtualAlloc(baseAddress_, reserveSize, MEM_RESERVE, PAGE_READWRITE); // Max size of obj file?
 		if (!memory)
-			FatalError(1347, TC("Failed to reserve virtual memory (%u)"), GetLastError());
+			FatalError(1347, TC("Failed to reserve virtual memory in %s (%u)"), name, GetLastError());
 		#else
 
 		int flags = MAP_PRIVATE | MAP_ANONYMOUS;
@@ -61,12 +63,12 @@ namespace uba
 		{
 			if (useHugePages)
 				return false;
-			FatalError(1347, "mmap failed to reserve %llu bytes (asking for %llu): %s", reserveSize, reserveSize_, strerror(errno));
+			FatalError(1347, "mmap failed to reserve %llu bytes in %s (asking for %llu) (%s)", reserveSize, name, reserveSize_, strerror(errno));
 		}
 		#endif
 
 		if (baseAddress_ && baseAddress_ != memory)
-			FatalError(9881, TC("Failed to reserve virtual memory at address (%u)"), GetLastError());
+			FatalError(9881, TC("Failed to reserve virtual memory at specific address in %s (%u)"), name, GetLastError());
 		return true;
 	}
 
@@ -79,7 +81,7 @@ namespace uba
 		VirtualFree(memory, 0, MEM_RELEASE);
 		#else
 		if (munmap(memory, reserveSize) == -1)
-			FatalError(9885, "munmap failed to free %llu bytes: %s", reserveSize, strerror(errno));
+			FatalError(9885, "munmap failed to free %llu bytes in %s (%s)", reserveSize, name, strerror(errno));
 		#endif
 		memory = nullptr;
 	}
@@ -99,7 +101,7 @@ namespace uba
 		u64 newPos = startPos + bytes;
 		
 		if (newPos > reserveSize)
-			FatalError(9882, TC("Ran out of reserved virtual address space. Reserved %llu, Needed %llu (%s)"), reserveSize, newPos, hint);
+			FatalError(9882, TC("Ran out of reserved virtual address space in %s. Reserved %llu, Needed %llu (%s)"), name, reserveSize, newPos, hint);
 
 		#if PLATFORM_WINDOWS
 		if (newPos > committedSize)
@@ -109,7 +111,7 @@ namespace uba
 				toCommit = reserveSize - committedSize;
 			void* addr = memory + committedSize;
 			if (!VirtualAlloc(addr, toCommit, MEM_COMMIT, PAGE_READWRITE))
-				FatalError(9883, TC("Failed to commit virtual memory for memory block. Total size %llu (%u) (%s)"), committedSize + toCommit, GetLastError(), hint);
+				FatalError(9883, TC("Failed to commit virtual memory for memory block in %s. Total size %llu (%u) (%s)"), name, committedSize + toCommit, GetLastError(), hint);
 			if (lockMemory)
 				VirtualLock(addr, toCommit);
 			committedSize += toCommit;
@@ -131,7 +133,7 @@ namespace uba
 		if (committedSize + toCommit > reserveSize)
 			toCommit = reserveSize - committedSize;
 		if (!VirtualAlloc(memory + committedSize, toCommit, MEM_COMMIT, PAGE_READWRITE))
-			FatalError(9883, TC("Failed to commit virtual memory for memory block. Total size %llu (%u) (%s)"), committedSize + toCommit, GetLastError(), hint);
+			FatalError(9883, TC("Failed to commit virtual memory for memory block in %s. Total size %llu (%u) (%s)"), name, committedSize + toCommit, GetLastError(), hint);
 		committedSize += toCommit;
 		#endif
 		return memory + writtenSize;
diff --git a/Engine/Source/Programs/UnrealBuildAccelerator/Core/Public/UbaMemory.h b/Engine/Source/Programs/UnrealBuildAccelerator/Core/Public/UbaMemory.h
--- a/Engine/Source/Programs/UnrealBuildAccelerator/Core/Public/UbaMemory.h
+++ b/Engine/Source/Programs/UnrealBuildAccelerator/Core/Public/UbaMemory.h
@@ -92,8 +92,8 @@ namespace uba
 
 	struct MemoryBlock
 	{
-		MemoryBlock(u64 reserveSize_, void* baseAddress_ = nullptr);
-		MemoryBlock(u8* baseAddress_ = nullptr);
+		MemoryBlock(const tchar* name_, u64 reserveSize_, void* baseAddress_ = nullptr);
+		MemoryBlock(const tchar* name_, u8* baseAddress_ = nullptr);
 		~MemoryBlock();
 		bool Init(u64 reserveSize_, void* baseAddress_ = nullptr, bool useHugePages = false);
 		void Deinit();
@@ -105,12 +105,12 @@ namespace uba
 		tchar* Strdup(const tchar* str);
 		void Swap(MemoryBlock& other);
 		
-
 		Futex lock;
 		u8* memory;
 		u64 reserveSize = 0;
 		u64 writtenSize = 0;
 		u64 committedSize = 0;
+		const tchar* name;
 	};
 
 
diff --git a/Engine/Source/Programs/UnrealBuildAccelerator/Detours/Private/UbaDetoursShared.cpp b/Engine/Source/Programs/UnrealBuildAccelerator/Detours/Private/UbaDetoursShared.cpp
--- a/Engine/Source/Programs/UnrealBuildAccelerator/Detours/Private/UbaDetoursShared.cpp
+++ b/Engine/Source/Programs/UnrealBuildAccelerator/Detours/Private/UbaDetoursShared.cpp
@@ -59,7 +59,7 @@ namespace uba
 		g_vfsMatchingLength = 0;
 
 		u64 reserveSizeMb = IsWindows ? 256 : 1024; // The sync primitives on linux/macos is much bigger
-		g_memoryBlockMem.Create(reserveSizeMb * 1024 * 1024);
+		g_memoryBlockMem.Create(TC("GlobalDetour"), reserveSizeMb * 1024 * 1024);
 		g_directoryTableMem.Create(g_memoryBlock);
 		g_mappedFileTableMem.Create(g_memoryBlock);
 		g_consoleStringCsMem.Create();
diff --git a/Engine/Source/Programs/UnrealBuildAccelerator/Detours/Private/Windows/UbaDetoursFunctionsWin.cpp b/Engine/Source/Programs/UnrealBuildAccelerator/Detours/Private/Windows/UbaDetoursFunctionsWin.cpp
--- a/Engine/Source/Programs/UnrealBuildAccelerator/Detours/Private/Windows/UbaDetoursFunctionsWin.cpp
+++ b/Engine/Source/Programs/UnrealBuildAccelerator/Detours/Private/Windows/UbaDetoursFunctionsWin.cpp
@@ -137,7 +137,7 @@ constexpr u64 DetouredHandleStart = 600000; // Let's hope noone uses the handles
 constexpr u64 DetouredHandleEnd = DetouredHandleStart + DetouredHandleMaxCount;
 constexpr u64 DetouredHandlesMemReserve = AlignUp(DetouredHandleMaxCount*sizeof(DetouredHandle), 64*1024);
 constexpr u64 DetouredHandlesMemStart = 0;
-MemoryBlock& g_detouredHandleMemoryBlock = *new MemoryBlock(DetouredHandlesMemReserve, (void*)DetouredHandlesMemStart);
+MemoryBlock& g_detouredHandleMemoryBlock = *new MemoryBlock(TC("DetouredHandles"), DetouredHandlesMemReserve, (void*)DetouredHandlesMemStart);
 u64 g_detouredHandlesStart = u64(g_detouredHandleMemoryBlock.memory);
 u64 g_detouredHandlesEnd = g_detouredHandlesStart + g_detouredHandleMemoryBlock.reserveSize;
 BlockAllocator<DetouredHandle> g_detouredHandleAllocator(g_detouredHandleMemoryBlock);
@@ -181,7 +181,7 @@ struct ListDirectoryHandle
 };
 
 constexpr u64 ListDirHandlesRange = 4*1024*1024;
-MemoryBlock& g_listDirHandleMemoryBlock = *new MemoryBlock(ListDirHandlesRange);
+MemoryBlock& g_listDirHandleMemoryBlock = *new MemoryBlock(TC("DetouredDirHandles"), ListDirHandlesRange);
 constexpr u64 ListDirHandlesStart = DetouredHandleEnd;
 constexpr u64 ListDirHandlesEnd = ListDirHandlesStart + ListDirHandlesRange/sizeof(ListDirectoryHandle);
 BlockAllocator<ListDirectoryHandle>& g_listDirectoryHandleAllocator = *new BlockAllocator<ListDirectoryHandle>(g_listDirHandleMemoryBlock);
diff --git a/Engine/Source/Programs/UnrealBuildAccelerator/Test/Private/UbaTestBasics.cpp b/Engine/Source/Programs/UnrealBuildAccelerator/Test/Private/UbaTestBasics.cpp
--- a/Engine/Source/Programs/UnrealBuildAccelerator/Test/Private/UbaTestBasics.cpp
+++ b/Engine/Source/Programs/UnrealBuildAccelerator/Test/Private/UbaTestBasics.cpp
@@ -395,7 +395,7 @@ namespace uba
 	bool TestMemoryBlock(LoggerWithWriter& logger, const StringBufferBase& rootDir)
 	{
 		{
-			MemoryBlock block(1024 * 1024);
+			MemoryBlock block(TC("TestMemoryBlock"), 1024 * 1024);
 			u64* mem = (u64*)block.Allocate(8, 1, TC("Foo"));
 			*mem = 0x1234;
 			block.Free(mem);
@@ -403,7 +403,7 @@ namespace uba
 
 		if (GetHugePageCount())
 		{
-			MemoryBlock block;
+			MemoryBlock block(TC("TestHugePage"));
 			if (!block.Init(1024 * 1024, nullptr, true))
 				return logger.Error(TC("Failed to allocate huge pages even though system says they exists"));
 			u64* mem = (u64*)block.Allocate(8, 1, TC("Foo"));
diff --git a/Engine/Source/Programs/UnrealBuildAccelerator/Test/Private/UbaTestCache.cpp b/Engine/Source/Programs/UnrealBuildAccelerator/Test/Private/UbaTestCache.cpp
--- a/Engine/Source/Programs/UnrealBuildAccelerator/Test/Private/UbaTestCache.cpp
+++ b/Engine/Source/Programs/UnrealBuildAccelerator/Test/Private/UbaTestCache.cpp
@@ -93,7 +93,7 @@ namespace uba
 	void UpdateCacheEntries(LoggerWithWriter& logger, CacheEntries& entries, u32 multiplier)
 	{
 		HashMap2<u32, u32> oldToNewCasKeyOffset;
-		MemoryBlock memory(64*1024);
+		MemoryBlock memory(TC("UpdateCacheEntries"), 64*1024);
 		oldToNewCasKeyOffset.Init(memory, 100);
 		if (multiplier)
 			for (u32 i=0; i!=200; ++i)
@@ -405,7 +405,7 @@ namespace uba
 	bool TestHashTable(LoggerWithWriter& logger, const StringBufferBase& rootDir)
 	{
 		{
-			MemoryBlock memoryBlock(1024*1024);
+			MemoryBlock memoryBlock(TC("TestHashTable"), 1024*1024);
 			HashMap<u32, u32> casMap;
 			casMap.Init(memoryBlock, 3);
 			if (casMap.Find(1))
diff --git a/Engine/Source/Programs/UnrealBuildAccelerator/Test/Private/UbaTestSession.cpp b/Engine/Source/Programs/UnrealBuildAccelerator/Test/Private/UbaTestSession.cpp
--- a/Engine/Source/Programs/UnrealBuildAccelerator/Test/Private/UbaTestSession.cpp
+++ b/Engine/Source/Programs/UnrealBuildAccelerator/Test/Private/UbaTestSession.cpp
@@ -657,7 +657,7 @@ namespace uba
 				return exitCode == 255 ? INVALID_FILE_ATTRIBUTES : exitCode;
 			};
 
-		MemoryBlock temp;
+		MemoryBlock temp(TC("CheckAttributes"));
 		DirectoryTable dirTable(temp);
 		dirTable.Init(session.GetDirectoryTableMemory(), 0, 0);
 
