[UBA]
* Fixed compile error in UbaObjTool

diff --git a/Engine/Source/Programs/UnrealBuildAccelerator/ObjTool/Private/UbaImportLibWriter.cpp b/Engine/Source/Programs/UnrealBuildAccelerator/ObjTool/Private/UbaImportLibWriter.cpp
--- a/Engine/Source/Programs/UnrealBuildAccelerator/ObjTool/Private/UbaImportLibWriter.cpp
+++ b/Engine/Source/Programs/UnrealBuildAccelerator/ObjTool/Private/UbaImportLibWriter.cpp
@@ -556,7 +556,7 @@ namespace uba
 
 	public:
 
-		ImportLib() : m_memory(128*1024*1024) // 512mb should be more than enough reserved memory for the biggest import libs
+		ImportLib() : m_memory(TC("ImportLib"), 128*1024*1024) // 512mb should be more than enough reserved memory for the biggest import libs
 		{
 		}
 
