[UBA]
* Workaround for bug where DetourCreateProcess seems to not like command lines where the first argument is both a "app.exe" and an argument directly after.. like /LIB. Workaround is to add a space in between
* Let list leak to prevent crashes during shutdown

diff --git a/Engine/Source/Programs/UnrealBuildAccelerator/Detours/Private/Windows/UbaDetoursFunctionsKernelBase.inl b/Engine/Source/Programs/UnrealBuildAccelerator/Detours/Private/Windows/UbaDetoursFunctionsKernelBase.inl
--- a/Engine/Source/Programs/UnrealBuildAccelerator/Detours/Private/Windows/UbaDetoursFunctionsKernelBase.inl
+++ b/Engine/Source/Programs/UnrealBuildAccelerator/Detours/Private/Windows/UbaDetoursFunctionsKernelBase.inl
@@ -2599,7 +2599,7 @@ HANDLE Detoured_OpenFileMappingW(DWORD dwDesiredAccess, BOOL bInheritHandle, LPC
 	return mappingHandle;
 }
 
-List<SharedMemoryView> g_specialViews;
+auto& g_specialViews = *new List<SharedMemoryView>();
 
 LPVOID Detoured_MapViewOfFileEx(HANDLE hFileMappingObject, DWORD dwDesiredAccess, DWORD dwFileOffsetHigh, DWORD dwFileOffsetLow, SIZE_T dwNumberOfBytesToMap, LPVOID lpBaseAddress)
 {
@@ -3209,6 +3209,9 @@ BOOL Detoured_CreateProcessW(LPCWSTR lpApplicationName, LPWSTR lpCommandLine, LP
 		isChild = false;
 	}
 
+	tchar* finalCommandLine = lpCommandLine;
+	auto g = MakeGuard([&]() { if (finalCommandLine != lpCommandLine) delete[] finalCommandLine; });
+
 	StringView commandLineWithoutApplication;
 	{
 		const tchar* endOfApplication;
@@ -3217,6 +3220,17 @@ BOOL Detoured_CreateProcessW(LPCWSTR lpApplicationName, LPWSTR lpCommandLine, LP
 			const tchar* quoteEnd = TStrchr(originalCmd.data+1, '"');
 			UBA_ASSERT(quoteEnd);
 			endOfApplication = quoteEnd + 1;
+
+			// This is an annoying special case that it seems like DetourCreateProcessWithDllsW does not like.. it does not separate the arguments properly
+			// so if there is "link.exe"/LIB as a command it will not produce proper wmain(argv)
+			if (*endOfApplication != ' ')
+			{
+				finalCommandLine = new tchar[originalCmd.count + 1];
+				u64 firstPart = endOfApplication - originalCmd.data;
+				memcpy(finalCommandLine, originalCmd.data, firstPart*sizeof(tchar));
+				finalCommandLine[firstPart] = ' ';
+				memcpy(finalCommandLine + firstPart + 1, endOfApplication, (originalCmd.count - firstPart)*sizeof(tchar));
+			}
 		}
 		else if (const tchar* firstSpace = TStrchr(originalCmd.data, ' '))
 			endOfApplication = firstSpace + 1;
@@ -3387,7 +3401,7 @@ BOOL Detoured_CreateProcessW(LPCWSTR lpApplicationName, LPWSTR lpCommandLine, LP
 	while (true)
 	{
 		res = true;
-		if (DetourCreateProcessWithDllsW(application.data, lpCommandLine, NULL, NULL, bInheritHandles, dwCreationFlags, lpEnvironment, currentDir.c_str(), lpStartupInfo, lpProcessInformation, sizeof_array(dlls), dlls, True_CreateProcessW))
+		if (DetourCreateProcessWithDllsW(application.data, finalCommandLine, NULL, NULL, bInheritHandles, dwCreationFlags, lpEnvironment, currentDir.c_str(), lpStartupInfo, lpProcessInformation, sizeof_array(dlls), dlls, True_CreateProcessW))
 			break;
 		res = false;
 		lastError = GetLastError();
