[UBA]
* Added workaround for bug in wine where string produced by GetModuleFileNameW when buffer is too small not being null terminated
* Fixed so UbaTestApp.exe can run outside of detouring just to be able to validate so tests are correct
* Bumped buffer for coordinator wrapper

diff --git a/Engine/Source/Programs/UnrealBuildAccelerator/Common/Public/UbaCoordinatorWrapper.h b/Engine/Source/Programs/UnrealBuildAccelerator/Common/Public/UbaCoordinatorWrapper.h
--- a/Engine/Source/Programs/UnrealBuildAccelerator/Common/Public/UbaCoordinatorWrapper.h
+++ b/Engine/Source/Programs/UnrealBuildAccelerator/Common/Public/UbaCoordinatorWrapper.h
@@ -30,7 +30,7 @@ namespace uba
 			if (!*coordinatorType)
 				return false;
 
-			StringBuffer<128> coordinatorBin(info.binariesDir);
+			StringBuffer<> coordinatorBin(info.binariesDir);
 			coordinatorBin.EnsureEndsWithSlash();
 
 			#if PLATFORM_WINDOWS
diff --git a/Engine/Source/Programs/UnrealBuildAccelerator/Detours/Private/Windows/UbaDetoursFunctionsKernelBase.inl b/Engine/Source/Programs/UnrealBuildAccelerator/Detours/Private/Windows/UbaDetoursFunctionsKernelBase.inl
--- a/Engine/Source/Programs/UnrealBuildAccelerator/Detours/Private/Windows/UbaDetoursFunctionsKernelBase.inl
+++ b/Engine/Source/Programs/UnrealBuildAccelerator/Detours/Private/Windows/UbaDetoursFunctionsKernelBase.inl
@@ -1351,8 +1351,7 @@ DWORD Shared_GetModuleFileNameInner(const tchar* func, HMODULE hModule, const St
 		return nSize;
 	}
 
-	memcpy(lpFilename, moduleName.data, moduleName.count * 2);
-	lpFilename[moduleName.count] = 0;
+	memcpy(lpFilename, moduleName.data, (moduleName.count+1) * 2);
 	DEBUG_LOG_DETOURED(func, L"%llu  %u (%ls) -> %u", uintptr_t(hModule), nSize, lpFilename, moduleName.count);
 	SetLastError(ERROR_SUCCESS);
 	return moduleName.count;
@@ -1366,6 +1365,8 @@ DWORD Shared_GetModuleFileNameW(const tchar* func, HMODULE hModule, LPWSTR lpFil
 		if (g_virtualApplication.EndsWith(L".bat"))
 		{
 			DWORD res = True_GetModuleFileNameW(hModule, lpFilename, nSize);
+			if (res == nSize && nSize && g_isRunningWine && GetLastError() == ERROR_INSUFFICIENT_BUFFER) // Wine bug, not terminating
+				lpFilename[nSize-1] = 0;
 			DEBUG_LOG_TRUE(L"GetModuleFileNameW", L"%llu  %u (%ls) -> %u", uintptr_t(hModule), nSize, lpFilename, res);
 			return res;
 		}
@@ -1383,6 +1384,8 @@ DWORD Shared_GetModuleFileNameW(const tchar* func, HMODULE hModule, LPWSTR lpFil
 	if (!g_runningRemote && !IsVfsEnabled())
 	{
 		DWORD res = True_GetModuleFileNameW(hModule, lpFilename, nSize);
+		if (res == nSize && nSize && g_isRunningWine && GetLastError() == ERROR_INSUFFICIENT_BUFFER) // Wine bug, not terminating
+			lpFilename[nSize-1] = 0;
 		DEBUG_LOG_TRUE(L"GetModuleFileNameW", L"%llu  %u (%ls) -> %u", uintptr_t(hModule), nSize, lpFilename, res);
 		return res;
 	}
@@ -1390,6 +1393,8 @@ DWORD Shared_GetModuleFileNameW(const tchar* func, HMODULE hModule, LPWSTR lpFil
 	StringBuffer<512> moduleName;
 	{
 		DWORD res = True_GetModuleFileNameW(hModule, moduleName.data, moduleName.capacity);
+		if (res == nSize && nSize && g_isRunningWine && GetLastError() == ERROR_INSUFFICIENT_BUFFER) // Wine bug, not terminating
+			lpFilename[nSize-1] = 0;
 		DEBUG_LOG_TRUE(L"GetModuleFileNameW", L"(INTERNAL) %llu  %u (%ls) -> %u", uintptr_t(hModule), moduleName.capacity, moduleName.data, res);
 		if (res == 0)
 			return res;
diff --git a/Engine/Source/Programs/UnrealBuildAccelerator/TestApp/Private/Windows/UbaTestAppWin.cpp b/Engine/Source/Programs/UnrealBuildAccelerator/TestApp/Private/Windows/UbaTestAppWin.cpp
--- a/Engine/Source/Programs/UnrealBuildAccelerator/TestApp/Private/Windows/UbaTestAppWin.cpp
+++ b/Engine/Source/Programs/UnrealBuildAccelerator/TestApp/Private/Windows/UbaTestAppWin.cpp
@@ -50,14 +50,16 @@ int wmain(int argc, wchar_t* argv[])
 
 	if (argc == 1)
 	{
-		if (!detoursHandle)
-			return LogError(L"Did not find UbaDetours.dll in process!!!\n");
+		bool runningRemote = false;
 
-		using UbaRunningRemoteFunc = bool();
-		UbaRunningRemoteFunc* runningRemoteFunc = (UbaRunningRemoteFunc*)GetProcAddress(detoursHandle, "UbaRunningRemote");
-		if (!runningRemoteFunc)
-			return LogError(L"Couldn't find UbaRunningRemote function in UbaDetours.dll");
-		bool runningRemote = (*runningRemoteFunc)();
+		if (detoursHandle)
+		{
+			using UbaRunningRemoteFunc = bool();
+			UbaRunningRemoteFunc* runningRemoteFunc = (UbaRunningRemoteFunc*)GetProcAddress(detoursHandle, "UbaRunningRemote");
+			if (!runningRemoteFunc)
+				return LogError(L"Couldn't find UbaRunningRemote function in UbaDetours.dll");
+			runningRemote = (*runningRemoteFunc)();
+		}
 
 		HMODULE modules[] = { 0, detoursHandle, GetModuleHandleW(L"UbaTestApp.exe") };
 		for (HMODULE module : modules)
@@ -68,11 +70,13 @@ int wmain(int argc, wchar_t* argv[])
 			if (res1 != 0)
 				return LogError(L"Expected zero");
 			wchar_t name[512];
+			memset(name, 123, sizeof(name));
 			DWORD realLen = GetModuleFileNameW(module, name, 512);
 			if (realLen == 0)
 				return LogError(L"Did not expect this function to fail");
 			if (GetLastError() == ERROR_INSUFFICIENT_BUFFER)
 				return LogError(L"Expected sufficient buffer");
+			memset(name, 123, sizeof(name));
 			name[realLen] = 254;
 			name[realLen+1] = 254;
 			DWORD res2 = GetModuleFileNameW(module, name, realLen);
@@ -83,7 +87,8 @@ int wmain(int argc, wchar_t* argv[])
 			if (name[realLen] != 254)
 				return LogError(L"Overwrite");
 			if (name[realLen-1] != 0)
-				return LogError(L"Not terminated");
+				return LogError(L"Insufficient buffer not terminated: %s", name);
+			memset(name, 123, sizeof(name));
 			name[realLen] = 254;
 			name[realLen + 1] = 254;
 			DWORD res3 = GetModuleFileNameW(module, name, realLen+1);
@@ -94,7 +99,7 @@ int wmain(int argc, wchar_t* argv[])
 			if (name[realLen+1] != 254)
 				return LogError(L"Overwrite");
 			if (name[realLen] != 0)
-				return LogError(L"Not terminated");
+				return LogError(L"Not terminated: %s", name);
 		}
 
 		wchar_t currentDir[MAX_PATH];
@@ -170,7 +175,7 @@ int wmain(int argc, wchar_t* argv[])
 				return LogError(L"Should not succeed creating directory that exists");
 
 			if (GetFileAttributesW(L"DirA") == 0)
-				return LogError(L"Failed to get attributes of directory");
+				return LogError(L"Failed to get attributes of directory DirA");
 
 			if (!RemoveDirectoryW(L"DirA"))
 				return LogError(L"Failed to remove directory");
@@ -180,44 +185,49 @@ int wmain(int argc, wchar_t* argv[])
 
 			if (CreateDirectoryW(L"Dir2\\Dir3", NULL))
 				return LogError(L"Should not succeed creating directory inside directory that does not exist");
-			if (GetLastError() != ERROR_ALREADY_EXISTS)
-				return LogError(L"Did not get correct error when failing to create existing directory");
-			if (GetFileAttributesW(L"Dir2\\Dir3\\Dir4\\Dir5") == INVALID_FILE_ATTRIBUTES)
-				return LogError(L"Failed to get attributes of directory");
+			if (detoursHandle)
+			{
+				if (GetLastError() != ERROR_ALREADY_EXISTS)
+					return LogError(L"Did not get correct error when failing to create existing directory: %u", GetLastError());
+				if (GetFileAttributesW(L"Dir2\\Dir3\\Dir4\\Dir5") == INVALID_FILE_ATTRIBUTES)
+					return LogError(L"Failed to get attributes of directory Dir2\\Dir3\\Dir4\\Dir5");
+			}
 		}
 
 
+		if (detoursHandle)
 		{
-			HANDLE fh = CreateFileW(L"File4.out", GENERIC_WRITE, FILE_SHARE_READ, NULL, OPEN_EXISTING, 0, NULL);
-			if (fh == INVALID_HANDLE_VALUE)
-				return LogError(L"Failed to open File4.out for read with write permissions");
-			DWORD bytesRead;
-			wchar_t data[1];
-			if (!ReadFile(fh, data, sizeof(data), &bytesRead, NULL) || bytesRead != sizeof(data) || data[0] != '0')
-				return LogError(L"Failed to read one byte from File4.out");
-			CloseHandle(fh);
-		}
-
-		auto createForRead = [](const wchar_t* fileName, wchar_t expectedData)
 			{
-				HANDLE fh = CreateFileW(fileName, GENERIC_READ, FILE_SHARE_READ, NULL, OPEN_EXISTING, 0, NULL);
+				HANDLE fh = CreateFileW(L"File4.out", GENERIC_WRITE, FILE_SHARE_READ, NULL, OPEN_EXISTING, 0, NULL);
 				if (fh == INVALID_HANDLE_VALUE)
-					return LogError(L"Failed to open File5.h for read with write permissions");
+					return LogError(L"Failed to open File4.out for read with write permissions");
 				DWORD bytesRead;
 				wchar_t data[1];
-				if (!ReadFile(fh, data, sizeof(data), &bytesRead, NULL) || bytesRead != sizeof(data) || data[0] != expectedData)
-					return LogError(L"Failed to read one byte from %s", fileName);
+				if (!ReadFile(fh, data, sizeof(data), &bytesRead, NULL) || bytesRead != sizeof(data) || data[0] != '0')
+					return LogError(L"Failed to read one byte from File4.out");
 				CloseHandle(fh);
-				return 0;
-			};
-
-		if (createForRead(L"File4.obj", '9'))
-			return -1;
-		if (createForRead(L"File4.obj", '9'))
-			return -1;
-		if (createForRead(L"File4.lib", '9'))
-			return -1;
+			}
 
+			auto createForRead = [](const wchar_t* fileName, wchar_t expectedData)
+				{
+					HANDLE fh = CreateFileW(fileName, GENERIC_READ, FILE_SHARE_READ, NULL, OPEN_EXISTING, 0, NULL);
+					if (fh == INVALID_HANDLE_VALUE)
+						return LogError(L"Failed to open %s for read with write permissions", fileName);
+					DWORD bytesRead;
+					wchar_t data[1];
+					if (!ReadFile(fh, data, sizeof(data), &bytesRead, NULL) || bytesRead != sizeof(data) || data[0] != expectedData)
+						return LogError(L"Failed to read one byte from %s", fileName);
+					CloseHandle(fh);
+					return 0;
+				};
+
+			if (createForRead(L"File4.obj", '9'))
+				return -1;
+			if (createForRead(L"File4.obj", '9'))
+				return -1;
+			if (createForRead(L"File4.lib", '9'))
+				return -1;
+		}
 		
 		if (WriteTestFile(L"Foo", 8, L"MemoryFile.out"))
 			return -1;
@@ -255,6 +265,9 @@ int wmain(int argc, wchar_t* argv[])
 	}
 	else if (wcscmp(argv[1], L"-reuse") == 0)
 	{
+		if (detoursHandle)
+			return LogError(L"Did not find UbaDetours.dll in process!!!\n");
+
 		wchar_t arguments[1024];
 		if (requestNextProcess(0, arguments, sizeof(arguments)))
 			return LogError(L"Didn't expect another process");
