[UBA]
* Fixed so uba creates a cmd.exe around command lines containing shell arguments like <>|&^||&& etc

diff --git a/Engine/Source/Programs/UnrealBuildAccelerator/Common/Private/UbaProcessStartInfoHolder.cpp b/Engine/Source/Programs/UnrealBuildAccelerator/Common/Private/UbaProcessStartInfoHolder.cpp
--- a/Engine/Source/Programs/UnrealBuildAccelerator/Common/Private/UbaProcessStartInfoHolder.cpp
+++ b/Engine/Source/Programs/UnrealBuildAccelerator/Common/Private/UbaProcessStartInfoHolder.cpp
@@ -2,6 +2,7 @@
 
 #include "UbaProcessStartInfoHolder.h"
 #include "UbaConfig.h"
+#include "UbaProcessUtils.h"
 
 namespace uba
 {
@@ -10,7 +11,24 @@ namespace uba
 	#if PLATFORM_WINDOWS
 		// Special handling to avoid calling cmd.exe if not needed
 		if (!Contains(application, TC("cmd.exe")))
+		{
+			// We also need to check if we need to add cmd.exe around the arguments
+			bool requiresShell = false;
+			ParseArguments(argumentsStr.c_str(), argumentsStr.size(), [&](const tchar* arg, u32 argLen)
+				{
+					StringView view(arg, argLen);
+					requiresShell |= argLen == 1 && (*arg == '<' || *arg == '>' || *arg == '|' || *arg == '^');
+					requiresShell |= argLen == 2 && (view.Equals(TC(">>")) || view.Equals(TC("&&")) || view.Equals(TC("||")));
+				});
+			if (requiresShell)
+			{
+				argumentsStr = TC("/C \"") + applicationStr + TC(" ") + argumentsStr + TC("\"");
+				arguments = argumentsStr.c_str();
+				applicationStr = TC("cmd.exe");
+				application = applicationStr.c_str();
+			}
 			return false;
+		}
 
 		const tchar* argsBegin = argumentsStr.c_str();
 
