[UBA]
* Fixed visualizer bug where boxes were not properly dirtied when horizontal zoomed

diff --git a/Engine/Source/Programs/UnrealBuildAccelerator/Visualizer/Private/Windows/UbaVisualizerWin.cpp b/Engine/Source/Programs/UnrealBuildAccelerator/Visualizer/Private/Windows/UbaVisualizerWin.cpp
--- a/Engine/Source/Programs/UnrealBuildAccelerator/Visualizer/Private/Windows/UbaVisualizerWin.cpp
+++ b/Engine/Source/Programs/UnrealBuildAccelerator/Visualizer/Private/Windows/UbaVisualizerWin.cpp
@@ -805,9 +805,12 @@ namespace uba
 					m_scrollPosY = Min(0.0f, float(m_scrollPosY - scrollAnchorOffsetY)*(m_zoomValue/oldZoomValue) + scrollAnchorOffsetY);
 					m_scrollPosX = Min(0.0f, float(m_scrollPosX - scrollAnchorOffsetX)*(m_zoomValue/oldZoomValue)*(newScaleValue/horizontalScaleValue) + scrollAnchorOffsetX);//LOWORD(lParam);
 
+					bool wider = false;
 					if (horizontalScaleValue != newScaleValue)
+					{
+						wider = horizontalScaleValue < newScaleValue;
 						m_config.horizontalScaleValue = newScaleValue;
-
+					}
 
 					UpdateAutoscroll();
 					UpdateSelection();
@@ -820,7 +823,7 @@ namespace uba
 					//if (!m_traceView.finished && m_scrollPosX <= minScroll)
 					//	m_autoScroll = true;
 
-					if (oldBoxHeight != newBoxHeight)
+					if (oldBoxHeight != newBoxHeight || wider)
 						DirtyBitmaps(true);
 					else if (m_config.ShowReadWriteColors)
 						for (auto& session : m_traceView.sessions)
