[UBA]
* Added more info to speedtest start message (how to change connection count)

diff --git a/Engine/Source/Programs/UnrealBuildAccelerator/Common/Private/UbaSessionClient.cpp b/Engine/Source/Programs/UnrealBuildAccelerator/Common/Private/UbaSessionClient.cpp
--- a/Engine/Source/Programs/UnrealBuildAccelerator/Common/Private/UbaSessionClient.cpp
+++ b/Engine/Source/Programs/UnrealBuildAccelerator/Common/Private/UbaSessionClient.cpp
@@ -2495,7 +2495,7 @@ namespace uba
 					Parse(connectionCount, arg+3, argLen-3);
 			});
 
-		logger.Info(TC("Starting speed test.. Waiting for %llu connections"), connectionCount);
+		logger.Info(TC("Starting speed test.. Waiting for %llu connections (use -c=<connectioncount>)"), connectionCount);
 
 		m_desiredConnectionCount = u32(connectionCount);
 		while (m_client.IsConnected() && m_client.GetConnectionCount() != connectionCount)
