[UBA]
* Added detouring of AssignProcessToJobObject and RegisterWaitForSingleObject (needed because process handles are detoured)
* Fixed so filemaps can be written to (or the file map is not written to but the file map is replaced with a normal file)

diff --git a/Engine/Source/Programs/UnrealBuildAccelerator/Detours/Private/Windows/UbaDetoursFunctionsKernelBase.inl b/Engine/Source/Programs/UnrealBuildAccelerator/Detours/Private/Windows/UbaDetoursFunctionsKernelBase.inl
--- a/Engine/Source/Programs/UnrealBuildAccelerator/Detours/Private/Windows/UbaDetoursFunctionsKernelBase.inl
+++ b/Engine/Source/Programs/UnrealBuildAccelerator/Detours/Private/Windows/UbaDetoursFunctionsKernelBase.inl
@@ -3234,7 +3234,7 @@ BOOL Detoured_CreateProcessW(LPCWSTR lpApplicationName, LPWSTR lpCommandLine, LP
 				u64 firstPart = endOfApplication - originalCmd.data;
 				memcpy(finalCommandLine, originalCmd.data, firstPart*sizeof(tchar));
 				finalCommandLine[firstPart] = ' ';
-				memcpy(finalCommandLine + firstPart + 1, endOfApplication, (originalCmd.count - firstPart)*sizeof(tchar));
+				memcpy(finalCommandLine + firstPart + 1, endOfApplication, (originalCmd.count - firstPart + 1)*sizeof(tchar));
 			}
 		}
 		else if (const tchar* firstSpace = TStrchr(originalCmd.data, ' '))
@@ -3534,6 +3534,30 @@ BOOL Detoured_TerminateProcess(HANDLE hProcess, UINT uExitCode)
 	return True_TerminateProcess(hProcess, uExitCode);
 }
 
+BOOL Detoured_AssignProcessToJobObject(HANDLE hJob, HANDLE hProcess)
+{
+	{
+		//SuppressDetourScope _;
+		//MessageBoxW(GetConsoleWindow(), L"DFWQFWEFE", TC("Assert"), MB_ABORTRETRYIGNORE);
+	}
+
+	HANDLE trueHandle = hProcess;
+	if (isDetouredHandle(hProcess))
+		trueHandle = asDetouredHandle(hProcess).trueHandle;
+	BOOL res = True_AssignProcessToJobObject(hJob, trueHandle);
+	DEBUG_LOG_DETOURED(L"AssignProcessToJobObject", L"%llu -> %ls", uintptr_t(trueHandle), ToString(res));
+	return res;
+}
+
+BOOL Detoured_RegisterWaitForSingleObject(PHANDLE phNewWaitObject, HANDLE hObject, WAITORTIMERCALLBACK Callback, PVOID Context, ULONG dwMilliseconds, ULONG dwFlags)
+{
+	HANDLE trueHandle = hObject;
+	if (isDetouredHandle(hObject))
+		trueHandle = asDetouredHandle(hObject).trueHandle;
+
+	return True_RegisterWaitForSingleObject(phNewWaitObject, trueHandle, Callback, Context, dwMilliseconds, dwFlags);
+}
+
 BOOL Detoured_GetExitCodeProcess(HANDLE hProcess, LPDWORD lpExitCode)
 {
 	DETOURED_CALL(GetExitCodeProcess);
diff --git a/Engine/Source/Programs/UnrealBuildAccelerator/Detours/Private/Windows/UbaDetoursFunctionsNtDll.inl b/Engine/Source/Programs/UnrealBuildAccelerator/Detours/Private/Windows/UbaDetoursFunctionsNtDll.inl
--- a/Engine/Source/Programs/UnrealBuildAccelerator/Detours/Private/Windows/UbaDetoursFunctionsNtDll.inl
+++ b/Engine/Source/Programs/UnrealBuildAccelerator/Detours/Private/Windows/UbaDetoursFunctionsNtDll.inl
@@ -1112,9 +1112,17 @@ NTSTATUS NTAPI Shared_NtCreateFile(bool IsCreateFunc, PHANDLE hFileHandle, ACCES
 		if (isWrite)
 		{
 			bool lastWasWrite = IsContentWrite(info.lastDesiredAccess, 0);
-			UBA_ASSERT(!info.isFileMap);
 			bool shouldReport = !lastWasWrite || info.deleted || isDeleteOnClose;
 			shouldReport = shouldReport && !keepInMemory;
+
+			if (info.isFileMap)
+			{
+				info.fileMapMem = nullptr;
+				info.fileMapMemSize = 0;
+				info.isFileMap = false;
+				shouldReport = true;
+			}
+
 			if (shouldReport)
 			{
 				u64 size = InvalidValue;
diff --git a/Engine/Source/Programs/UnrealBuildAccelerator/Detours/Private/Windows/UbaDetoursFunctionsWin.h b/Engine/Source/Programs/UnrealBuildAccelerator/Detours/Private/Windows/UbaDetoursFunctionsWin.h
--- a/Engine/Source/Programs/UnrealBuildAccelerator/Detours/Private/Windows/UbaDetoursFunctionsWin.h
+++ b/Engine/Source/Programs/UnrealBuildAccelerator/Detours/Private/Windows/UbaDetoursFunctionsWin.h
@@ -128,6 +128,8 @@
 	DETOURED_FUNCTIONS_KERNELBASE_DEBUG \
 
 #define DETOURED_FUNCTIONS_KERNEL32 \
+	DETOURED_FUNCTION(AssignProcessToJobObject) \
+	DETOURED_FUNCTION(RegisterWaitForSingleObject) \
 	DETOURED_FUNCTION(CreateFileMappingA) \
 	DETOURED_FUNCTION(GetExitCodeProcess) \
 	DETOURED_FUNCTION(CreateTimerQueueTimer) \
diff --git a/Engine/Source/Programs/UnrealBuildAccelerator/Test/Private/UbaTestAll.h b/Engine/Source/Programs/UnrealBuildAccelerator/Test/Private/UbaTestAll.h
--- a/Engine/Source/Programs/UnrealBuildAccelerator/Test/Private/UbaTestAll.h
+++ b/Engine/Source/Programs/UnrealBuildAccelerator/Test/Private/UbaTestAll.h
@@ -88,6 +88,8 @@ namespace uba
 		UBA_TEST(TestRemoteRegisterVirtualFile) \
 		UBA_TEST(TestRemoteProcessSpecialCase1) \
 		UBA_TEST(TestSessionSpecialCopy) \
+		UBA_TEST(TestFileMapWrite) \
+		UBA_TEST(TestFileMapWriteRemote) \
 
 
 #define UBA_LINUX_TESTS \
diff --git a/Engine/Source/Programs/UnrealBuildAccelerator/Test/Private/UbaTestSession.cpp b/Engine/Source/Programs/UnrealBuildAccelerator/Test/Private/UbaTestSession.cpp
--- a/Engine/Source/Programs/UnrealBuildAccelerator/Test/Private/UbaTestSession.cpp
+++ b/Engine/Source/Programs/UnrealBuildAccelerator/Test/Private/UbaTestSession.cpp
@@ -1001,4 +1001,24 @@ namespace uba
 				return true;
 			});
 	}
+
+	bool WriteFileMap(LoggerWithWriter& logger, SessionServer& session, const tchar* workingDir, const RunProcessFunction& runProcess, void* extraData)
+	{
+		StringBuffer<MaxPath> fileR;
+		if (!CreateTextFile(fileR, logger, workingDir, TC("File.h"), "Foo"))
+			return false;
+
+		//auto runProcess2 = [&](ProcessStartInfo& pi) { return session.RunProcess(pi, true, true); };
+		return RunTestApp(logger, workingDir, runProcess, TC("-writeFileMap"));
+	}
+
+	bool TestFileMapWrite(LoggerWithWriter& logger, const StringBufferBase& testRootDir)
+	{
+		return RunLocal(logger, testRootDir, WriteFileMap);
+	}
+
+	bool TestFileMapWriteRemote(LoggerWithWriter& logger, const StringBufferBase& testRootDir)
+	{
+		return RunRemote(logger, testRootDir, WriteFileMap);
+	}
 }
diff --git a/Engine/Source/Programs/UnrealBuildAccelerator/TestApp/Private/Windows/UbaTestAppWin.cpp b/Engine/Source/Programs/UnrealBuildAccelerator/TestApp/Private/Windows/UbaTestAppWin.cpp
--- a/Engine/Source/Programs/UnrealBuildAccelerator/TestApp/Private/Windows/UbaTestAppWin.cpp
+++ b/Engine/Source/Programs/UnrealBuildAccelerator/TestApp/Private/Windows/UbaTestAppWin.cpp
@@ -448,6 +448,20 @@ int wmain(int argc, wchar_t* argv[])
 		if (memcmp(data, L"Bar", 6) != 0)
 			return -1;
 	}
+	else if (wcscmp(argv[1], L"-writeFileMap") == 0)
+	{
+		wchar_t buf[32];
+		if (ReadTestFile(buf, 1, L"File.h") == -1)
+			return -1;
+		buf[0] = 'M';
+		if (WriteTestFile(buf, 1, L"File.h") == -1)
+			return -1;
+		buf[0] = 0;
+		if (ReadTestFile(buf, 1, L"File.h") == -1)
+			return -1;
+		if (buf[0] != 'M')
+			return -1;
+	}
 	else
 	{
 		if (!detoursHandle)
