UBA: Fixed the p4 reconcile in Deploy-Windows.bat

#rnx
#Synced-CL 50967057

diff --git a/Engine/Source/Programs/UnrealBuildAccelerator/Deploy-Windows.bat b/Engine/Source/Programs/UnrealBuildAccelerator/Deploy-Windows.bat
--- a/Engine/Source/Programs/UnrealBuildAccelerator/Deploy-Windows.bat
+++ b/Engine/Source/Programs/UnrealBuildAccelerator/Deploy-Windows.bat
@@ -1,8 +1,6 @@
 @echo off
 setlocal 
 
-pushd "%~dp0"
-
 set config=Shipping
 
 if "%~1"=="" goto blank
@@ -13,7 +11,7 @@ set config=%1
 
 echo.
 echo === Building Targets ===
-pushd "../../../.."
+pushd "%~dp0../../../.."
 call Engine/Build/BatchFiles/RunUBT.bat -NoUba -NoUbaLocal -NoSNDBS -NoXGE -novfs ^
 	-Target="UbaAgent Win64 %config%" ^
 	-Target="UbaCli Win64 %config%" ^
@@ -29,8 +27,8 @@ popd
 
 echo.
 echo === Reconciling artifacts ===
+pushd "%~dp0"
 p4 reconcile ../../../Binaries/Win64/UnrealBuildAccelerator/...
-
 popd
 
 endlocal
