[UBA]
* Fixed so child processes can have relative working dir
* Changed so scheduler finished callback is called after stats are updated
* Fixed bug in storage client where HasCasFile can return a casentry that is not correct
* Bumped temp dir path max to 256 characters
* Improved asserts

diff --git a/Engine/Source/Programs/UnrealBuildAccelerator/Common/Private/UbaLogger.cpp b/Engine/Source/Programs/UnrealBuildAccelerator/Common/Private/UbaLogger.cpp
--- a/Engine/Source/Programs/UnrealBuildAccelerator/Common/Private/UbaLogger.cpp
+++ b/Engine/Source/Programs/UnrealBuildAccelerator/Common/Private/UbaLogger.cpp
@@ -234,15 +234,14 @@ namespace uba
 		m_writer.Log(type, str, strLen, prefix, prefixLen);
 	}
 
-	class ConsoleLogWriter : public LogWriter
+	class ConsoleLogWriter final : public LogWriter
 	{
 	public:
 		ConsoleLogWriter();
 		virtual void BeginScope() override;
 		virtual void EndScope() override;
 		virtual void Log(LogEntryType type, const tchar* str, u32 strLen, const tchar* prefix = nullptr, u32 prefixLen = 0) override;
-	private:
-		void LogNoLock(LogEntryType type, const tchar* str, u32 strLen, const tchar* prefix, u32 prefixLen);
+		void LogNoLock(LogEntryType type, const tchar* str, u32 strLen, const tchar* prefix, u32 prefixLen, bool addEndOfLine);
 		Futex m_lock;
 #if PLATFORM_WINDOWS
 		HANDLE m_stdout = 0;
@@ -295,16 +294,16 @@ namespace uba
 	void ConsoleLogWriter::Log(LogEntryType type, const tchar* str, u32 strLen, const tchar* prefix, u32 prefixLen)
 	{
 		if (t_consoleLogScopeCount)
-			return LogNoLock(type, str, strLen, prefix, prefixLen);
+			return LogNoLock(type, str, strLen, prefix, prefixLen, true);
 		SCOPED_FUTEX(m_lock, lock);
-		LogNoLock(type, str, strLen, prefix, prefixLen);
+		LogNoLock(type, str, strLen, prefix, prefixLen, true);
 #if PLATFORM_WINDOWS
 		if (!m_stdout)
 #endif
 			fflush(stdout);
 	}
 
-	void ConsoleLogWriter::LogNoLock(LogEntryType type, const tchar* str, u32 strLen, const tchar* prefix, u32 prefixLen)
+	void ConsoleLogWriter::LogNoLock(LogEntryType type, const tchar* str, u32 strLen, const tchar* prefix, u32 prefixLen, bool addEndOfLine)
 	{
 #if PLATFORM_WINDOWS
 		if (!m_stdout)
@@ -339,7 +338,8 @@ namespace uba
 				WriteConsoleW(m_stdout, str, strLen, NULL, NULL);
 				break;
 			}
-			WriteConsoleW(m_stdout, TC("\r\n"), 2, NULL, NULL);
+			if (addEndOfLine)
+				WriteConsoleW(m_stdout, TC("\r\n"), 2, NULL, NULL);
 		}
 #else
 		if (prefixLen)
@@ -348,10 +348,17 @@ namespace uba
 			TFputs(TC(" - "), stdout);
 		}
 		TFputs(str, stdout);
-		TFputs(TC("\n"), stdout);
+		if (addEndOfLine)
+			TFputs(TC("\n"), stdout);
 #endif
 	}
 
+	void LogToConsoleNoLineFeed(LogEntryType type, const tchar* str, u32 strLen)
+	{
+		g_consoleLogWriter.BeginScope();
+		g_consoleLogWriter.EndScope();
+	}
+
 	LastErrorToText::LastErrorToText(u32 lastError)
 	{
 #if PLATFORM_WINDOWS
diff --git a/Engine/Source/Programs/UnrealBuildAccelerator/Common/Private/UbaProcess.cpp b/Engine/Source/Programs/UnrealBuildAccelerator/Common/Private/UbaProcess.cpp
--- a/Engine/Source/Programs/UnrealBuildAccelerator/Common/Private/UbaProcess.cpp
+++ b/Engine/Source/Programs/UnrealBuildAccelerator/Common/Private/UbaProcess.cpp
@@ -1326,6 +1326,11 @@ namespace uba
 		reader.ReadString(currentDir);
 		if (currentDir.IsEmpty())
 			currentDir.Append(m_startInfo.workingDir);
+		else if (!IsAbsolutePath(currentDir.data))
+		{
+			TString temp(currentDir.ToString());
+			FixPath(temp.c_str(), m_startInfo.workingDir, m_startInfo.workingDirStr.size(), currentDir.Clear());
+		}
 		bool startSuspended = reader.ReadBool();
 		bool isChild = reader.ReadBool();
 
diff --git a/Engine/Source/Programs/UnrealBuildAccelerator/Common/Private/UbaStorage.cpp b/Engine/Source/Programs/UnrealBuildAccelerator/Common/Private/UbaStorage.cpp
--- a/Engine/Source/Programs/UnrealBuildAccelerator/Common/Private/UbaStorage.cpp
+++ b/Engine/Source/Programs/UnrealBuildAccelerator/Common/Private/UbaStorage.cpp
@@ -2495,7 +2495,7 @@ namespace uba
 
 			CasEntry& casEntry = *casEntryPtr;
 			CasEntryReadLock casEntryLock(casEntry);
-			UBA_ASSERT(casEntry.verified);
+			UBA_ASSERTF(casEntry.verified, TC("CasEntry %s not verified for copy to %s"), CasKeyString(actualKey).str, destination);
 			UBA_ASSERT(casEntry.exists);
 			--casEntry.refCount; // Decrease after HasCasFile increase it
 
diff --git a/Engine/Source/Programs/UnrealBuildAccelerator/Common/Private/UbaStorageClient.cpp b/Engine/Source/Programs/UnrealBuildAccelerator/Common/Private/UbaStorageClient.cpp
--- a/Engine/Source/Programs/UnrealBuildAccelerator/Common/Private/UbaStorageClient.cpp
+++ b/Engine/Source/Programs/UnrealBuildAccelerator/Common/Private/UbaStorageClient.cpp
@@ -903,9 +903,16 @@ namespace uba
 		auto findIt = m_localStorageFiles.find(localKey);
 		if (findIt != m_localStorageFiles.end())
 		{
-			if (out)
-				*out = &findIt->second.casEntry;
-			return true;
+			LocalFile& localFile = findIt->second;
+			if (localFile.casEntry.exists && localFile.casEntry.verified)
+			{
+				if (out)
+				{
+					++localFile.casEntry.refCount; // For good measure
+					*out = &localFile.casEntry;
+				}
+				return true;
+			}
 		}
 		lock.Leave();
 		return StorageImpl::HasCasFile(casKey, out);
@@ -955,7 +962,9 @@ namespace uba
 		{
 			localFile.casEntry.memoryHandle = m_sharedMemory.DuplicateHandle(memoryHandle, TC(""));
 			localFile.casEntry.size = fileSize;
+			localFile.casEntry.verified = true;
 			localFile.casEntry.exists = true;
+			localFile.casEntry.refCount = 1;
 		}
 
 		if (!isPersistentMapping && !localFile.casEntry.memoryHandle.IsValid() && localFile.fileName.empty())
@@ -963,6 +972,7 @@ namespace uba
 			localFile.casEntry.size = fileSize;
 			localFile.casEntry.verified = true;
 			localFile.casEntry.exists = true;
+			localFile.casEntry.refCount = 1;
 			localFile.fileName = fileName;
 		}
 
@@ -1103,6 +1113,7 @@ namespace uba
 						localFile.casEntry.size = e.size;
 						localFile.casEntry.verified = true;
 						localFile.casEntry.exists = true;
+						localFile.casEntry.refCount = 1;
 						localFile.fileName = fullPath.data;
 					}
 					return;
@@ -1135,6 +1146,7 @@ namespace uba
 							localFile.casEntry.size = s;
 							localFile.casEntry.verified = true;
 							localFile.casEntry.exists = true;
+							localFile.casEntry.refCount = 1;
 							localFile.fileName = filePath;
 						}
 
diff --git a/Engine/Source/Programs/UnrealBuildAccelerator/Common/Public/UbaLogger.h b/Engine/Source/Programs/UnrealBuildAccelerator/Common/Public/UbaLogger.h
--- a/Engine/Source/Programs/UnrealBuildAccelerator/Common/Public/UbaLogger.h
+++ b/Engine/Source/Programs/UnrealBuildAccelerator/Common/Public/UbaLogger.h
@@ -161,5 +161,7 @@ namespace uba
 	#define LOG_STALL_SCOPE(logger, timeSeconds, messageFormat)
 	#endif
 
+	void LogToConsoleNoLineFeed(LogEntryType type, const tchar* str, u32 strLen);
+
 	////////////////////////////////////////////////////////////////////////////////////////////////////
 }
diff --git a/Engine/Source/Programs/UnrealBuildAccelerator/Detours/Private/UbaDetoursShared.cpp b/Engine/Source/Programs/UnrealBuildAccelerator/Detours/Private/UbaDetoursShared.cpp
--- a/Engine/Source/Programs/UnrealBuildAccelerator/Detours/Private/UbaDetoursShared.cpp
+++ b/Engine/Source/Programs/UnrealBuildAccelerator/Detours/Private/UbaDetoursShared.cpp
@@ -20,7 +20,7 @@ namespace uba
 	VARIABLE_MEM(StringBuffer<512>, g_virtualWorkingDir);
 	VARIABLE_MEM(StringBuffer<256>, g_exeDir);
 	VARIABLE_MEM(StringBuffer<128>, g_systemRoot);
-	VARIABLE_MEM(StringBuffer<128>, g_systemTemp);
+	VARIABLE_MEM(StringBuffer<256>, g_systemTemp);
 	VARIABLE_MEM(MemoryBlock, g_memoryBlock);
 	VARIABLE_MEM(DirectoryTable, g_directoryTable);
 	VARIABLE_MEM(MappedFileTable, g_mappedFileTable);
diff --git a/Engine/Source/Programs/UnrealBuildAccelerator/Detours/Private/UbaDetoursShared.h b/Engine/Source/Programs/UnrealBuildAccelerator/Detours/Private/UbaDetoursShared.h
--- a/Engine/Source/Programs/UnrealBuildAccelerator/Detours/Private/UbaDetoursShared.h
+++ b/Engine/Source/Programs/UnrealBuildAccelerator/Detours/Private/UbaDetoursShared.h
@@ -58,7 +58,7 @@ namespace uba
 	extern StringBuffer<256>& g_logName;
 
 	extern StringBuffer<512>& g_virtualWorkingDir;
-	extern StringBuffer<128>& g_systemTemp;
+	extern StringBuffer<256>& g_systemTemp;
 	extern StringBuffer<128>& g_systemRoot;
 
 	extern ProcessStats& g_stats;
