[UBA]
* Fixed LogToConsoleNoLineFeed not writing anything

diff --git a/Engine/Source/Programs/UnrealBuildAccelerator/Common/Private/UbaLogger.cpp b/Engine/Source/Programs/UnrealBuildAccelerator/Common/Private/UbaLogger.cpp
--- a/Engine/Source/Programs/UnrealBuildAccelerator/Common/Private/UbaLogger.cpp
+++ b/Engine/Source/Programs/UnrealBuildAccelerator/Common/Private/UbaLogger.cpp
@@ -251,7 +251,7 @@ namespace uba
 	UBA_API LogWriter& g_consoleLogWriter = *new ConsoleLogWriter(); // Leak to prevent asan annoyances during shutdown when asserts happen
 	thread_local u32 t_consoleLogScopeCount = 0;
 
-	class NullLogWriter : public LogWriter
+	class NullLogWriter final : public LogWriter
 	{
 	public:
 		virtual void BeginScope() override {}
@@ -355,8 +355,10 @@ namespace uba
 
 	void LogToConsoleNoLineFeed(LogEntryType type, const tchar* str, u32 strLen)
 	{
-		g_consoleLogWriter.BeginScope();
-		g_consoleLogWriter.EndScope();
+		auto& logWriter = (ConsoleLogWriter&)g_consoleLogWriter;
+		logWriter.BeginScope();
+		logWriter.LogNoLock(type, str, strLen, nullptr, 0, false);
+		logWriter.EndScope();
 	}
 
 	LastErrorToText::LastErrorToText(u32 lastError)
@@ -505,7 +507,7 @@ namespace uba
 #if UBA_DEBUG_LOGGER || UBA_DEBUG
 	thread_local u32 t_debugLogScopeCount = 0;
 
-	class DebugLogWriter : public LogWriter
+	class DebugLogWriter final : public LogWriter
 	{
 	public:
 		virtual void BeginScope() override
