[UBA]
* Added UE ParkingLot sync primitives to replace current ones. Disabled by default.

IMPLEMENTED BY [REDACTED_PERSON], thank you!

diff --git a/Engine/Source/Programs/UnrealBuildAccelerator/Common/Private/UbaLogger.cpp b/Engine/Source/Programs/UnrealBuildAccelerator/Common/Private/UbaLogger.cpp
--- a/Engine/Source/Programs/UnrealBuildAccelerator/Common/Private/UbaLogger.cpp
+++ b/Engine/Source/Programs/UnrealBuildAccelerator/Common/Private/UbaLogger.cpp
@@ -3,6 +3,7 @@
 #include "UbaLogger.h"
 #include "UbaBinaryReaderWriter.h"
 #include "UbaFileAccessor.h"
+#include "UbaParkingLot.h"
 #include "UbaPlatform.h"
 #include "UbaStringBuffer.h"
 #include "UbaThread.h"
@@ -44,8 +45,13 @@ namespace uba
 
 	ANALYSIS_NORETURN void UbaAssert(const tchar* text, const char* file, u32 line, const char* expr, bool allowTerminate, u32 terminateCode, void* context, u32 skipCallstackCount)
 	{
+		#if UBA_USE_PARKINGLOT
+		static ParkingLot::WordMutex assertMutex;
+		SCOPED_CRITICAL_SECTION(assertMutex, scopedLock);
+		#else
 		static ReaderWriterLock& assertLock = *new ReaderWriterLock(); // Leak to prevent asan annoyances during shutdown when asserts happen
 		SCOPED_WRITE_LOCK(assertLock, lock);
+		#endif
 
 		static u8* writerMem = new u8[4096];
 		BinaryWriter writer(writerMem, 0, 4096);
diff --git a/Engine/Source/Programs/UnrealBuildAccelerator/Common/Public/UbaCompactTables.h b/Engine/Source/Programs/UnrealBuildAccelerator/Common/Public/UbaCompactTables.h
--- a/Engine/Source/Programs/UnrealBuildAccelerator/Common/Public/UbaCompactTables.h
+++ b/Engine/Source/Programs/UnrealBuildAccelerator/Common/Public/UbaCompactTables.h
@@ -68,7 +68,7 @@ namespace uba
 		u32 InternalAddWithNewSegment(u32 offset, u32 parentOffset, bool separator, const StringView& segment);
 		void InternalAddWithExistingSegment(u32 offset, u32 parentOffset, bool separator, u32 segmentOffset);
 
-		Futex m_lock;
+		mutable Futex m_lock;
 		MemoryBlock m_pathTableMem;
 		HashMap<StringKey, u32, true> m_offsets;
 		HashMap<StringKey, u32, true> m_segmentOffsets;
diff --git a/Engine/Source/Programs/UnrealBuildAccelerator/Core/Private/UbaParkingLot.cpp b/Engine/Source/Programs/UnrealBuildAccelerator/Core/Private/UbaParkingLot.cpp
new file mode 100644
--- /dev/null
+++ b/Engine/Source/Programs/UnrealBuildAccelerator/Core/Private/UbaParkingLot.cpp
@@ -0,0 +1,1311 @@
+// Copyright Epic Games, Inc. All Rights Reserved.
+
+#include "UbaParkingLot.h"
+#include "UbaPlatform.h"
+
+#include "UbaStringBuffer.h" // for UBA_ASSERTF
+
+#include <algorithm>
+#include <atomic>
+
+#if PLATFORM_LINUX
+#include <linux/futex.h>
+#include <sys/syscall.h>
+#include <time.h>
+#include <unistd.h>
+#endif
+
+namespace uba::ParkingLot::Private
+{
+#if PLATFORM_LINUX
+	inline static const struct timespec* MsToTimeSpec(u32 timeMs, struct timespec& outTime)
+	{
+		if (timeMs != 0xffffffff)
+		{
+			u32 sec = timeMs / 1'000;
+			outTime.tv_sec = sec;
+			outTime.tv_nsec = (timeMs - sec * 1'000) * 1'000'000;
+			return &outTime;
+		}
+		return nullptr;
+	}
+#endif
+
+	class PlatformManualResetEvent
+	{
+	public:
+		PlatformManualResetEvent() = default;
+		PlatformManualResetEvent(const PlatformManualResetEvent&) = delete;
+		PlatformManualResetEvent& operator=(const PlatformManualResetEvent&) = delete;
+
+		inline void Reset()
+		{
+			state.store(0, std::memory_order_relaxed);
+		}
+
+		inline bool Poll()
+		{
+			return !!state.load(std::memory_order_acquire);
+		}
+
+		inline void Wait()
+		{
+			if (!state.load(std::memory_order_acquire))
+			{
+				WaitSlow();
+			}
+		}
+
+		inline bool WaitFor(u32 waitTime)
+		{
+			return state.load(std::memory_order_acquire) || WaitForSlow(waitTime);
+		}
+
+		//inline bool WaitUntil(FMonotonicTimePoint waitTime)
+		//{
+		//	return state.load(std::memory_order_acquire) || WaitUntilSlow(waitTime);
+		//}
+
+		void Notify()
+		{
+			if (state.exchange(1, std::memory_order_release) == 0)
+			{
+			#if PLATFORM_WINDOWS
+				WakeByAddressSingle((void*)&state);
+			#elif PLATFORM_LINUX
+				syscall(SYS_futex, &state, FUTEX_WAKE_PRIVATE, INT_MAX, nullptr, nullptr, 0);
+			#else
+				#error
+			#endif
+			}
+		}
+
+	private:
+		UBA_NOINLINE void WaitSlow()
+		{
+			for (;;)
+			{
+				u32 localState = state.load(std::memory_order_acquire);
+				if (localState)
+				{
+					return;
+				}
+			#if PLATFORM_WINDOWS
+				WaitOnAddress(&state, &localState, sizeof(u32), INFINITE);
+			#elif PLATFORM_LINUX
+			#else
+				#error
+			#endif
+			}
+		}
+
+		UBA_NOINLINE bool WaitForSlow(u32 waitTimeMs)
+		{
+			u32 localState = state.load(std::memory_order_acquire);
+			if (localState || waitTimeMs == 0)
+			{
+				return !!localState;
+			}
+
+		#if PLATFORM_WINDOWS
+			const bool timedOut = !WaitOnAddress(&state, &localState, sizeof(u32), waitTimeMs) && GetLastError() == ERROR_TIMEOUT;
+		#elif PLATFORM_LINUX
+			struct timespec waitTimeSpec;
+			const struct timespec* waitTimeSpecPtr = MsToTimeSpec(waitTimeMs, waitTimeSpec);
+			const bool timedOut = syscall(SYS_futex, &state, FUTEX_WAIT_PRIVATE, 0, waitTimeSpecPtr, nullptr, 0) == -1 && errno == ETIMEDOUT;
+		#else
+			#error
+		#endif
+
+			localState = state.load(std::memory_order_acquire);
+			if (localState || timedOut)
+			{
+				return !!localState;
+			}
+
+			// Handle a spurious wake by waiting until the wait time has elapsed one more time because WaitUntilSlow
+			// handles spurious wakes in a loop and avoids exceeding the originally requested wake time by more than
+			// the typical variation due to scheduling imprecision.
+			// TODO: How to handle spurious wakes?
+			//return WaitUntilSlow(FMonotonicTimePoint::Now() + waitTimeMs);
+			return !!localState;
+		}
+
+		//UBA_NOINLINE bool WaitUntilSlow(FMonotonicTimePoint WaitTime)
+		//{
+		//	FMonotonicTimeSpan WaitSpan = WaitTime - FMonotonicTimePoint::Now();
+		//
+		//	for (;;)
+		//	{
+		//		bool localShouldWait = shouldWait.load(std::memory_order_acquire);
+		//		if (!localShouldWait || WaitSpan <= FMonotonicTimeSpan::Zero())
+		//		{
+		//			return !localShouldWait;
+		//		}
+		//		const DWORD WaitMs = WaitTime.IsInfinity() ? INFINITE : DWORD(FPlatformMath::CeilToInt64(WaitSpan.ToMilliseconds()));
+		//		WaitOnAddress(&shouldWait, &localShouldWait, sizeof(bool), WaitMs);
+		//		WaitSpan = WaitTime - FMonotonicTimePoint::Now();
+		//	}
+		//}
+
+		//std::atomic<bool> shouldWait = true;
+		// 0 = Reset, 1 = Notified
+		std::atomic<u32> state = 0;
+	};
+} // uba::ParkingLot::Private
+
+////////////////////////////////////////////////////////////////////////////////////////////////////
+
+namespace uba::ParkingLot
+{
+	struct WordMutexQueueNode
+	{
+		// Points to the next node in the tail-to-head direction. Only null for the current tail.
+		WordMutexQueueNode* prev = nullptr;
+		// Points to the next node in the head-to-tail direction. The tail points to the head.
+		// Null until UnlockSlow() has traversed from the tail to fill in next pointers.
+		WordMutexQueueNode* next = nullptr;
+
+		Private::PlatformManualResetEvent event;
+	};
+
+	bool WordMutex::TryEnter()
+	{
+		u64 expected = 0;
+		return state.compare_exchange_strong(expected, isLockedFlag, std::memory_order_acquire, std::memory_order_relaxed);
+	}
+
+	void WordMutex::Enter()
+	{
+		u64 expected = 0;
+		if (state.compare_exchange_weak(expected, isLockedFlag, std::memory_order_acquire, std::memory_order_relaxed))
+		{
+			return;
+		}
+
+		EnterSlow();
+	}
+
+	void WordMutex::Leave()
+	{
+		// Unlock immediately to allow other threads to acquire the lock while this thread looks for a thread to wake.
+		u64 currentState = state.fetch_sub(isLockedFlag, std::memory_order_release);
+		UBA_ASSERT(currentState & isLockedFlag);
+
+		// An empty queue indicates that there are no threads to wake.
+		const bool isQueueEmpty = !(currentState & queueMask);
+		// A locked queue indicates that another thread is looking for a thread to wake.
+		const bool isQueueLocked = (currentState & isQueueLockedFlag);
+
+		if (isQueueEmpty || isQueueLocked)
+		{
+			return;
+		}
+
+		LeaveSlow(currentState);
+	}
+
+	UBA_NOINLINE void WordMutex::EnterSlow()
+	{
+		static_assert((alignof(WordMutexQueueNode) & queueMask) == alignof(WordMutexQueueNode),
+			"Alignment of WordMutexQueueNode is insufficient to pack flags into the lower bits.");
+
+		constexpr u32 spinLimit = 40;
+		u32 spinCount = 0;
+		for (;;)
+		{
+			u64 currentState = state.load(std::memory_order_relaxed);
+
+			// Try to acquire the lock if it was unlocked, even if there is a queue.
+			// Acquiring the lock despite the queue means that this lock is not FIFO and thus not fair.
+			if (!(currentState & isLockedFlag))
+			{
+				if (state.compare_exchange_weak(currentState, currentState | isLockedFlag, std::memory_order_acquire, std::memory_order_relaxed))
+				{
+					return;
+				}
+				continue;
+			}
+
+			// Spin up to the spin limit while there is no queue.
+			if (!(currentState & queueMask) && spinCount < spinLimit)
+			{
+				YieldInSpinWait();
+				++spinCount;
+				continue;
+			}
+
+			// Create the node that will be used to add this thread to the queue.
+			WordMutexQueueNode self;
+
+			self.event.Reset();
+
+			// The state points to the tail of the queue, and each node points to the previous node.
+			if (WordMutexQueueNode* tail = (WordMutexQueueNode*)(currentState & queueMask))
+			{
+				self.prev = tail;
+			}
+			else
+			{
+				self.next = &self;
+			}
+
+			// Swap this thread in as the tail, which makes it visible to any other thread that acquires the queue lock.
+			if (!state.compare_exchange_weak(currentState, (currentState & ~queueMask) | u64(&self), std::memory_order_acq_rel, std::memory_order_relaxed))
+			{
+				continue;
+			}
+
+			// Wait until another thread wakes this thread, which can happen as soon as the preceding store completes.
+			self.event.Wait();
+
+			// Loop back and try to acquire the lock.
+			spinCount = 0;
+		}
+	}
+
+	UBA_NOINLINE void WordMutex::LeaveSlow(u64 currentState)
+	{
+		// IsLockedFlag was cleared by Unlock().
+		currentState &= ~isLockedFlag;
+
+		for (;;)
+		{
+			// Try to lock the queue.
+			if (state.compare_exchange_weak(currentState, currentState | isQueueLockedFlag, std::memory_order_acquire, std::memory_order_relaxed))
+			{
+				currentState |= isQueueLockedFlag;
+				break;
+			}
+
+			// A locked queue indicates that another thread is looking for a thread to wake.
+			if ((currentState & isQueueLockedFlag) || !(currentState & queueMask))
+			{
+				return;
+			}
+		}
+
+		for (;;)
+		{
+			// This thread now holds the queue lock. Neither the queue nor State will change while the queue is locked.
+			// The state points to the tail of the queue, and each node points to the previous node.
+			WordMutexQueueNode* tail = (WordMutexQueueNode*)(currentState & queueMask);
+
+			// Traverse from the tail to find the head and set next pointers for any nodes added since the last unlock.
+			for (WordMutexQueueNode* node = tail; !tail->next;)
+			{
+				WordMutexQueueNode* prev = node->prev;
+				UBA_ASSERT(prev);
+				tail->next = prev->next;
+				prev->next = node;
+				node = prev;
+			}
+
+			// Another thread may acquire the lock while this thread has been finding a thread to unlock.
+			// That case will not be detected on the first iteration of the loop, but only when this
+			// thread has failed to unlock the queue at least once. Attempt to unlock the queue here
+			// and allow the next unlock to find a thread to wake.
+			if (currentState & isLockedFlag)
+			{
+				if (state.compare_exchange_weak(currentState, currentState & ~isQueueLockedFlag, std::memory_order_release, std::memory_order_acquire))
+				{
+					return;
+				}
+				continue;
+			}
+
+			// The next node from the tail is the head.
+			WordMutexQueueNode* head = tail->next;
+
+			// Remove the head from the queue and unlock the queue.
+			if (WordMutexQueueNode* newHead = head->next; newHead == head)
+			{
+				// Unlock and clear the queue. Failure needs to restart the loop, because newly-added
+				// nodes will have a pointer to the node being removed.
+				if (!state.compare_exchange_strong(currentState, currentState & isLockedFlag, std::memory_order_release, std::memory_order_acquire))
+				{
+					continue;
+				}
+			}
+			else
+			{
+				// Clear pointers to the head node being removed.
+				UBA_ASSERT(newHead);
+				newHead->prev = nullptr;
+				tail->next = newHead;
+
+				// Unlock the queue regardless of whether new nodes have been added in the meantime.
+				state.fetch_and(~isQueueLockedFlag, std::memory_order_release);
+			}
+
+			// Wake the thread that was at the head of the queue.
+			head->event.Notify();
+			break;
+		}
+	}
+
+} // uba::ParkingLot
+
+////////////////////////////////////////////////////////////////////////////////////////////////////
+
+namespace uba::ParkingLot::Private
+{
+	inline u32 CountLeadingZerosNonZero(u32 value)
+	{
+	#ifdef __clang__
+		return __builtin_clz(value);
+	#else
+		unsigned long bitIndex;
+		_BitScanReverse(&bitIndex, value);
+		return 31 - bitIndex;
+	#endif
+	}
+
+	inline u32 RoundUpToPowerOfTwo(u32 value)
+	{
+		if (value <= 1)
+			return 1;
+		return u32(1) << (32 - CountLeadingZerosNonZero(value - 1));
+	}
+
+	////////////////////////////////////////////////////////////////////////////////////////////////////
+
+	template <typename RefType, typename AssignedType = RefType>
+	struct GuardValue
+	{
+		GuardValue(const GuardValue&) = delete;
+		GuardValue& operator=(const GuardValue&) = delete;
+
+		[[nodiscard]] GuardValue(RefType& value, const AssignedType& newValue)
+			: valueRef(value)
+			, oldValue(value)
+		{
+			valueRef = newValue;
+		}
+
+		~GuardValue()
+		{
+			valueRef = oldValue;
+		}
+
+	private:
+		RefType& valueRef;
+		AssignedType oldValue;
+	};
+
+	////////////////////////////////////////////////////////////////////////////////////////////////////
+
+	template <typename ReferencedType>
+	class RefCountPtr
+	{
+		using ReferenceType = ReferencedType*;
+
+	public:
+		RefCountPtr() = default;
+
+		inline RefCountPtr(ReferencedType* inReference)
+		{
+			reference = inReference;
+			if (reference)
+			{
+				reference->AddRef();
+			}
+		}
+
+		inline RefCountPtr(const RefCountPtr& other)
+		{
+			reference = other.reference;
+			if (reference)
+			{
+				reference->AddRef();
+			}
+		}
+
+		inline RefCountPtr(RefCountPtr&& other)
+		{
+			reference = other.reference;
+			other.reference = nullptr;
+		}
+
+		inline ~RefCountPtr()
+		{
+			if (reference)
+			{
+				reference->Release();
+			}
+		}
+
+		inline RefCountPtr& operator=(ReferencedType* newReference)
+		{
+			if (reference != newReference)
+			{
+				ReferencedType* oldReference = reference;
+				reference = newReference;
+				if (reference)
+				{
+					reference->AddRef();
+				}
+				if (oldReference)
+				{
+					oldReference->Release();
+				}
+			}
+			return *this;
+		}
+
+		inline RefCountPtr& operator=(const RefCountPtr& other)
+		{
+			return *this = other.reference;
+		}
+
+		inline RefCountPtr& operator=(RefCountPtr&& other)
+		{
+			if (this != &other)
+			{
+				ReferencedType* oldReference = reference;
+				reference = other.reference;
+				other.reference = nullptr;
+				if (oldReference)
+				{
+					oldReference->Release();
+				}
+			}
+			return *this;
+		}
+
+		inline ReferencedType* operator->() const
+		{
+			return reference;
+		}
+
+		inline ReferencedType& operator*() const
+		{
+			return *reference;
+		}
+
+		inline explicit operator bool() const
+		{
+			return !!reference;
+		}
+
+	private:
+		ReferencedType* reference = nullptr;
+	};
+
+	////////////////////////////////////////////////////////////////////////////////////////////////////
+
+	template <typename LockType>
+	class DynamicUniqueLock final
+	{
+	public:
+		DynamicUniqueLock() = default;
+
+		DynamicUniqueLock(const DynamicUniqueLock&) = delete;
+		DynamicUniqueLock& operator=(const DynamicUniqueLock&) = delete;
+
+		[[nodiscard]] inline explicit DynamicUniqueLock(LockType& lock)
+			: mutex(&lock)
+		{
+			mutex->Enter();
+			isLocked = true;
+		}
+
+		[[nodiscard]] inline DynamicUniqueLock(DynamicUniqueLock&& other)
+			: mutex(other.mutex)
+			, isLocked(other.isLocked)
+		{
+			other.mutex = nullptr;
+			other.isLocked = false;
+		}
+
+		inline DynamicUniqueLock& operator=(DynamicUniqueLock&& other)
+		{
+			if (isLocked)
+			{
+				mutex->Leave();
+			}
+			mutex = other.mutex;
+			isLocked = other.isLocked;
+			other.mutex = nullptr;
+			other.isLocked = false;
+			return *this;
+		}
+
+		inline ~DynamicUniqueLock()
+		{
+			if (isLocked)
+			{
+				mutex->Leave();
+			}
+		}
+
+		void Enter()
+		{
+			UBA_ASSERT(!isLocked);
+			UBA_ASSERT(mutex);
+			mutex->Enter();
+			isLocked = true;
+		}
+
+		void Leave()
+		{
+			UBA_ASSERT(isLocked);
+			isLocked = false;
+			mutex->Leave();
+		}
+
+		inline bool OwnsLock() const
+		{
+			return isLocked;
+		}
+
+		inline explicit operator bool() const
+		{
+			return OwnsLock();
+		}
+
+	private:
+		LockType* mutex = nullptr;
+		bool isLocked = false;
+	};
+
+	////////////////////////////////////////////////////////////////////////////////////////////////////
+
+	/**
+	 * Identifies which callback function is currently executing.
+	 *
+	 * NOTE: ORDER MATTERS FOR COMPARISONS!
+	 */
+	enum class ExecutingFunction : u8
+	{
+		None = 0,
+		BeforeWait = 1,
+		CanWait = 2,
+		OnWakeState = 3,
+	};
+
+	const tchar* LexToString(ExecutingFunction function)
+	{
+		switch (function)
+		{
+		default:
+		case ExecutingFunction::None: return TC("None");
+		case ExecutingFunction::BeforeWait: return TC("BeforeWait");
+		case ExecutingFunction::CanWait: return TC("CanWait");
+		case ExecutingFunction::OnWakeState: return TC("OnWakeState");
+		}
+	}
+
+	static thread_local ExecutingFunction parkingLotExecutingFunction = ExecutingFunction::None;
+
+	////////////////////////////////////////////////////////////////////////////////////////////////////
+
+	struct alignas(CacheLineSize) Thread final
+	{
+		Thread* next = nullptr;
+		std::atomic<const void*> waitAddress = nullptr;
+		u64 wakeToken = 0;
+		PlatformManualResetEvent event;
+		std::atomic<u32> referenceCount = 0;
+		u64 padding[4]{};
+
+		inline void AddRef()
+		{
+			referenceCount.fetch_add(1, std::memory_order_relaxed);
+		}
+
+		inline void Release()
+		{
+			if (referenceCount.fetch_sub(1, std::memory_order_acq_rel) == 1)
+			{
+				delete this;
+			}
+		}
+
+		Thread() = default;
+		~Thread() = default;
+
+		Thread(const Thread&) = delete;
+		Thread& operator=(const Thread&) = delete;
+	};
+
+	class ThreadLocalData
+	{
+	public:
+		static RefCountPtr<Thread> Get()
+		{
+			static thread_local ThreadLocalData threadLocalData;
+			ThreadLocalData& localThreadLocalData = threadLocalData;
+			if (localThreadLocalData.thread)
+			{
+				return localThreadLocalData.thread;
+			}
+			if (localThreadLocalData.isDestroyed)
+			{
+				return new Thread;
+			}
+			localThreadLocalData.thread = new Thread;
+			return localThreadLocalData.thread;
+		}
+
+	private:
+		inline static std::atomic<u32> threadCount = 0;
+
+		RefCountPtr<Thread> thread;
+		bool isDestroyed = false;
+
+		ThreadLocalData();
+		~ThreadLocalData();
+
+		ThreadLocalData(const ThreadLocalData&) = delete;
+		ThreadLocalData& operator=(const ThreadLocalData&) = delete;
+	};
+
+	////////////////////////////////////////////////////////////////////////////////////////////////////
+
+	enum class QueueAction
+	{
+		Stop,
+		Continue,
+		RemoveAndStop,
+		RemoveAndContinue,
+	};
+
+	class alignas(CacheLineSize) Bucket final
+	{
+	public:
+		[[nodiscard]] inline DynamicUniqueLock<WordMutex> LockDynamic() { return DynamicUniqueLock(mutex); }
+
+		inline void Lock() { mutex.Enter(); }
+		inline void Unlock() { mutex.Leave(); }
+
+		inline bool IsEmpty() const { return !head; }
+
+		void Enqueue(Thread* thread);
+		Thread* Dequeue();
+
+		template <typename VisitorType>
+		void DequeueIf(VisitorType&& visitor);
+
+		Bucket() = default;
+		~Bucket() = default;
+
+	private:
+		Bucket(const Bucket&) = delete;
+		Bucket& operator=(const Bucket&) = delete;
+
+		WordMutex mutex;
+		Thread* head = nullptr;
+		Thread* tail = nullptr;
+		u64 padding[5]{};
+	};
+
+	void Bucket::Enqueue(Thread* thread)
+	{
+		UBA_ASSERT(thread);
+		UBA_ASSERT(!thread->next);
+		if (tail)
+		{
+			tail->next = thread;
+			tail = thread;
+		}
+		else
+		{
+			head = thread;
+			tail = thread;
+		}
+	}
+
+	Thread* Bucket::Dequeue()
+	{
+		Thread* thread = head;
+		if (thread)
+		{
+			head = thread->next;
+			thread->next = nullptr;
+			if (tail == thread)
+			{
+				tail = nullptr;
+			}
+		}
+		return thread;
+	}
+
+	template <typename VisitorType>
+	void Bucket::DequeueIf(VisitorType&& visitor)
+	{
+		Thread** next = &head;
+		Thread* prev = nullptr;
+		while (Thread* thread = *next)
+		{
+			switch (visitor(thread))
+			{
+			case QueueAction::Stop:
+				return;
+			case QueueAction::Continue:
+				prev = thread;
+				next = &thread->next;
+				break;
+			case QueueAction::RemoveAndStop:
+				if (tail == thread)
+				{
+					tail = prev;
+				}
+				*next = thread->next;
+				thread->next = nullptr;
+				return;
+			case QueueAction::RemoveAndContinue:
+				if (tail == thread)
+				{
+					tail = prev;
+				}
+				*next = thread->next;
+				thread->next = nullptr;
+				break;
+			}
+		}
+	}
+
+	////////////////////////////////////////////////////////////////////////////////////////////////////
+
+	#pragma warning(push)
+	#pragma warning(disable:4200)  // nonstandard extension used: zero-length array in struct/union
+
+	class Table final
+	{
+		constexpr static u32 minSize = 32;
+
+	public:
+		static Bucket& FindOrCreateBucket(const void* address, DynamicUniqueLock<WordMutex>& outLock);
+		static Bucket* FindBucket(const void* address, DynamicUniqueLock<WordMutex>& outLock);
+
+		static void Reserve(u32 threadCount);
+
+	private:
+		static Table& CreateOrGet();
+		static Table& Create(u32 size);
+		static void Destroy(Table& table);
+
+		static bool TryLock(Table& table, Vector<Bucket*>& outBuckets);
+		static void Unlock(Vector<Bucket*>& lockedBuckets);
+
+		static u32 HashAddress(const void* address);
+
+		inline static std::atomic<Table*> globalTable;
+
+		Table() = default;
+		~Table() = default;
+
+		Table(const Table&) = delete;
+		Table& operator=(const Table&) = delete;
+
+		template <typename AllocatorType>
+		Bucket& FindOrCreateBucket(u32 index, AllocatorType&& bucketAllocator);
+
+		u32 bucketCount = 0;
+		std::atomic<Bucket*> buckets[0];
+	};
+
+	#pragma warning(pop)
+
+	Bucket& Table::FindOrCreateBucket(const void* address, DynamicUniqueLock<WordMutex>& outLock)
+	{
+		const u32 hash = HashAddress(address);
+
+		for (;;)
+		{
+			Table& table = CreateOrGet();
+			const u32 index = hash % table.bucketCount;
+			Bucket& bucket = table.FindOrCreateBucket(index, [] { return new Bucket; });
+			outLock = bucket.LockDynamic();
+
+			if (&table == globalTable.load(std::memory_order_acquire))
+			{
+				return bucket;
+			}
+
+			// Restart because the table was resized since it was loaded above.
+			outLock.Leave();
+		}
+	}
+
+	Bucket* Table::FindBucket(const void* address, DynamicUniqueLock<WordMutex>& outLock)
+	{
+		const u32 hash = HashAddress(address);
+
+		for (;;)
+		{
+			if (Table* table = globalTable.load(std::memory_order_acquire))
+			{
+				const u32 index = hash % table->bucketCount;
+				if (Bucket* bucket = table->buckets[index].load(std::memory_order_acquire))
+				{
+					outLock = bucket->LockDynamic();
+
+					if (table == globalTable.load(std::memory_order_acquire))
+					{
+						return bucket;
+					}
+
+					// Restart because the table was resized since it was loaded above.
+					outLock = {};
+					continue;
+				}
+			}
+			return nullptr;
+		}
+	}
+
+	template <typename AllocatorType>
+	Bucket& Table::FindOrCreateBucket(u32 index, AllocatorType&& bucketAllocator)
+	{
+		std::atomic<Bucket*>& bucketPtr = buckets[index];
+		Bucket* bucket = bucketPtr.load(std::memory_order_acquire);
+		if (!bucket)
+		{
+			Bucket* newBucket = bucketAllocator();
+			if (bucketPtr.compare_exchange_strong(bucket, newBucket, std::memory_order_release, std::memory_order_acquire))
+			{
+				bucket = newBucket;
+			}
+			else
+			{
+				delete newBucket;
+			}
+			UBA_ASSERT(bucket);
+		}
+		return *bucket;
+	}
+
+	Table& Table::CreateOrGet()
+	{
+		Table* table = globalTable.load(std::memory_order_acquire);
+
+		if (table)
+		{
+			return *table;
+		}
+
+		Table& newTable = Create(minSize);
+
+		if (globalTable.compare_exchange_strong(table, &newTable, std::memory_order_release, std::memory_order_acquire))
+		{
+			return newTable;
+		}
+
+		Destroy(newTable);
+
+		UBA_ASSERT(table);
+		return *table;
+	}
+
+	Table& Table::Create(const u32 size)
+	{
+		const u64 memorySize = sizeof(Table) + sizeof(Bucket*) * u64(size);
+		void* const memory = aligned_alloc(alignof(Table), memorySize);
+		memset(memory, 0, memorySize);
+		Table& table = *new(memory) Table;
+		table.bucketCount = size;
+		return table;
+	}
+
+	void Table::Destroy(Table& table)
+	{
+		table.~Table();
+		aligned_free(&table);
+	}
+
+	void Table::Reserve(const u32 threadCount)
+	{
+		const u32 targetBucketCount = RoundUpToPowerOfTwo(threadCount);
+		Vector<Bucket*> existingBuckets;
+
+		for (;;)
+		{
+			Table& existingTable = CreateOrGet();
+
+			if (existingTable.bucketCount >= targetBucketCount)
+			{
+				// Reserve is called every time a thread is created and has amortized constant time
+				// because of its power-of-two table growth. Most calls return here without locking.
+				return;
+			}
+
+			if (!TryLock(existingTable, existingBuckets))
+			{
+				continue;
+			}
+
+			// Gather waiting threads to be redistributed into the buckets of the new table.
+			// Threads with the same address remain in the same relative order as they were queued.
+			Vector<Thread*> threads;
+			for (Bucket* bucket : existingBuckets)
+			{
+				while (Thread* thread = bucket->Dequeue())
+				{
+					threads.push_back(thread);
+				}
+			}
+
+			Table& newTable = Create(targetBucketCount);
+
+			// Reuse existing now-empty buckets when populating the new table.
+			Vector<Bucket*> availableBuckets = existingBuckets;
+			const auto allocateBucket = [&availableBuckets]() -> Bucket*
+			{
+				if (!availableBuckets.empty())
+				{
+					Bucket* bucket = availableBuckets.back();
+					availableBuckets.pop_back();
+					return bucket;
+				}
+				return new Bucket;
+			};
+
+			// Add waiting threads to the new table.
+			for (Thread* thread : threads)
+			{
+				const u32 hash = HashAddress(thread->waitAddress.load(std::memory_order_relaxed));
+				const u32 index = hash % newTable.bucketCount;
+				Bucket& bucket = newTable.FindOrCreateBucket(index, allocateBucket);
+				bucket.Enqueue(thread);
+			}
+
+			// Assign any available buckets to the table to avoid having to free them.
+			for (u32 index = 0; !availableBuckets.empty() && index < newTable.bucketCount; ++index)
+			{
+				newTable.FindOrCreateBucket(index, allocateBucket);
+			}
+			UBA_ASSERT(availableBuckets.empty());
+
+			// Make the new table visible to other threads.
+			[[maybe_unused]] Table* compareTable = globalTable.exchange(&newTable, std::memory_order_release);
+			UBA_ASSERT(compareTable == &existingTable);
+
+			// Unlock buckets that came from the existing table now that the new table is visible.
+			Unlock(existingBuckets);
+			return;
+		}
+	}
+
+	bool Table::TryLock(Table& table, Vector<Bucket*>& outBuckets)
+	{
+		outBuckets.clear();
+		outBuckets.reserve(table.bucketCount);
+
+		// Gather buckets from the table, creating them as needed because the lock is on the bucket.
+		for (u32 index = 0; index < table.bucketCount; ++index)
+		{
+			outBuckets.push_back(&table.FindOrCreateBucket(index, [] { return new Bucket; }));
+		}
+
+		// Lock the buckets in order by address to ensure consistent ordering regardless of the table being locked.
+		std::sort(outBuckets.begin(), outBuckets.end());
+		for (Bucket* bucket : outBuckets)
+		{
+			bucket->Lock();
+		}
+
+		// Table is locked if the global table pointer still points to it, otherwise it has grown.
+		if (&table == globalTable)
+		{
+			return true;
+		}
+
+		// Unlock and return that the table could not be locked.
+		Unlock(outBuckets);
+		outBuckets.clear();
+		return false;
+	}
+
+	void Table::Unlock(Vector<Bucket*>& lockedBuckets)
+	{
+		for (Bucket* bucket : lockedBuckets)
+		{
+			bucket->Unlock();
+		}
+	}
+
+	u32 Table::HashAddress(const void* address)
+	{
+		constexpr u64 A = 0xdc2b17dc9d2fbc29;
+		constexpr u64 B = 0xcb1014192cb2c5fc;
+		constexpr u64 C = 0x5b12db9242bd7ce7;
+		const u64 value = u64(address);
+		return u32(((A * (value >> 32)) + (B * (value & 0xffffffff)) + C) >> 32);
+	}
+
+	////////////////////////////////////////////////////////////////////////////////////////////////////
+
+	ThreadLocalData::ThreadLocalData()
+	{
+		Table::Reserve(threadCount.fetch_add(1, std::memory_order_relaxed) + 1);
+	}
+
+	ThreadLocalData::~ThreadLocalData()
+	{
+		threadCount.fetch_sub(1, std::memory_order_relaxed);
+		thread = nullptr;
+		isDestroyed = true;
+	}
+
+	////////////////////////////////////////////////////////////////////////////////////////////////////
+
+	WaitState Wait(const void* address, bool (*canWait)(void*), void* canWaitContext, void (*beforeWait)(void*), void* beforeWaitContext)
+	{
+		using namespace Private;
+
+		// A failure of this assertion is either an error in a callback, because callbacks may not call into ParkingLot,
+		// or a dependency on ParkingLot from within the memory allocator used by ParkingLot.
+		UBA_ASSERTF(parkingLotExecutingFunction == ExecutingFunction::None,
+			TC("Detected a re-entrant call into ParkingLot from within a %s callback."), LexToString(parkingLotExecutingFunction));
+
+		RefCountPtr<Thread> self = ThreadLocalData::Get();
+
+		UBA_ASSERTF(!self->waitAddress.load(std::memory_order_relaxed), TC("waitAddress must be null. This can happen if Wait is called by BeforeWait."));
+		UBA_ASSERTF(self->wakeToken == 0, TC("wakeToken must be 0. This is an error in ParkingLot."));
+
+		WaitState state;
+
+		// Enqueue the thread if CanWait returns true while the bucket is locked.
+		{
+			DynamicUniqueLock<WordMutex> bucketLock;
+			Bucket& bucket = Table::FindOrCreateBucket(address, bucketLock);
+			if (canWait)
+			{
+				GuardValue guardActive(parkingLotExecutingFunction, ExecutingFunction::CanWait);
+				if (!canWait(canWaitContext))
+				{
+					return state;
+				}
+			}
+			state.didWait = true;
+			self->waitAddress.store(address, std::memory_order_relaxed);
+			self->event.Reset();
+			bucket.Enqueue(&*self);
+		}
+
+		// BeforeWait must be invoked after the bucket is unlocked.
+		if (beforeWait)
+		{
+			GuardValue guardActive(parkingLotExecutingFunction, ExecutingFunction::BeforeWait);
+			beforeWait(beforeWaitContext);
+		}
+
+		// Wait until the thread has been dequeued.
+		self->event.Wait();
+
+		// waitAddress is reset when the thread is dequeued.
+		UBA_ASSERT(!self->waitAddress.load(std::memory_order_relaxed));
+		state.didWake = true;
+		state.wakeToken = self->wakeToken;
+		self->wakeToken = 0;
+		return state;
+	}
+
+	static WaitState TimedWait(
+		const void* address,
+		bool (*canWait)(void*),
+		void* canWaitContext,
+		void (*beforeWait)(void*),
+		void* beforeWaitContext,
+		void (*waitOnEvent)(void*, PlatformManualResetEvent&),
+		void* waitOnEventContext)
+	{
+		using namespace Private;
+
+		// A failure of this assertion is either an error in a callback, because callbacks may not call into ParkingLot,
+		// or a dependency on ParkingLot from within the memory allocator used by ParkingLot.
+		UBA_ASSERTF(parkingLotExecutingFunction == ExecutingFunction::None,
+			TC("Detected a re-entrant call into ParkingLot from within a %s callback."), LexToString(parkingLotExecutingFunction));
+
+		RefCountPtr<Thread> self = ThreadLocalData::Get();
+
+		UBA_ASSERTF(!self->waitAddress.load(std::memory_order_relaxed), TC("waitAddress must be null. This can happen if Wait is called by BeforeWait."));
+		UBA_ASSERTF(self->wakeToken == 0, TC("wakeToken must be 0. This is an error in ParkingLot."));
+
+		WaitState state;
+
+		// Enqueue the thread if CanWait returns true while the bucket is locked.
+		{
+			DynamicUniqueLock<WordMutex> bucketLock;
+			Bucket& bucket = Table::FindOrCreateBucket(address, bucketLock);
+			if (canWait)
+			{
+				GuardValue guardActive(parkingLotExecutingFunction, ExecutingFunction::CanWait);
+				if (!canWait(canWaitContext))
+				{
+					return state;
+				}
+			}
+			state.didWait = true;
+			self->waitAddress.store(address, std::memory_order_relaxed);
+			self->event.Reset();
+			bucket.Enqueue(&*self);
+		}
+
+		// BeforeWait must be invoked after the bucket is unlocked.
+		if (beforeWait)
+		{
+			GuardValue guardActive(parkingLotExecutingFunction, ExecutingFunction::BeforeWait);
+			beforeWait(beforeWaitContext);
+		}
+
+		// Wait until the timeout or until the thread has been dequeued.
+		waitOnEvent(waitOnEventContext, self->event);
+
+		// waitAddress is reset when the thread is dequeued.
+		if (!self->waitAddress.load(std::memory_order_relaxed))
+		{
+			state.didWake = true;
+			state.wakeToken = self->wakeToken;
+			self->wakeToken = 0;
+			return state;
+		}
+
+		// The timeout was reached and the thread needs to dequeue itself.
+		// This can race with a call to wake a thread, which means Self is unsafe to access outside of the lock.
+		bool didDequeue = false;
+		if (DynamicUniqueLock<WordMutex> bucketLock; Bucket* bucket = Table::FindBucket(address, bucketLock))
+		{
+			bucket->DequeueIf([self = &*self, &didDequeue](Thread* thread)
+			{
+				if (thread == self)
+				{
+					didDequeue = true;
+					thread->waitAddress.store(nullptr, std::memory_order_relaxed);
+					return QueueAction::RemoveAndStop;
+				}
+				return QueueAction::Continue;
+			});
+		}
+
+		// The thread did not dequeue itself, which means that we need to wait until the other thread
+		// has finished waking this thread by setting its wait address to null.
+		if (!didDequeue)
+		{
+			self->event.Wait();
+			state.didWake = true;
+			state.wakeToken = self->wakeToken;
+			self->wakeToken = 0;
+		}
+
+		return state;
+	}
+
+	static void WaitForTime(void* waitTime, PlatformManualResetEvent& event)
+	{
+		event.WaitFor(*(const u32*)waitTime);
+	}
+
+	WaitState WaitFor(const void* address, bool (*canWait)(void*), void* canWaitContext, void (*beforeWait)(void*), void* beforeWaitContext, u32 waitTimeMs)
+	{
+		return TimedWait(address, canWait, canWaitContext, beforeWait, beforeWaitContext, WaitForTime, &waitTimeMs);
+	}
+
+	//static void WaitUntilTime(void* waitTime, PlatformManualResetEvent& event)
+	//{
+	//	event.WaitUntil(*(FMonotonicTimePoint*)waitTime);
+	//}
+
+	//WaitState WaitUntil(const void* address, bool (*canWait)(void*), void* canWaitContext, void (*beforeWait)(void*), void* beforeWaitContext, FMonotonicTimePoint waitTime)
+	//{
+	//	UBA_ASSERT(!waitTime.IsNaN());
+	//	return TimedWait(address, canWait, canWaitContext, beforeWait, beforeWaitContext, WaitUntilTime, &waitTime);
+	//}
+
+	void WakeOne(const void* address, u64 (*onWakeState)(void*, WakeState), void* onWakeStateContext)
+	{
+		using namespace Private;
+
+		// A failure of this assertion is either an error in a callback, because callbacks may not call into ParkingLot,
+		// or a dependency on ParkingLot from within the memory allocator used by ParkingLot.
+		UBA_ASSERTF(parkingLotExecutingFunction <= ExecutingFunction::BeforeWait,
+			TC("Detected a re-entrant call into ParkingLot from within a %s callback."), LexToString(parkingLotExecutingFunction));
+
+		RefCountPtr<Thread> wakeThread;
+		u64 wakeToken = 0;
+
+		{
+			DynamicUniqueLock<WordMutex> bucketLock;
+			Bucket& bucket = Table::FindOrCreateBucket(address, bucketLock);
+			bucket.DequeueIf([address, &wakeThread](Thread* thread)
+			{
+				if (thread->waitAddress.load(std::memory_order_relaxed) == address)
+				{
+					wakeThread = thread;
+					return QueueAction::RemoveAndStop;
+				}
+				return QueueAction::Continue;
+			});
+			GuardValue guardActive(parkingLotExecutingFunction, ExecutingFunction::OnWakeState);
+			WakeState wakeState;
+			wakeState.didWake = !!wakeThread;
+			wakeState.hasWaitingThreads = !bucket.IsEmpty();
+			wakeToken = onWakeState ? onWakeState(onWakeStateContext, wakeState) : 0;
+		}
+
+		if (wakeThread)
+		{
+			wakeThread->wakeToken = wakeToken;
+			wakeThread->waitAddress.store(nullptr, std::memory_order_relaxed);
+			wakeThread->event.Notify();
+		}
+	}
+
+	static u64 CopyWakeState(void* outState, WakeState state)
+	{
+		*(WakeState*)outState = state;
+		return 0;
+	}
+} // uba::ParkingLot::Private
+
+namespace uba::ParkingLot
+{
+	WakeState WakeOne(const void* address)
+	{
+		WakeState outState;
+		Private::WakeOne(address, Private::CopyWakeState, &outState);
+		return outState;
+	}
+
+	u32 WakeMultiple(const void* const address, const u32 wakeCount)
+	{
+		using namespace Private;
+
+		// A failure of this assertion is either an error in a callback, because callbacks may not call into ParkingLot,
+		// or a dependency on ParkingLot from within the memory allocator used by ParkingLot.
+		UBA_ASSERTF(parkingLotExecutingFunction <= ExecutingFunction::BeforeWait,
+			TC("Detected a re-entrant call into ParkingLot from within a %s callback."), LexToString(parkingLotExecutingFunction));
+
+		// TODO: TInlineAllocator?
+		//Vector<RefCountPtr<FThread>, TInlineAllocator<128>> wakeThreads;
+		Vector<RefCountPtr<Thread>> wakeThreads;
+
+		if (DynamicUniqueLock<WordMutex> bucketLock; Bucket* bucket = Table::FindBucket(address, bucketLock))
+		{
+			bucket->DequeueIf([address, &wakeThreads, wakeCount](Thread* thread)
+			{
+				if (thread->waitAddress.load(std::memory_order_relaxed) == address)
+				{
+					wakeThreads.emplace_back(thread);
+					return wakeThreads.size() == wakeCount ? QueueAction::RemoveAndStop : QueueAction::RemoveAndContinue;
+				}
+				return QueueAction::Continue;
+			});
+		}
+
+		for (RefCountPtr<Thread>& wakeThread : wakeThreads)
+		{
+			wakeThread->waitAddress.store(nullptr, std::memory_order_relaxed);
+			wakeThread->event.Notify();
+		}
+
+		return u32(wakeThreads.size());
+	}
+
+	u32 WakeAll(const void* address)
+	{
+		return WakeMultiple(address, 0xffffffff);
+	}
+} // uba::ParkingLot
diff --git a/Engine/Source/Programs/UnrealBuildAccelerator/Core/Private/UbaSynchronization.cpp b/Engine/Source/Programs/UnrealBuildAccelerator/Core/Private/UbaSynchronization.cpp
--- a/Engine/Source/Programs/UnrealBuildAccelerator/Core/Private/UbaSynchronization.cpp
+++ b/Engine/Source/Programs/UnrealBuildAccelerator/Core/Private/UbaSynchronization.cpp
@@ -1,6 +1,9 @@
 // Copyright Epic Games, Inc. All Rights Reserved.
 
 #include "UbaSynchronization.h"
+#if UBA_USE_PARKINGLOT
+#include "UbaParkingLot.h"
+#endif
 #include "UbaPlatform.h"
 #include "UbaStringBuffer.h"
 #include <shared_mutex>
@@ -9,6 +12,322 @@
 
 namespace uba
 {
+	#if UBA_USE_PARKINGLOT
+
+	inline static u32 getCurrentThreadId()
+	{
+		#if PLATFORM_WINDOWS
+		return GetCurrentThreadId();
+		#elif PLATFORM_LINUX
+		static thread_local u32 threadId = []
+		{
+			static_assert(sizeof(pid_t) <= sizeof(u32));
+			u32 id = (u32)static_cast<pid_t>(syscall(SYS_gettid));
+			UBA_ASSERT(id != 0);
+			return id;
+		}();
+		return threadId;
+		#elif PLATFORM_MAC
+		return (u32)pthread_mach_thread_np(pthread_self());
+		#endif
+	}
+
+	bool CriticalSection::TryEnter()
+	{
+		const u32 currentThreadId = getCurrentThreadId();
+		u32 currentState = state.load(std::memory_order_relaxed);
+
+		// Try to acquire the lock if it was unlocked, even if there are waiting threads.
+		// Acquiring the lock despite the waiting threads means that this lock is not FIFO and thus not fair.
+		if (!(currentState & lockCountMask))
+		{
+			if (state.compare_exchange_strong(currentState, currentState | (1 << lockCountShift), std::memory_order_acquire, std::memory_order_relaxed))
+			{
+				UBA_ASSERT(threadId.load(std::memory_order_relaxed) == 0);
+				threadId.store(currentThreadId, std::memory_order_relaxed);
+				return true;
+			}
+		}
+
+		// Lock recursively if this is the thread that holds the lock.
+		if (threadId.load(std::memory_order_relaxed) == currentThreadId)
+		{
+			state.fetch_add(1 << lockCountShift, std::memory_order_relaxed);
+			return true;
+		}
+
+		return false;
+	}
+
+	void CriticalSection::Enter()
+	{
+		const u32 currentThreadId = getCurrentThreadId();
+		u32 currentState = state.load(std::memory_order_relaxed);
+
+		// Try to acquire the lock if it was unlocked, even if there are waiting threads.
+		// Acquiring the lock despite the waiting threads means that this lock is not FIFO and thus not fair.
+		if (!(currentState & lockCountMask))
+		{
+			if (state.compare_exchange_weak(currentState, currentState | (1 << lockCountShift), std::memory_order_acquire, std::memory_order_relaxed))
+			{
+				UBA_ASSERT(threadId.load(std::memory_order_relaxed) == 0);
+				threadId.store(currentThreadId, std::memory_order_relaxed);
+				return;
+			}
+		}
+
+		// Lock recursively if this is the thread that holds the lock.
+		if (threadId.load(std::memory_order_relaxed) == currentThreadId)
+		{
+			state.fetch_add(1 << lockCountShift, std::memory_order_relaxed);
+			return;
+		}
+
+		EnterSlow(currentState, currentThreadId);
+	}
+
+	void CriticalSection::EnterSlow(u32 currentState, u32 currentThreadId)
+	{
+		constexpr u32 spinLimit = 40;
+		u32 spinCount = 0;
+		for (;;)
+		{
+			// Try to acquire the lock if it was unlocked, even if there are waiting threads.
+			// Acquiring the lock despite the waiting threads means that this lock is not FIFO and thus not fair.
+			if (!(currentState & lockCountMask))
+			{
+				if (state.compare_exchange_weak(currentState, currentState | (1 << lockCountShift), std::memory_order_acquire, std::memory_order_relaxed))
+				{
+					UBA_ASSERT(threadId.load(std::memory_order_relaxed) == 0);
+					threadId.store(currentThreadId, std::memory_order_relaxed);
+					return;
+				}
+				continue;
+			}
+
+			// Spin up to the spin limit while there are no waiting threads.
+			if (!(currentState & mayHaveWaitingLockFlag) && spinCount < spinLimit)
+			{
+				YieldInSpinWait();
+				++spinCount;
+				currentState = state.load(std::memory_order_relaxed);
+				continue;
+			}
+
+			// Store that there are waiting threads. Restart if the state has changed since it was loaded.
+			if (!(currentState & mayHaveWaitingLockFlag))
+			{
+				if (!state.compare_exchange_weak(currentState, currentState | mayHaveWaitingLockFlag, std::memory_order_relaxed))
+				{
+					continue;
+				}
+				currentState |= mayHaveWaitingLockFlag;
+			}
+
+			// Wait if the state has not changed. Either way, loop back and try to acquire the lock after trying to wait.
+			ParkingLot::Wait(&state, [this, currentState] { return state.load(std::memory_order_relaxed) == currentState; }, nullptr);
+			currentState = state.load(std::memory_order_relaxed);
+		}
+	}
+
+	void CriticalSection::Leave()
+	{
+		u32 currentState = state.load(std::memory_order_relaxed);
+		UBA_ASSERT(currentState & lockCountMask);
+		UBA_ASSERT(threadId.load(std::memory_order_relaxed) == getCurrentThreadId());
+
+		if ((currentState & lockCountMask) == (1 << lockCountShift))
+		{
+			// Remove the association with this thread before unlocking.
+			threadId.store(0, std::memory_order_relaxed);
+
+			// Unlock immediately to allow other threads to acquire the lock while this thread looks for a thread to wake.
+			const u32 lastState = state.fetch_sub(1 << lockCountShift, std::memory_order_release);
+
+			// Wake one exclusive waiter if there are waiting threads.
+			if (lastState & mayHaveWaitingLockFlag)
+			{
+				WakeWaitingThread();
+			}
+		}
+		else
+		{
+			// This is recursively locked. Decrement the lock count.
+			state.fetch_sub(1 << lockCountShift, std::memory_order_relaxed);
+		}
+	}
+
+	void CriticalSection::WakeWaitingThread()
+	{
+		ParkingLot::WakeOne(&state, [this](ParkingLot::WakeState wakeState) -> u64
+		{
+			if (!wakeState.hasWaitingThreads)
+			{
+				state.fetch_and(~mayHaveWaitingLockFlag, std::memory_order_relaxed);
+			}
+			return 0;
+		});
+	}
+
+	inline const void* ReaderWriterLock::GetLockAddress() const
+	{
+		return &state;
+	}
+
+	inline const void* ReaderWriterLock::GetReadLockAddress() const
+	{
+		// Read locks need a distinct address from exclusive locks to allow threads waiting for exclusive ownership
+		// to be woken up without waking any threads waiting for read ownership.
+		return (const u8*)&state + 1;
+	}
+
+	UBA_NOINLINE void ReaderWriterLock::EnterSlow()
+	{
+		constexpr u32 spinLimit = 40;
+		u32 spinCount = 0;
+		for (u32 currentState = state.load(std::memory_order_relaxed);;)
+		{
+			// Try to acquire the lock if it was unlocked, even if there are waiting threads.
+			// Acquiring the lock despite the waiting threads means that this lock is not FIFO and thus not fair.
+			if (!(currentState & (isLockedFlag | readLockCountMask)))
+			{
+				if (state.compare_exchange_weak(currentState, currentState | isLockedFlag, std::memory_order_acquire, std::memory_order_relaxed))
+				{
+					return;
+				}
+				continue;
+			}
+
+			// Spin up to the spin limit while there are no waiting threads.
+			if (!(currentState & mayHaveWaitingLockFlag) && spinCount < spinLimit)
+			{
+				YieldInSpinWait();
+				++spinCount;
+				currentState = state.load(std::memory_order_relaxed);
+				continue;
+			}
+
+			// Store that there are waiting threads. Restart if the state has changed since it was loaded.
+			if (!(currentState & mayHaveWaitingLockFlag))
+			{
+				if (!state.compare_exchange_weak(currentState, currentState | mayHaveWaitingLockFlag, std::memory_order_relaxed))
+				{
+					continue;
+				}
+				currentState |= mayHaveWaitingLockFlag;
+			}
+
+			// Wait if the state has not changed. Either way, loop back and try to acquire the lock after trying to wait.
+			ParkingLot::Wait(GetLockAddress(), [this]
+			{
+				const u32 newState = state.load(std::memory_order_relaxed);
+				return (newState & (isLockedFlag | readLockCountMask)) && (newState & mayHaveWaitingLockFlag);
+			}, nullptr);
+			currentState = state.load(std::memory_order_relaxed);
+		}
+	}
+
+	UBA_NOINLINE void ReaderWriterLock::EnterReadSlow() const
+	{
+		constexpr u32 spinLimit = 40;
+		u32 spinCount = 0;
+		for (u32 currentState = state.load(std::memory_order_relaxed);;)
+		{
+			// Try to acquire the lock if it is unlocked and there are no waiting threads.
+			if (!(currentState & (isLockedFlag | mayHaveWaitingLockFlag)))
+			{
+				if (state.compare_exchange_weak(currentState, currentState + (1 << readLockCountShift), std::memory_order_acquire, std::memory_order_relaxed))
+				{
+					return;
+				}
+				continue;
+			}
+
+			// Spin up to the spin limit while there are no waiting threads.
+			if (!(currentState & mayHaveWaitingLockFlag) && spinCount < spinLimit)
+			{
+				YieldInSpinWait();
+				++spinCount;
+				currentState = state.load(std::memory_order_relaxed);
+				continue;
+			}
+
+			// Store that there are waiting threads. Restart if the state has changed since it was loaded.
+			if (!(currentState & mayHaveWaitingReadLockFlag))
+			{
+				if (!state.compare_exchange_weak(currentState, currentState | mayHaveWaitingReadLockFlag, std::memory_order_relaxed))
+				{
+					continue;
+				}
+				currentState |= mayHaveWaitingReadLockFlag;
+			}
+
+			// Wait if the state has not changed. Either way, loop back and try to acquire the lock after trying to wait.
+			ParkingLot::Wait(GetReadLockAddress(), [this, currentState]
+			{
+				return state.load(std::memory_order_relaxed) == currentState;
+			}, nullptr);
+			currentState = state.load(std::memory_order_relaxed);
+		}
+	}
+
+	UBA_NOINLINE void ReaderWriterLock::WakeWaitingThread() const
+	{
+		ParkingLot::WakeOne(GetLockAddress(), [this](ParkingLot::WakeState wakeState) -> u64
+		{
+			if (!wakeState.didWake)
+			{
+				// Keep the flag until no thread wakes, otherwise read locks may win before
+				// an exclusive lock has a chance.
+				state.fetch_and(~mayHaveWaitingLockFlag, std::memory_order_relaxed);
+			}
+			return 0;
+		});
+	}
+
+	UBA_NOINLINE void ReaderWriterLock::WakeWaitingThreads(u32 lastState)
+	{
+		if (lastState & mayHaveWaitingLockFlag)
+		{
+			// Wake one thread that is waiting to acquire an exclusive lock.
+			bool didWake = false;
+			ParkingLot::WakeOne(GetLockAddress(), [this, &didWake](ParkingLot::WakeState wakeState) -> u64
+			{
+				if (!wakeState.didWake)
+				{
+					// Keep the flag until no thread wakes, otherwise read locks may win before
+					// an exclusive lock has a chance.
+					state.fetch_and(~mayHaveWaitingLockFlag, std::memory_order_relaxed);
+				}
+				didWake = wakeState.didWake;
+				return 0;
+			});
+			if (didWake)
+			{
+				return;
+			}
+
+			// Reload the state if there were no read waiters because new
+			// ones may have registered themselves since lastState was read.
+			if (!(lastState & mayHaveWaitingReadLockFlag))
+			{
+				lastState = state.load(std::memory_order_relaxed);
+			}
+		}
+
+		if (lastState & mayHaveWaitingReadLockFlag)
+		{
+			// Wake every thread that is waiting to acquire a read lock.
+			// The awoken threads might race against other exclusive locks.
+			if (state.fetch_and(~mayHaveWaitingReadLockFlag, std::memory_order_relaxed) & mayHaveWaitingReadLockFlag)
+			{
+				ParkingLot::WakeAll(GetReadLockAddress());
+			}
+		}
+	}
+
+	#else // #if UBA_USE_PARKINGLOT
+
 	CriticalSection::CriticalSection(bool recursive)
 	{
 		#if PLATFORM_WINDOWS
@@ -49,7 +368,7 @@ namespace uba
 		#endif
 	}
 
-	void CriticalSection::Enter() const
+	void CriticalSection::Enter()
 	{
 		#if PLATFORM_WINDOWS
 		EnterCriticalSection((CRITICAL_SECTION*)&data);
@@ -59,7 +378,7 @@ namespace uba
 		#endif
 	}
 
-	void CriticalSection::Leave() const
+	void CriticalSection::Leave()
 	{
 		#if PLATFORM_WINDOWS
 		LeaveCriticalSection((CRITICAL_SECTION*)&data);
@@ -176,6 +495,8 @@ namespace uba
 		#endif
 	}
 
+	#endif // #if UBA_USE_PARKINGLOT / #else
+
 	#if UBA_TRACK_CONTENTION
 
 	List<ContentionTracker>& GetContentionTrackerList()
diff --git a/Engine/Source/Programs/UnrealBuildAccelerator/Core/Public/UbaBase.h b/Engine/Source/Programs/UnrealBuildAccelerator/Core/Public/UbaBase.h
--- a/Engine/Source/Programs/UnrealBuildAccelerator/Core/Public/UbaBase.h
+++ b/Engine/Source/Programs/UnrealBuildAccelerator/Core/Public/UbaBase.h
@@ -33,15 +33,29 @@ namespace uba
 	#define PERCENT_HS "%s"
 	#endif
 
+	#if (defined(_M_IX86) || defined(__i386__) || defined(_M_X64) || defined(__amd64__) || defined(__x86_64__)) && !defined(_M_ARM64EC)
+	#define PLATFORM_CPU_X86_FAMILY	1
+	#else
+	#define PLATFORM_CPU_X86_FAMILY	0
+	#endif
+
+	#if (defined(__arm__) || defined(_M_ARM) || defined(__aarch64__) || defined(_M_ARM64) || defined(_M_ARM64EC))
+	#define PLATFORM_CPU_ARM_FAMILY	1
+	#else
+	#define PLATFORM_CPU_ARM_FAMILY	0
+	#endif
+
 	#if PLATFORM_WINDOWS && defined(_M_ARM64) // For now we only care about windows arm..
 	inline constexpr bool IsArmBinary = true;
 	#else
 	inline constexpr bool IsArmBinary = false;
 	#endif
 
-#if PLATFORM_WINDOWS && defined(UBA_BUILD)
+	#if PLATFORM_WINDOWS && defined(UBA_BUILD)
 	#pragma warning(disable:4100) // This is needed because of single header compiles where AdditionalCompilerArguments is not included
 	#endif
+
+	inline constexpr u32 CacheLineSize = 64;
 }
 
 #if !defined(UBA_API)
@@ -53,3 +67,11 @@ namespace uba
 		#define UBA_API 
 	#endif
 #endif
+
+#if defined(_MSC_VER)
+	#define UBA_LIFETIMEBOUND [[msvc::lifetimebound]]
+#elif defined(__clang__)
+	#define UBA_LIFETIMEBOUND [[clang::lifetimebound]]
+#else
+	#define UBA_LIFETIMEBOUND
+#endif
diff --git a/Engine/Source/Programs/UnrealBuildAccelerator/Core/Public/UbaMemory.h b/Engine/Source/Programs/UnrealBuildAccelerator/Core/Public/UbaMemory.h
--- a/Engine/Source/Programs/UnrealBuildAccelerator/Core/Public/UbaMemory.h
+++ b/Engine/Source/Programs/UnrealBuildAccelerator/Core/Public/UbaMemory.h
@@ -3,14 +3,15 @@
 #pragma once
 
 #include "UbaSynchronization.h"
+#include <functional>
 #include <list>
 #include <map>
 #include <set>
 #include <string>
+#include <type_traits>
 #include <unordered_map>
 #include <unordered_set>
 #include <vector>
-#include <functional>
 
 #if UBA_USE_MIMALLOC
 #include <mimalloc.h>
@@ -343,4 +344,73 @@ namespace uba
 
 	bool SupportsHugePages();
 	u64 GetHugePageCount();
+
+	template <typename FunctionType>
+	class FunctionWithContext;
+
+	template <typename ReturnType, typename... ArgTypes>
+	class FunctionWithContext<ReturnType (ArgTypes...)>
+	{
+	public:
+		using FunctionType = ReturnType (void*, ArgTypes...);
+
+		template <typename CallableType>
+		inline FunctionWithContext(CallableType&& callable UBA_LIFETIMEBOUND)
+			requires (!std::is_same_v<std::decay_t<CallableType>, FunctionWithContext>) &&
+				std::is_invocable_r_v<ReturnType, std::decay_t<CallableType>, ArgTypes...>
+			: function(&Call<std::decay_t<CallableType>>)
+			, context(&callable)
+		{
+		}
+
+		template <typename CallableType>
+		inline FunctionWithContext& operator=(CallableType&& callable UBA_LIFETIMEBOUND)
+			requires (!std::is_same_v<std::decay_t<CallableType>, FunctionWithContext>) &&
+				std::is_invocable_r_v<ReturnType, std::decay_t<CallableType>, ArgTypes...>
+		{
+			function = &Call<std::decay_t<CallableType>>;
+			context = &callable;
+			return *this;
+		}
+
+		inline explicit FunctionWithContext(FunctionType* inFunction, void* inContext)
+			: function(inFunction)
+			, context(inContext)
+		{
+		}
+
+		inline constexpr FunctionWithContext(decltype(nullptr))
+		{
+		}
+
+		inline explicit operator bool() const
+		{
+			return !!function;
+		}
+
+		inline ReturnType operator()(ArgTypes... args) const
+		{
+			return function(context, (ArgTypes&&)args...);
+		}
+
+		inline FunctionType* GetFunction() const
+		{
+			return function;
+		}
+
+		inline void* GetContext() const
+		{
+			return context;
+		}
+
+	private:
+		template <typename CallableType>
+		static ReturnType Call(void* callable, ArgTypes... args)
+		{
+			return std::invoke(*(CallableType*)callable, (ArgTypes&&)args...);
+		}
+
+		FunctionType* function = nullptr;
+		void* context = nullptr;
+	};
 }
diff --git a/Engine/Source/Programs/UnrealBuildAccelerator/Core/Public/UbaParkingLot.h b/Engine/Source/Programs/UnrealBuildAccelerator/Core/Public/UbaParkingLot.h
new file mode 100644
--- /dev/null
+++ b/Engine/Source/Programs/UnrealBuildAccelerator/Core/Public/UbaParkingLot.h
@@ -0,0 +1,100 @@
+// Copyright Epic Games, Inc. All Rights Reserved.
+
+#pragma once
+
+#include "UbaBase.h"
+#include "UbaMemory.h"
+
+#include <atomic>
+#include <type_traits>
+
+namespace uba::ParkingLot
+{
+	struct WaitState final
+	{
+		/** Did this thread wait? True only if canWait returned true. */
+		bool didWait = false;
+		/** Did this wake from a wait? True only if a Wake call woke the thread, false for timeouts. */
+		bool didWake = false;
+		/** Optional value that was provided by the callback in WakeOne. */
+		u64 wakeToken = 0;
+	};
+
+	struct WakeState final
+	{
+		/** Did a thread wake up? */
+		bool didWake = false;
+		/** Does the queue MAYBE have another thread waiting on this address? */
+		bool hasWaitingThreads = false;
+	};
+
+	namespace Private
+	{
+		UBA_API WaitState Wait(const void* address, bool (*canWait)(void*), void* canWaitContext, void (*beforeWait)(void*), void* beforeWaitContext);
+		UBA_API WaitState WaitFor(const void* address, bool (*canWait)(void*), void* canWaitContext, void (*beforeWait)(void*), void* beforeWaitContext, u32 waitTimeMs);
+		//UBA_API WaitState WaitUntil(const void* address, bool (*canWait)(void*), void* canWaitContext, void (*beforeWait)(void*), void* beforeWaitContext, FMonotonicTimePoint waitTimeMs);
+		UBA_API void WakeOne(const void* address, u64 (*onWakeState)(void*, WakeState), void* onWakeStateContext);
+	} // Private
+
+	/** Queue the calling thread to wait if canWait returns true. beforeWait is only called if canWait returns true. */
+	inline WaitState Wait(const void* address, FunctionWithContext<bool ()> canWait, FunctionWithContext<void ()> beforeWait)
+	{
+		return Private::Wait(address, canWait.GetFunction(), canWait.GetContext(), beforeWait.GetFunction(), beforeWait.GetContext());
+	}
+
+	/** Queue the calling thread to wait if canWait returns true. beforeWait is only called if canWait returns true. */
+	inline WaitState WaitFor(const void* address, FunctionWithContext<bool ()> canWait, FunctionWithContext<void ()> beforeWait, u32 waitTimeMs)
+	{
+		return Private::WaitFor(address, canWait.GetFunction(), canWait.GetContext(), beforeWait.GetFunction(), beforeWait.GetContext(), waitTimeMs);
+	}
+
+	///** Queue the calling thread to wait if canWait returns true. beforeWait is only called if canWait returns true. */
+	//inline WaitState WaitUntil(const void* address, FunctionWithContext<bool ()> canWait, FunctionWithContext<void ()> beforeWait, FMonotonicTimePoint waitTime)
+	//{
+	//	return Private::WaitUntil(address, canWait.GetFunction(), canWait.GetContext(), beforeWait.GetFunction(), beforeWait.GetContext(), waitTime);
+	//}
+
+	/** Wake one thread from the queue of threads waiting on the address. */
+	inline void WakeOne(const void* address, FunctionWithContext<u64 (WakeState)> onWakeState)
+	{
+		return Private::WakeOne(address, onWakeState.GetFunction(), onWakeState.GetContext());
+	}
+
+	/** Wake one thread from the queue of threads waiting on the address. */
+	UBA_API WakeState WakeOne(const void* address);
+
+	/** Wake up to wakeCount threads from the queue of threads waiting on the address. */
+	UBA_API u32 WakeMultiple(const void* address, u32 wakeCount);
+
+	/** Wake all threads from the queue of threads waiting on the address. */
+	UBA_API u32 WakeAll(const void* address);
+
+	/**
+	 * A mutex that is the size of a pointer and does not depend on ParkingLot.
+	 *
+	 * Prefer CriticalSection or ReaderWriterLock to WordMutex whenever possible.
+	 * This mutex is not fair and does not support recursive locking.
+	 */
+	class WordMutex final
+	{
+	public:
+		constexpr WordMutex() = default;
+
+		WordMutex(const WordMutex&) = delete;
+		WordMutex& operator=(const WordMutex&) = delete;
+
+		bool TryEnter();
+		void Enter();
+		void Leave();
+
+	private:
+		void EnterSlow();
+		void LeaveSlow(u64 currentState);
+
+		static constexpr u64 isLockedFlag = 1 << 0;
+		static constexpr u64 isQueueLockedFlag = 1 << 1;
+		static constexpr u64 queueMask = ~(isLockedFlag | isQueueLockedFlag);
+
+		std::atomic<u64> state = 0;
+	};
+} // uba::ParkingLot
diff --git a/Engine/Source/Programs/UnrealBuildAccelerator/Core/Public/UbaPlatform.h b/Engine/Source/Programs/UnrealBuildAccelerator/Core/Public/UbaPlatform.h
--- a/Engine/Source/Programs/UnrealBuildAccelerator/Core/Public/UbaPlatform.h
+++ b/Engine/Source/Programs/UnrealBuildAccelerator/Core/Public/UbaPlatform.h
@@ -50,6 +50,10 @@
 	#define ANALYSIS_NORETURN // __attribute__((analyzer_noreturn))
 #endif
 
+#if PLATFORM_CPU_X86_FAMILY
+	#include <immintrin.h>
+#endif
+
 #define UBA_EXPERIMENTAL 0
 
 #if UBA_DEBUG && defined(_MSC_VER)
@@ -99,7 +103,24 @@ namespace uba
 	bool IsRunningDarling();
 	bool IsRunningArm();
 	bool IsEscapePressed();
+
 	void Sleep(u32 milliseconds);
+
+	inline void YieldInSpinWait()
+	{
+	#if PLATFORM_CPU_X86_FAMILY
+		_mm_pause();
+	#elif PLATFORM_CPU_ARM_FAMILY
+		#if !defined(__clang__)
+		__yield();
+		#else
+		__builtin_arm_yield();
+		#endif
+	#else
+		#error YieldInSpinWait() is not implemented on this architecture.
+	#endif
+	}
+
 	u32 GetUserDefaultUILanguage();
 	u32 GetLastError();
 	void SetLastError(u32 error);
diff --git a/Engine/Source/Programs/UnrealBuildAccelerator/Core/Public/UbaSynchronization.h b/Engine/Source/Programs/UnrealBuildAccelerator/Core/Public/UbaSynchronization.h
--- a/Engine/Source/Programs/UnrealBuildAccelerator/Core/Public/UbaSynchronization.h
+++ b/Engine/Source/Programs/UnrealBuildAccelerator/Core/Public/UbaSynchronization.h
@@ -7,6 +7,7 @@
 #include <utility>
 
 #define UBA_TRACK_CONTENTION 0
+#define UBA_USE_PARKINGLOT 0
 
 #define STRING_JOIN(arg1, arg2) STRING_JOIN_INNER(arg1, arg2)
 #define STRING_JOIN_INNER(arg1, arg2) arg1 ## arg2
@@ -31,14 +32,39 @@ namespace uba
 
 	////////////////////////////////////////////////////////////////////////////////////////////////////
 
+	#if UBA_USE_PARKINGLOT
+	class CriticalSection
+	{
+	public:
+		CriticalSection() = default;
+
+		[[nodiscard]] bool TryEnter();
+		void Enter();
+		void Leave();
+
+	private:
+		void EnterSlow(u32 currentState, u32 currentThreadId);
+		void WakeWaitingThread();
+
+		static constexpr u32 mayHaveWaitingLockFlag = 1 << 0;
+		static constexpr u32 lockCountShift = 1;
+		static constexpr u32 lockCountMask = 0xffff'fffe;
+
+		Atomic<u32> state = 0;
+		Atomic<u32> threadId = 0;
+
+		CriticalSection(const CriticalSection&) = delete;
+		CriticalSection& operator=(const CriticalSection&) = delete;
+	};
+	#else
 	class CriticalSection
 	{
 	public:
 		CriticalSection(bool recursive = true);
 		~CriticalSection();
 
-		void Enter() const;
-		void Leave() const;
+		void Enter();
+		void Leave();
 
 		bool TryEnter() const;
 
@@ -52,9 +78,90 @@ namespace uba
 		CriticalSection(const CriticalSection&) = delete;
 		CriticalSection& operator=(const CriticalSection&) = delete;
 	};
+	#endif
 
 	////////////////////////////////////////////////////////////////////////////////////////////////////
 
+	#if UBA_USE_PARKINGLOT
+	class ReaderWriterLock
+	{
+	public:
+		ReaderWriterLock() = default;
+
+		inline bool TryEnter()
+		{
+			u32 expected = state.load(std::memory_order_relaxed);
+			return !(expected & (isLockedFlag | readLockCountMask)) &&
+				state.compare_exchange_strong(expected, expected | isLockedFlag,
+					std::memory_order_acquire, std::memory_order_relaxed);
+		}
+
+		inline void Enter()
+		{
+			u32 expected = 0;
+			if (state.compare_exchange_weak(expected, isLockedFlag, std::memory_order_acquire, std::memory_order_relaxed))
+			{
+				return;
+			}
+			EnterSlow();
+		}
+
+		inline void Leave()
+		{
+			// Unlock immediately to allow other threads to acquire the lock while this thread looks for a thread to wake.
+			u32 lastState = state.fetch_sub(isLockedFlag, std::memory_order_release);
+			if (!(lastState & (mayHaveWaitingLockFlag | mayHaveWaitingReadLockFlag)))
+			{
+				return;
+			}
+			WakeWaitingThreads(lastState);
+		}
+
+		inline void EnterRead() const
+		{
+			u32 expected = state.load(std::memory_order_relaxed);
+			if (!(expected & (isLockedFlag | mayHaveWaitingLockFlag)) &&
+				state.compare_exchange_weak(expected, expected + (1 << readLockCountShift),
+					std::memory_order_acquire, std::memory_order_relaxed))
+			{
+				return;
+			}
+			EnterReadSlow();
+		}
+
+		inline void LeaveRead() const
+		{
+			// Unlock immediately to allow other threads to acquire the lock while this thread looks for a thread to wake.
+			const u32 lastState = state.fetch_sub(1 << readLockCountShift, std::memory_order_release);
+			constexpr u32 wakeState = mayHaveWaitingLockFlag | (1 << readLockCountShift);
+			if ((lastState & ~mayHaveWaitingReadLockFlag) != wakeState)
+			{
+				return;
+			}
+			WakeWaitingThread();
+		}
+
+	private:
+		void EnterSlow();
+		void EnterReadSlow() const;
+		void WakeWaitingThread() const;
+		void WakeWaitingThreads(u32 lastState);
+
+		const void* GetLockAddress() const;
+		const void* GetReadLockAddress() const;
+
+		static constexpr u32 isLockedFlag = 1 << 0;
+		static constexpr u32 mayHaveWaitingLockFlag = 1 << 1;
+		static constexpr u32 mayHaveWaitingReadLockFlag = 1 << 2;
+		static constexpr u32 readLockCountShift = 3;
+		static constexpr u32 readLockCountMask = 0xffff'fff8;
+
+		mutable Atomic<u32> state = 0;
+
+		ReaderWriterLock(const ReaderWriterLock&) = delete;
+		ReaderWriterLock& operator=(const ReaderWriterLock&) = delete;
+	};
+	#else
 	class ReaderWriterLock
 	{
 	public:
@@ -82,6 +189,7 @@ namespace uba
 		ReaderWriterLock(const ReaderWriterLock&) = delete;
 		ReaderWriterLock& operator=(const ReaderWriterLock&) = delete;
 	};
+	#endif
 
 	////////////////////////////////////////////////////////////////////////////////////////////////////
 
@@ -111,15 +219,16 @@ namespace uba
 
 	////////////////////////////////////////////////////////////////////////////////////////////////////
 
-	class ScopedCriticalSection 
+	template <typename MutexType>
+	class ScopedCriticalSection
 	{
 	public:
-		ScopedCriticalSection(const CriticalSection& cs) : m_cs(cs), m_active(true) { cs.Enter(); }
+		ScopedCriticalSection(MutexType& mutex) : m_mutex(mutex), m_active(true) { mutex.Enter(); }
 		~ScopedCriticalSection() { Leave(); }
-		void Enter() { if (m_active) return; m_cs.Enter(); m_active = true; }
-		void Leave() { if (!m_active) return; m_cs.Leave(); m_active = false; }
+		void Enter() { if (m_active) return; m_mutex.Enter(); m_active = true; }
+		void Leave() { if (!m_active) return; m_mutex.Leave(); m_active = false; }
 	private:
-		const CriticalSection& m_cs;
+		MutexType& m_mutex;
 		bool m_active;
 	};
 
@@ -153,14 +262,14 @@ namespace uba
 
 	////////////////////////////////////////////////////////////////////////////////////////////////////
 
-	#if PLATFORM_WINDOWS
+	#if UBA_USE_PARKINGLOT || PLATFORM_WINDOWS
 	using Futex = ReaderWriterLock;
 	using ScopedFutex = ScopedWriteLock;
 	using ScopedFutexRead = ScopedReadLock;
 	#else
 	class Futex : public CriticalSection { public: Futex() : CriticalSection(false) {} };
-	using ScopedFutex = ScopedCriticalSection;
-	using ScopedFutexRead = ScopedCriticalSection;
+	using ScopedFutex = ScopedCriticalSection<Futex>;
+	using ScopedFutexRead = ScopedCriticalSection<Futex>;
 	#endif
 
 
