[UBA]
* Fixed so FatalError/Rpc_WriteLogf both use StringBuffer to prevent duplicated code and fix crash
* Fixed conflicting flags when creating child processes on win32.
* Removed git.exe special case since it now works remotely as well
* Bumped .lib reserve size for link.exe
* Improved assert message

diff --git a/Engine/Source/Programs/UnrealBuildAccelerator/Common/Private/UbaStorage.cpp b/Engine/Source/Programs/UnrealBuildAccelerator/Common/Private/UbaStorage.cpp
--- a/Engine/Source/Programs/UnrealBuildAccelerator/Common/Private/UbaStorage.cpp
+++ b/Engine/Source/Programs/UnrealBuildAccelerator/Common/Private/UbaStorage.cpp
@@ -2795,10 +2795,8 @@ namespace uba
 		auto findIt = m_deferredCasCreationLookupByName.find(fileNameKey);
 		if (findIt == m_deferredCasCreationLookupByName.end())
 			return;
-
-		m_logger.Logf(IsDebug ? LogEntryType_Warning : LogEntryType_Info, TC("Flush cas creation of file %s (This should be rare)"), fileName);
-
 		CasKey key = findIt->second.key;
+		m_logger.Logf(IsDebug ? LogEntryType_Warning : LogEntryType_Info, TC("Flush cas creation of file %s at %s (This should be rare)"), CasKeyString(key).str, fileName);
 		deferredLock.Leave();
 		HasCasFile(key);
 	}
diff --git a/Engine/Source/Programs/UnrealBuildAccelerator/Core/Private/UbaApplicationRules.cpp b/Engine/Source/Programs/UnrealBuildAccelerator/Core/Private/UbaApplicationRules.cpp
--- a/Engine/Source/Programs/UnrealBuildAccelerator/Core/Private/UbaApplicationRules.cpp
+++ b/Engine/Source/Programs/UnrealBuildAccelerator/Core/Private/UbaApplicationRules.cpp
@@ -193,8 +193,8 @@ namespace uba
 		{
 			if (file.Contains(TCV("lnk{")))
 				return 32 * 1024 * 1024;
-			if (file.EndsWith(TCV(".lib"))) // This is import lib which is tiny but not sure if we output any full libs
-				return 512 * 1024 * 1024;
+			if (file.EndsWith(TCV(".lib"))) // This is usually for import lib which is tiny but not sure if we output any full libs
+				return 2ull * 1024 * 1024 * 1024;
 			if (file.EndsWith(TCV(".dll")) || file.EndsWith(TCV(".exe")))
 				return 4ull * 1024 * 1024 * 1024;
 			return Super::FileTypeMaxSize(file, isSystemOrTempFile);
diff --git a/Engine/Source/Programs/UnrealBuildAccelerator/Detours/Private/UbaDetoursShared.cpp b/Engine/Source/Programs/UnrealBuildAccelerator/Detours/Private/UbaDetoursShared.cpp
--- a/Engine/Source/Programs/UnrealBuildAccelerator/Detours/Private/UbaDetoursShared.cpp
+++ b/Engine/Source/Programs/UnrealBuildAccelerator/Detours/Private/UbaDetoursShared.cpp
@@ -392,14 +392,12 @@ namespace uba
 
 	ANALYSIS_NORETURN void FatalError(u32 code, const tchar* format, ...)
 	{
+		StringBuffer<2048> sb;
+		sb.Append(GetApplicationShortName()).Append(TCV(" ERROR: "));
 		va_list arg;
 		va_start(arg, format);
-		tchar buffer[1024];
-		if (Tvsprintf_s(buffer, sizeof_array(buffer), format, arg) <= 0)
-			TStrcpy_s(buffer, sizeof_array(buffer), format);
+		sb.Append(format, arg);
 		va_end(arg);
-		StringBuffer<2048> sb;
-		sb.Append(GetApplicationShortName()).Append(TCV(" ERROR: ")).Append(buffer);
 		Rpc_WriteLog(sb.data, sb.count, true, true);
 
 		#if PLATFORM_WINDOWS // Maybe all platforms should call exit()?
@@ -424,15 +422,10 @@ namespace uba
 	{
 		va_list arg;
 		va_start(arg, format);
-		tchar buffer[1024];
-		int count = Tvsprintf_s(buffer, 1024, format, arg);
-		if (count < 0)
-		{
-			TStrcpy_s(buffer, 1024, format);
-			count = int(TStrlen(buffer));
-		}
+		StringBuffer<2048> buffer;
+		buffer.Append(format, arg);
 		va_end(arg);
-		Rpc_WriteLog(buffer, u32(count), false, false);
+		Rpc_WriteLog(buffer.data, buffer.count, false, false);
 	}
 
 	UBA_NOINLINE void Rpc_ResolveCallstack(StringBufferBase& out, u32 skipCallstackCount, void* context)
diff --git a/Engine/Source/Programs/UnrealBuildAccelerator/Detours/Private/Windows/UbaDetoursFunctionsKernelBase.inl b/Engine/Source/Programs/UnrealBuildAccelerator/Detours/Private/Windows/UbaDetoursFunctionsKernelBase.inl
--- a/Engine/Source/Programs/UnrealBuildAccelerator/Detours/Private/Windows/UbaDetoursFunctionsKernelBase.inl
+++ b/Engine/Source/Programs/UnrealBuildAccelerator/Detours/Private/Windows/UbaDetoursFunctionsKernelBase.inl
@@ -619,6 +619,7 @@ LPVOID Detoured_VirtualAlloc(LPVOID lpAddress, SIZE_T dwSize, DWORD flAllocation
 BOOL Detoured_GetQueuedCompletionStatusEx(HANDLE CompletionPort, LPOVERLAPPED_ENTRY lpCompletionPortEntries, ULONG ulCount, PULONG ulNumEntriesRemoved, DWORD dwMilliseconds, BOOL fAlertable)
 {
 	DETOURED_CALL(GetQueuedCompletionStatusEx);
+	UBA_ASSERT(!isDetouredHandle(CompletionPort));
 	DEBUG_LOG_TRUE(L"GetQueuedCompletionStatusEx", L"%llu (Timeout: %ums)", u64(CompletionPort), dwMilliseconds);
 	bool res = True_GetQueuedCompletionStatusEx(CompletionPort, lpCompletionPortEntries, ulCount, ulNumEntriesRemoved, dwMilliseconds, fAlertable);
 	if (res)
@@ -1010,7 +1011,7 @@ BOOL Detoured_SetEndOfFile(HANDLE hFile)
 		FileInfo& fi = *dh.fileObject->fileInfo;
 		if (MemoryFile* mf = fi.memoryFile)
 		{
-			DEBUG_LOG_DETOURED(L"SetEndOfFile (MEMORY)", L"%llu (%ls) -> Success", uintptr_t(hFile), HandleToName(hFile));
+			DEBUG_LOG_DETOURED(L"SetEndOfFile (MEMORY)", L"%llu NewSize: %llu (%ls) -> Success", uintptr_t(hFile), mf->writtenSize, HandleToName(hFile));
 			mf->writtenSize = dh.pos;
 			mf->isReported = false;
 			mf->EnsureCommitted(dh, mf->writtenSize);
@@ -3204,7 +3205,7 @@ BOOL Detoured_CreateProcessW(LPCWSTR lpApplicationName, LPWSTR lpCommandLine, LP
 
 	bool isChild = true;
 	// We don't care about tracking mspdbsrv or vctip.. they are services just spawned by this process
-	if (testName.Contains(L"mspdbsrv.exe") || testName.Contains(L"vctip.exe") || testName.Contains(L"git.exe"))
+	if (testName.Contains(L"mspdbsrv.exe") || testName.Contains(L"vctip.exe"))
 	{
 		if (!g_runningRemote)
 		{
@@ -3386,13 +3387,15 @@ BOOL Detoured_CreateProcessW(LPCWSTR lpApplicationName, LPWSTR lpCommandLine, LP
 	//	(void)siex;
 	//}
 
-	lpStartupInfo->dwFlags |= STARTF_USESHOWWINDOW;
-	lpStartupInfo->wShowWindow = SW_HIDE;
-
 	if (g_rules->AllowDetach())
 		dwCreationFlags |= DETACHED_PROCESS;
 	else if (!g_setStdHandleCalled) // If console mode has been called there are probably some sort of stdin/out rerouting stuff.. so let's not set CREATE_NO_WINDOW
 		dwCreationFlags |= CREATE_NO_WINDOW;
+	else
+	{
+		lpStartupInfo->dwFlags |= STARTF_USESHOWWINDOW;
+		lpStartupInfo->wShowWindow = SW_HIDE;
+	}
 
 	dwCreationFlags |= CREATE_SUSPENDED;
 	BOOL res = true;
@@ -4178,6 +4181,7 @@ BOOL Detoured_SetHandleInformation(HANDLE hObject, DWORD dwMask, DWORD dwFlags)
 {
 	DETOURED_CALL(SetHandleInformation);
 	DEBUG_LOG_TRUE(L"SetHandleInformation", L"%llu", uintptr_t(hObject));
+	UBA_ASSERT(isDetouredHandle(hObject));
 	return True_SetHandleInformation(hObject, dwMask, dwFlags); // Calls NtQueryObject and NtSetInformationObject internally
 }
 
@@ -4507,7 +4511,13 @@ BOOL Detoured_FlushConsoleInputBuffer(HANDLE hConsoleInput)
 {
 	DETOURED_CALL(FlushConsoleInputBuffer);
 	DEBUG_LOG_DETOURED(L"FlushConsoleInputBuffer", L"");
-	return True_FlushConsoleInputBuffer(hConsoleInput);
+	HANDLE trueHandle = hConsoleInput;
+	if (isDetouredHandle(hConsoleInput))
+	{
+		trueHandle = asDetouredHandle(hConsoleInput).trueHandle;
+		UBA_ASSERT(trueHandle != INVALID_HANDLE_VALUE);
+	}
+	return True_FlushConsoleInputBuffer(trueHandle);
 }
 
 BOOL Detoured_SetConsoleTextAttribute(HANDLE hConsoleOutput, WORD wAttributes)
