[UBA]
* Fixed so parkinglot code is not compiled at all if disabled

diff --git a/Engine/Source/Programs/UnrealBuildAccelerator/Core/Private/UbaParkingLot.cpp b/Engine/Source/Programs/UnrealBuildAccelerator/Core/Private/UbaParkingLot.cpp
--- a/Engine/Source/Programs/UnrealBuildAccelerator/Core/Private/UbaParkingLot.cpp
+++ b/Engine/Source/Programs/UnrealBuildAccelerator/Core/Private/UbaParkingLot.cpp
@@ -1,6 +1,9 @@
 // Copyright Epic Games, Inc. All Rights Reserved.
 
 #include "UbaParkingLot.h"
+
+#if UBA_USE_PARKINGLOT
+
 #include "UbaPlatform.h"
 
 #include "UbaStringBuffer.h" // for UBA_ASSERTF
@@ -1309,3 +1312,5 @@ namespace uba::ParkingLot
 		return WakeMultiple(address, 0xffffffff);
 	}
 } // uba::ParkingLot
+
+#endif // UBA_USE_PARKINGLOT
\ No newline at end of file
diff --git a/Engine/Source/Programs/UnrealBuildAccelerator/Core/Public/UbaParkingLot.h b/Engine/Source/Programs/UnrealBuildAccelerator/Core/Public/UbaParkingLot.h
--- a/Engine/Source/Programs/UnrealBuildAccelerator/Core/Public/UbaParkingLot.h
+++ b/Engine/Source/Programs/UnrealBuildAccelerator/Core/Public/UbaParkingLot.h
@@ -5,6 +5,8 @@
 #include "UbaBase.h"
 #include "UbaMemory.h"
 
+#if UBA_USE_PARKINGLOT
+
 #include <atomic>
 #include <type_traits>
 
@@ -98,3 +100,5 @@ namespace uba::ParkingLot
 		std::atomic<u64> state = 0;
 	};
 } // uba::ParkingLot
+
+#endif // UBA_USE_PARKINGLOT
\ No newline at end of file
