[UBA]
* Fixed bug where cmd.exe was sometimes using host version when we always want to use remote version
* Fixed Scheduler::GetProcessCountThatCanRunRemotelyNow taking dependencies into account
* Added assert for checking so no application path use windows short path (~)
* Improved some entries in detouring log

diff --git a/Engine/Source/Programs/UnrealBuildAccelerator/Common/Private/UbaScheduler.cpp b/Engine/Source/Programs/UnrealBuildAccelerator/Common/Private/UbaScheduler.cpp
--- a/Engine/Source/Programs/UnrealBuildAccelerator/Common/Private/UbaScheduler.cpp
+++ b/Engine/Source/Programs/UnrealBuildAccelerator/Common/Private/UbaScheduler.cpp
@@ -480,7 +480,16 @@ namespace uba
 				continue;
 			if (entry.runStatus != RunStatus_QueuedForRun || entry.cacheStatus != CacheStatus_Failed)
 				continue;
-			++count;
+
+			while (entry.dependencyCount)
+			{
+				auto depIndex = *entry.dependencies;
+				if (m_processEntries[depIndex].runStatus != RunStatus_Success)
+					break;
+				*entry.dependencies = entry.dependencies[--entry.dependencyCount];
+			}
+			if (!entry.dependencyCount)
+				++count;
 		}
 
 		count += m_activeRemoteProcesses;
@@ -508,16 +517,15 @@ namespace uba
 				continue;
 			if (entry.runStatus != RunStatus_QueuedForRun || entry.cacheStatus != CacheStatus_Failed)
 				continue;
-			bool canRun = true;
-			for (u32 j=0, je=entry.dependencyCount; j!=je; ++j)
+
+			while (entry.dependencyCount)
 			{
-				auto depIndex = entry.dependencies[j];
-				if (m_processEntries[depIndex].runStatus == RunStatus_Success)
-					continue;
-				canRun = false;
-				break;
+				auto depIndex = *entry.dependencies;
+				if (m_processEntries[depIndex].runStatus != RunStatus_Success)
+					break;
+				*entry.dependencies = entry.dependencies[--entry.dependencyCount];
 			}
-			if (!canRun)
+			if (entry.dependencyCount)
 				continue;
 			float weight = entry.info->weight;
 			totalWeight += weight;
diff --git a/Engine/Source/Programs/UnrealBuildAccelerator/Common/Private/UbaSession.cpp b/Engine/Source/Programs/UnrealBuildAccelerator/Common/Private/UbaSession.cpp
--- a/Engine/Source/Programs/UnrealBuildAccelerator/Common/Private/UbaSession.cpp
+++ b/Engine/Source/Programs/UnrealBuildAccelerator/Common/Private/UbaSession.cpp
@@ -2197,6 +2197,7 @@ namespace uba
 	{
 		UBA_ASSERTF(startInfo.workingDir && *startInfo.workingDir, TC("Working dir must be set when spawning process"));
 		UBA_ASSERTF(!TStrchr(startInfo.workingDir, '~'), TC("WorkingDir path must use long name (%s)"), startInfo.workingDir);
+		UBA_ASSERTF(!IsWindows || !TStrchr(startInfo.application, '~'), TC("Application path must use long name (%s)"), startInfo.application);
 	}
 
 	ProcessHandle Session::InternalRunProcess(const ProcessStartInfo& startInfo, bool async, ProcessImpl* parent, bool enableDetour, void* userData)
diff --git a/Engine/Source/Programs/UnrealBuildAccelerator/Common/Private/UbaSessionClient.cpp b/Engine/Source/Programs/UnrealBuildAccelerator/Common/Private/UbaSessionClient.cpp
--- a/Engine/Source/Programs/UnrealBuildAccelerator/Common/Private/UbaSessionClient.cpp
+++ b/Engine/Source/Programs/UnrealBuildAccelerator/Common/Private/UbaSessionClient.cpp
@@ -377,21 +377,26 @@ namespace uba
 	{
 		ProcessStartInfoHolder& startInfo = process.m_startInfo;
 		outRealWorkingDir = m_processWorkingDir.data;
-		if (StartsWith(startInfo.application, TC("ubacopy")))
+
+		const tchar* application = startInfo.application;
+		StringView applicationView = ToView(application);
+		
+		if (applicationView.StartsWith(TC("ubacopy")))
 			return true;
 
 		#if PLATFORM_LINUX
-		//if (EndsWith(startInfo.application, TStrlen(startInfo.application), "/sh")) // For some reason the local shell application on the helpers hang (need to investigate aws shell application)
+		//if (applicationView.EndsWith(TStrlen(startInfo.application), "/sh")) // For some reason the local shell application on the helpers hang (need to investigate aws shell application)
 		//	return true;
 		#elif PLATFORM_WINDOWS
-		if (ToView(startInfo.application).EndsWith(TCV("system32\\cmd.exe"))) // We want to use local cmd.exe.. wine does not support kernel calls called by windows 11's cmd.exe for example
+		if (applicationView.Equals(TCV("cmd.exe")) || applicationView.EndsWith(TCV("\\cmd.exe"))) // We want to use local cmd.exe.. wine does not support kernel calls called by windows 11's cmd.exe for example
+		{
+			outRealApplication.Clear().Append(m_systemPath).Append(TCV("\\System32\\cmd.exe"));
 			return true;
+		}
 		#endif
 
 		outRealApplication.Clear();
 
-		const tchar* application = startInfo.application;
-		UBA_ASSERT(application && *application);
 		bool isAbsolute = IsAbsolutePath(startInfo.application);
 
 		SCOPED_FUTEX(m_handledApplicationEnvironmentsLock, environmentslock);
diff --git a/Engine/Source/Programs/UnrealBuildAccelerator/Detours/Private/Windows/UbaDetoursFunctionsKernelBase.inl b/Engine/Source/Programs/UnrealBuildAccelerator/Detours/Private/Windows/UbaDetoursFunctionsKernelBase.inl
--- a/Engine/Source/Programs/UnrealBuildAccelerator/Detours/Private/Windows/UbaDetoursFunctionsKernelBase.inl
+++ b/Engine/Source/Programs/UnrealBuildAccelerator/Detours/Private/Windows/UbaDetoursFunctionsKernelBase.inl
@@ -3112,7 +3112,7 @@ BOOL Detoured_GetConsoleMode(HANDLE hConsoleHandle, LPDWORD lpMode)
 	}
 
 	auto res = True_GetConsoleMode(hConsoleHandle, lpMode);
-	DEBUG_LOG_TRUE(L"GetConsoleMode", L"%llu %u-> %ls", uintptr_t(hConsoleHandle), *lpMode, ToString(res));
+	DEBUG_LOG_TRUE(L"GetConsoleMode", L"%llu %u -> %ls", uintptr_t(hConsoleHandle), *lpMode, ToString(res));
 	return res;
 }
 
@@ -3160,6 +3160,7 @@ BOOL Detoured_GetConsoleTitleW(LPTSTR lpConsoleTitle, DWORD nSize)
 {
 	DETOURED_CALL(GetConsoleTitleW);
 	DEBUG_LOG_DETOURED(L"GetConsoleTitleW", L"");
+	SetLastError(ERROR_SUCCESS);
 	*lpConsoleTitle = 0;
 	return true;
 }
@@ -3387,9 +3388,14 @@ BOOL Detoured_CreateProcessW(LPCWSTR lpApplicationName, LPWSTR lpCommandLine, LP
 	//	(void)siex;
 	//}
 
+	// CREATE_NO_WINDOW is an optimization that prevents additional overhead per spawned child process.
+	// .. disabled for now because there are hierarchies of processes that does not like this. Haven't figured out exactly why
+	// .. can be reproed by running some of the very complex cmake/python scripts when compiling clang (lto.exports etc)
+	bool createNoWindow = false; //!g_setStdHandleCalled; // If console mode has been called there are probably some sort of stdin/out rerouting stuff.. so let's not set CREATE_NO_WINDOW
+
 	if (g_rules->AllowDetach())
 		dwCreationFlags |= DETACHED_PROCESS;
-	else if (!g_setStdHandleCalled) // If console mode has been called there are probably some sort of stdin/out rerouting stuff.. so let's not set CREATE_NO_WINDOW
+	else if (createNoWindow)
 		dwCreationFlags |= CREATE_NO_WINDOW;
 	else
 	{
@@ -4440,7 +4446,7 @@ BOOL WINAPI Detoured_SetConsoleCtrlHandler(PHANDLER_ROUTINE HandlerRoutine, BOOL
 UINT Detoured_GetConsoleOutputCP()
 {
 	DETOURED_CALL(GetConsoleOutputCP);
-	DEBUG_LOG_DETOURED(L"GetConsoleOutputCP", L"");
+	DEBUG_LOG_DETOURED(L"GetConsoleOutputCP", L"-> 437");
 	return 437;
 	//auto res = True_GetConsoleOutputCP();
 	//return res;
@@ -4472,9 +4478,9 @@ BOOL Detoured_SetConsoleCursorPosition(HANDLE hConsoleOutput, COORD dwCursorPosi
 BOOL Detoured_GetConsoleScreenBufferInfo(HANDLE hConsoleOutput, PCONSOLE_SCREEN_BUFFER_INFO lpConsoleScreenBufferInfo)
 {
 	DETOURED_CALL(GetConsoleScreenBufferInfo);
-	DEBUG_LOG_DETOURED(L"GetConsoleScreenBufferInfo", L"");
-
-	return True_GetConsoleScreenBufferInfo(hConsoleOutput, lpConsoleScreenBufferInfo);
+	bool res = True_GetConsoleScreenBufferInfo(hConsoleOutput, lpConsoleScreenBufferInfo);
+	DEBUG_LOG_DETOURED(L"GetConsoleScreenBufferInfo", L"(%llu) -> %s", u64(hConsoleOutput), ToString(res));
+	return res;
 	/*
 	// We make these up just to make clang output errors with proper line breaks..
 	lpConsoleScreenBufferInfo->dwSize = COORD{ 120,200 };
