UBA: Finished the implementation of PlatformManualResetEvent

#rnx
#Synced-CL 50967057

diff --git a/Engine/Source/Programs/UnrealBuildAccelerator/Core/Private/UbaParkingLot.cpp b/Engine/Source/Programs/UnrealBuildAccelerator/Core/Private/UbaParkingLot.cpp
--- a/Engine/Source/Programs/UnrealBuildAccelerator/Core/Private/UbaParkingLot.cpp
+++ b/Engine/Source/Programs/UnrealBuildAccelerator/Core/Private/UbaParkingLot.cpp
@@ -5,6 +5,7 @@
 #if UBA_USE_PARKINGLOT
 
 #include "UbaPlatform.h"
+#include "UbaTimer.h"
 
 #include "UbaStringBuffer.h" // for UBA_ASSERTF
 
@@ -18,20 +19,32 @@
 #include <unistd.h>
 #endif
 
+#if !(PLATFORM_WINDOWS || PLATFORM_LINUX)
+#include <chrono>
+#include <condition_variable>
+#include <mutex>
+#endif
+
 namespace uba::ParkingLot::Private
 {
 #if PLATFORM_LINUX
-	inline static const struct timespec* MsToTimeSpec(u32 timeMs, struct timespec& outTime)
+	template<typename TimeType>
+	inline static const struct timespec* MsToTimeSpec(TimeType timeMs, struct timespec& outTime)
 	{
 		if (timeMs != 0xffffffff)
 		{
-			u32 sec = timeMs / 1'000;
+			TimeType sec = timeMs / 1'000;
 			outTime.tv_sec = sec;
 			outTime.tv_nsec = (timeMs - sec * 1'000) * 1'000'000;
 			return &outTime;
 		}
 		return nullptr;
 	}
+
+	inline static u64 TimeSpecToMs(const struct timespec& time)
+	{
+		return u64(time.tv_sec) * 1'000 + u64(time.tv_nsec) / 1'000'000;
+	}
 #endif
 
 	class PlatformManualResetEvent
@@ -43,62 +56,89 @@ namespace uba::ParkingLot::Private
 
 		inline void Reset()
 		{
+			#if PLATFORM_WINDOWS || PLATFORM_LINUX
 			state.store(0, std::memory_order_relaxed);
+			#else
+			std::unique_lock selfLock(lock);
+			shouldWait = true;
+			#endif
 		}
 
 		inline bool Poll()
 		{
+			#if PLATFORM_WINDOWS || PLATFORM_LINUX
 			return !!state.load(std::memory_order_acquire);
+			#else
+			std::unique_lock selfLock(lock);
+			return !shouldWait;
+			#endif
 		}
 
 		inline void Wait()
 		{
+			#if PLATFORM_WINDOWS || PLATFORM_LINUX
 			if (!state.load(std::memory_order_acquire))
-			{
 				WaitSlow();
-			}
+			#else
+			std::unique_lock selfLock(lock);
+			condition.wait(selfLock, [this] { return !shouldWait; });
+			#endif
 		}
 
 		inline bool WaitFor(u32 waitTime)
 		{
+			#if PLATFORM_WINDOWS || PLATFORM_LINUX
 			return state.load(std::memory_order_acquire) || WaitForSlow(waitTime);
-		}
+			#else
+			if (waitTime == 0xffffffff)
+				return Wait(), true;
 
-		//inline bool WaitUntil(FMonotonicTimePoint waitTime)
-		//{
-		//	return state.load(std::memory_order_acquire) || WaitUntilSlow(waitTime);
-		//}
+			std::unique_lock selfLock(lock);
+			if (waitTime == 0)
+				return !shouldWait;
+			return condition.wait_for(selfLock, std::chrono::milliseconds(waitTime), [this] { return !shouldWait; });
+			#endif
+		}
 
 		void Notify()
 		{
+			#if PLATFORM_WINDOWS || PLATFORM_LINUX
 			if (state.exchange(1, std::memory_order_release) == 0)
 			{
-			#if PLATFORM_WINDOWS
+				#if PLATFORM_WINDOWS
 				WakeByAddressSingle((void*)&state);
-			#elif PLATFORM_LINUX
+				#elif PLATFORM_LINUX
 				syscall(SYS_futex, &state, FUTEX_WAKE_PRIVATE, INT_MAX, nullptr, nullptr, 0);
-			#else
+				#else
 				#error
-			#endif
+				#endif
 			}
+			#else
+			// The lock must be held by Notify() until it is finished accessing its members, otherwise a waiting thread
+			// may destroy the event after its wait but before Notify() has finished.
+			std::unique_lock selfLock(lock);
+			shouldWait = false;
+			condition.notify_one();
+			#endif
 		}
 
 	private:
+		#if PLATFORM_WINDOWS || PLATFORM_LINUX
 		UBA_NOINLINE void WaitSlow()
 		{
 			for (;;)
 			{
 				u32 localState = state.load(std::memory_order_acquire);
 				if (localState)
-				{
 					return;
-				}
-			#if PLATFORM_WINDOWS
+
+				#if PLATFORM_WINDOWS
 				WaitOnAddress(&state, &localState, sizeof(u32), INFINITE);
-			#elif PLATFORM_LINUX
-			#else
+				#elif PLATFORM_LINUX
+				syscall(SYS_futex, &state, FUTEX_WAIT_PRIVATE, 0, nullptr, nullptr, 0);
+				#else
 				#error
-			#endif
+				#endif
 			}
 		}
 
@@ -106,54 +146,79 @@ namespace uba::ParkingLot::Private
 		{
 			u32 localState = state.load(std::memory_order_acquire);
 			if (localState || waitTimeMs == 0)
-			{
 				return !!localState;
-			}
 
-		#if PLATFORM_WINDOWS
+			#if PLATFORM_WINDOWS
 			const bool timedOut = !WaitOnAddress(&state, &localState, sizeof(u32), waitTimeMs) && GetLastError() == ERROR_TIMEOUT;
-		#elif PLATFORM_LINUX
+			#elif PLATFORM_LINUX
+			// FUTEX_WAIT takes a relative wait time.
 			struct timespec waitTimeSpec;
 			const struct timespec* waitTimeSpecPtr = MsToTimeSpec(waitTimeMs, waitTimeSpec);
 			const bool timedOut = syscall(SYS_futex, &state, FUTEX_WAIT_PRIVATE, 0, waitTimeSpecPtr, nullptr, 0) == -1 && errno == ETIMEDOUT;
-		#else
+			#else
 			#error
-		#endif
+			#endif
 
 			localState = state.load(std::memory_order_acquire);
 			if (localState || timedOut)
-			{
 				return !!localState;
-			}
 
-			// Handle a spurious wake by waiting until the wait time has elapsed one more time because WaitUntilSlow
+			// Handle a spurious wake by waiting until the wait time has elapsed one more time because WaitForSlower
 			// handles spurious wakes in a loop and avoids exceeding the originally requested wake time by more than
 			// the typical variation due to scheduling imprecision.
-			// TODO: How to handle spurious wakes?
-			//return WaitUntilSlow(FMonotonicTimePoint::Now() + waitTimeMs);
-			return !!localState;
+			if (waitTimeMs == 0xffffffff)
+				return WaitSlow(), true;
+			return WaitForSlower(waitTimeMs, localState);
 		}
 
-		//UBA_NOINLINE bool WaitUntilSlow(FMonotonicTimePoint WaitTime)
-		//{
-		//	FMonotonicTimeSpan WaitSpan = WaitTime - FMonotonicTimePoint::Now();
-		//
-		//	for (;;)
-		//	{
-		//		bool localShouldWait = shouldWait.load(std::memory_order_acquire);
-		//		if (!localShouldWait || WaitSpan <= FMonotonicTimeSpan::Zero())
-		//		{
-		//			return !localShouldWait;
-		//		}
-		//		const DWORD WaitMs = WaitTime.IsInfinity() ? INFINITE : DWORD(FPlatformMath::CeilToInt64(WaitSpan.ToMilliseconds()));
-		//		WaitOnAddress(&shouldWait, &localShouldWait, sizeof(bool), WaitMs);
-		//		WaitSpan = WaitTime - FMonotonicTimePoint::Now();
-		//	}
-		//}
-
-		//std::atomic<bool> shouldWait = true;
+		UBA_NOINLINE bool WaitForSlower(u32 waitTimeMs, u32 localState)
+		{
+			#if PLATFORM_WINDOWS
+			const u64 endTimeMs = TimeToMs(GetTime()) + waitTimeMs;
+
+			for (;;)
+			{
+				const bool timedOut = !WaitOnAddress(&state, &localState, sizeof(u32), waitTimeMs) && GetLastError() == ERROR_TIMEOUT;
+
+				localState = state.load(std::memory_order_acquire);
+				if (localState || timedOut)
+					return !!localState;
+
+				const u64 curTimeMs = TimeToMs(GetTime());
+				if (endTimeMs <= curTimeMs)
+					return false;
+
+				waitTimeMs = u32(endTimeMs - curTimeMs);
+			}
+			#elif PLATFORM_LINUX
+			struct timespec waitTimeSpec;
+			if (clock_gettime(CLOCK_MONOTONIC, &waitTimeSpec) == -1)
+				FatalError(1401, TC("clock_gettime(CLOCK_MONOTONIC) failed"));
+			MsToTimeSpec(TimeSpecToMs(waitTimeSpec) + waitTimeMs, waitTimeSpec);
+
+			for (;;)
+			{
+				if (state.load(std::memory_order_acquire))
+					return true;
+
+				// FUTEX_WAIT_BITSET takes an absolute wait time.
+				if (syscall(SYS_futex, &state, FUTEX_WAIT_BITSET_PRIVATE, 0, &waitTimeSpec, nullptr, FUTEX_BITSET_MATCH_ANY) == -1 && errno == ETIMEDOUT)
+					return false;
+			}
+			#else
+			#error
+			#endif
+		}
+		#endif
+
+		#if PLATFORM_WINDOWS || PLATFORM_LINUX
 		// 0 = Reset, 1 = Notified
 		std::atomic<u32> state = 0;
+		#else
+		std::mutex lock;
+		std::condition_variable condition;
+		bool shouldWait = true;
+		#endif
 	};
 } // uba::ParkingLot::Private
 
@@ -603,6 +668,9 @@ namespace uba::ParkingLot::Private
 
 	////////////////////////////////////////////////////////////////////////////////////////////////////
 
+	#pragma warning(push)
+	#pragma warning(disable:4324) // structure was padded due to alignment specifier
+
 	struct alignas(CacheLineSize) Thread final
 	{
 		Thread* next = nullptr;
@@ -610,7 +678,6 @@ namespace uba::ParkingLot::Private
 		u64 wakeToken = 0;
 		PlatformManualResetEvent event;
 		std::atomic<u32> referenceCount = 0;
-		u64 padding[4]{};
 
 		inline void AddRef()
 		{
@@ -632,6 +699,8 @@ namespace uba::ParkingLot::Private
 		Thread& operator=(const Thread&) = delete;
 	};
 
+	#pragma warning(pop)
+
 	class ThreadLocalData
 	{
 	public:
@@ -772,7 +841,7 @@ namespace uba::ParkingLot::Private
 	////////////////////////////////////////////////////////////////////////////////////////////////////
 
 	#pragma warning(push)
-	#pragma warning(disable:4200)  // nonstandard extension used: zero-length array in struct/union
+	#pragma warning(disable:4200) // nonstandard extension used: zero-length array in struct/union
 
 	class Table final
 	{
@@ -1206,17 +1275,6 @@ namespace uba::ParkingLot::Private
 		return TimedWait(address, canWait, canWaitContext, beforeWait, beforeWaitContext, WaitForTime, &waitTimeMs);
 	}
 
-	//static void WaitUntilTime(void* waitTime, PlatformManualResetEvent& event)
-	//{
-	//	event.WaitUntil(*(FMonotonicTimePoint*)waitTime);
-	//}
-
-	//WaitState WaitUntil(const void* address, bool (*canWait)(void*), void* canWaitContext, void (*beforeWait)(void*), void* beforeWaitContext, FMonotonicTimePoint waitTime)
-	//{
-	//	UBA_ASSERT(!waitTime.IsNaN());
-	//	return TimedWait(address, canWait, canWaitContext, beforeWait, beforeWaitContext, WaitUntilTime, &waitTime);
-	//}
-
 	void WakeOne(const void* address, u64 (*onWakeState)(void*, WakeState), void* onWakeStateContext)
 	{
 		using namespace Private;
diff --git a/Engine/Source/Programs/UnrealBuildAccelerator/Core/Public/UbaBase.h b/Engine/Source/Programs/UnrealBuildAccelerator/Core/Public/UbaBase.h
--- a/Engine/Source/Programs/UnrealBuildAccelerator/Core/Public/UbaBase.h
+++ b/Engine/Source/Programs/UnrealBuildAccelerator/Core/Public/UbaBase.h
@@ -75,3 +75,5 @@ namespace uba
 #else
 	#define UBA_LIFETIMEBOUND
 #endif
+
+#define UBA_USE_PARKINGLOT 0
diff --git a/Engine/Source/Programs/UnrealBuildAccelerator/Core/Public/UbaParkingLot.h b/Engine/Source/Programs/UnrealBuildAccelerator/Core/Public/UbaParkingLot.h
--- a/Engine/Source/Programs/UnrealBuildAccelerator/Core/Public/UbaParkingLot.h
+++ b/Engine/Source/Programs/UnrealBuildAccelerator/Core/Public/UbaParkingLot.h
@@ -34,7 +34,6 @@ namespace uba::ParkingLot
 	{
 		UBA_API WaitState Wait(const void* address, bool (*canWait)(void*), void* canWaitContext, void (*beforeWait)(void*), void* beforeWaitContext);
 		UBA_API WaitState WaitFor(const void* address, bool (*canWait)(void*), void* canWaitContext, void (*beforeWait)(void*), void* beforeWaitContext, u32 waitTimeMs);
-		//UBA_API WaitState WaitUntil(const void* address, bool (*canWait)(void*), void* canWaitContext, void (*beforeWait)(void*), void* beforeWaitContext, FMonotonicTimePoint waitTimeMs);
 		UBA_API void WakeOne(const void* address, u64 (*onWakeState)(void*, WakeState), void* onWakeStateContext);
 	} // Private
 
@@ -50,12 +49,6 @@ namespace uba::ParkingLot
 		return Private::WaitFor(address, canWait.GetFunction(), canWait.GetContext(), beforeWait.GetFunction(), beforeWait.GetContext(), waitTimeMs);
 	}
 
-	///** Queue the calling thread to wait if canWait returns true. beforeWait is only called if canWait returns true. */
-	//inline WaitState WaitUntil(const void* address, FunctionWithContext<bool ()> canWait, FunctionWithContext<void ()> beforeWait, FMonotonicTimePoint waitTime)
-	//{
-	//	return Private::WaitUntil(address, canWait.GetFunction(), canWait.GetContext(), beforeWait.GetFunction(), beforeWait.GetContext(), waitTime);
-	//}
-
 	/** Wake one thread from the queue of threads waiting on the address. */
 	inline void WakeOne(const void* address, FunctionWithContext<u64 (WakeState)> onWakeState)
 	{
diff --git a/Engine/Source/Programs/UnrealBuildAccelerator/Core/Public/UbaSynchronization.h b/Engine/Source/Programs/UnrealBuildAccelerator/Core/Public/UbaSynchronization.h
--- a/Engine/Source/Programs/UnrealBuildAccelerator/Core/Public/UbaSynchronization.h
+++ b/Engine/Source/Programs/UnrealBuildAccelerator/Core/Public/UbaSynchronization.h
@@ -7,7 +7,6 @@
 #include <utility>
 
 #define UBA_TRACK_CONTENTION 0
-#define UBA_USE_PARKINGLOT 0
 
 #define STRING_JOIN(arg1, arg2) STRING_JOIN_INNER(arg1, arg2)
 #define STRING_JOIN_INNER(arg1, arg2) arg1 ## arg2
@@ -220,7 +219,7 @@ namespace uba
 	////////////////////////////////////////////////////////////////////////////////////////////////////
 
 	template <typename MutexType>
-	class ScopedCriticalSection
+	class ScopedCriticalSection 
 	{
 	public:
 		ScopedCriticalSection(MutexType& mutex) : m_mutex(mutex), m_active(true) { mutex.Enter(); }
diff --git a/Engine/Source/Programs/UnrealBuildAccelerator/Core/Public/UbaTimer.h b/Engine/Source/Programs/UnrealBuildAccelerator/Core/Public/UbaTimer.h
--- a/Engine/Source/Programs/UnrealBuildAccelerator/Core/Public/UbaTimer.h
+++ b/Engine/Source/Programs/UnrealBuildAccelerator/Core/Public/UbaTimer.h
@@ -20,7 +20,6 @@ namespace uba
 	inline u64 TimeToMs(u64 time)					{ return TimeToMs(time, GetFrequency()); }
 	inline float TimeToS(u64 time)					{ return TimeToS(time, GetFrequency()); }
 
-	inline u64 TickToTime(u64 us)					{ return us * GetFrequency() / 10'000'000; }
 	inline u64 UsToTime(u64 us)						{ return us * GetFrequency() / 1'000'000; }
 	inline u64 MsToTime(u64 ms)						{ return ms * GetFrequency() / 1'000; }
 
