[UBA]
* Added Scheduler::GetTotalProcesses
* Fixed so SearchPathForFile can find files that has no extensions (it tries only with .exe)
* Fixed so DirectoryCache::CreateDirectory can handle both forward and backward slashes
* Optimized command line parser
* Optimized ProcessStartInfo Expand function
* Minor off-by-one comparison for reading strings over binary reader
* Fixed so detouring code can handle environment variables SystemRoot and TEMP not being set (it is possible when building projects such as chromium)

diff --git a/Engine/Source/Programs/UnrealBuildAccelerator/Common/Private/UbaFile.cpp b/Engine/Source/Programs/UnrealBuildAccelerator/Common/Private/UbaFile.cpp
--- a/Engine/Source/Programs/UnrealBuildAccelerator/Common/Private/UbaFile.cpp
+++ b/Engine/Source/Programs/UnrealBuildAccelerator/Common/Private/UbaFile.cpp
@@ -502,8 +502,18 @@ namespace uba
 
 		StringBuffer<> fullPath;
 
+		StringView fileView = ToView(file);
+		StringView ext;
+
+		#if PLATFORM_WINDOWS
+		if (fileView.count < 4 || fileView[fileView.count - 4] != '.')
+			ext = TCV(".exe");
+		#endif
+
 		auto TestFileExists = [&](const tchar* extraInfo)
 			{
+				fullPath.EnsureEndsWithSlash().Append(fileView).Append(ext);
+
 				if (GetFileAttributesW(fullPath.data) != INVALID_FILE_ATTRIBUTES)
 				{
 					FixPath(fullPath.data, nullptr, 0, out);
@@ -517,14 +527,14 @@ namespace uba
 
 		if (applicationDir.count)
 		{
-			fullPath.Append(applicationDir).EnsureEndsWithSlash().Append(file);
+			fullPath.Append(applicationDir);
 			if (TestFileExists(TC("")))
 				return true;
 		}
 
 		if (workingDir.count)
 		{
-			fullPath.Clear().Append(workingDir).EnsureEndsWithSlash().Append(file);
+			fullPath.Clear().Append(workingDir);
 			if (TestFileExists(TC("")))
 				return true;
 		}
@@ -552,10 +562,6 @@ namespace uba
 			*it = 0;
 
 			fullPath.Clear().Append(lastStart);
-			if (*lastStart)
-				fullPath.EnsureEndsWithSlash();
-			fullPath.Append(file);
-
 			if (TestFileExists(TC(" using PATH environment variable")))
 				return true;
 
@@ -1079,7 +1085,7 @@ namespace uba
 		bool& alreadyExists = outAlreadyExists ? *outAlreadyExists : throwaway;
 
 		u64 dirLen = TStrlen(dir);
-		if (dir[dirLen - 1] == PathSeparator)
+		if (dir[dirLen - 1] == PathSeparator || (IsWindows && (dir[dirLen - 1] == NonPathSeparator)))
 			--dirLen;
 		TString key(dir, dir + dirLen);
 		dir = key.c_str();
@@ -1109,7 +1115,11 @@ namespace uba
 			return logger.Error(TC("Failed to create directory %s (%s)"), dir, LastErrorToText(lastError).data);
 
 		tchar temp[512];
-		const tchar* lastSep = TStrrchr(dir, PathSeparator);
+		const tchar* lastSep = nullptr;
+		for (const tchar* p = dir; *p; ++p)
+			if (*p == PathSeparator || *p == NonPathSeparator)
+				lastSep = p;
+		UBA_ASSERT(lastSep);
 		u64 pos = u64(lastSep - dir);
 		memcpy(temp, dir, pos * sizeof(tchar));
 		temp[pos] = 0;
diff --git a/Engine/Source/Programs/UnrealBuildAccelerator/Common/Private/UbaProcessStartInfoHolder.cpp b/Engine/Source/Programs/UnrealBuildAccelerator/Common/Private/UbaProcessStartInfoHolder.cpp
--- a/Engine/Source/Programs/UnrealBuildAccelerator/Common/Private/UbaProcessStartInfoHolder.cpp
+++ b/Engine/Source/Programs/UnrealBuildAccelerator/Common/Private/UbaProcessStartInfoHolder.cpp
@@ -18,7 +18,7 @@ namespace uba
 				{
 					StringView view(arg, argLen);
 					requiresShell |= argLen == 1 && (*arg == '<' || *arg == '>' || *arg == '|' || *arg == '^');
-					requiresShell |= argLen == 2 && (view.Equals(TC(">>")) || view.Equals(TC("&&")) || view.Equals(TC("||")));
+					requiresShell |= argLen == 2 && (view.Equals(TCV(">>")) || view.Equals(TCV("&&")) || view.Equals(TCV("||")));
 				});
 			if (requiresShell)
 			{
diff --git a/Engine/Source/Programs/UnrealBuildAccelerator/Common/Private/UbaScheduler.cpp b/Engine/Source/Programs/UnrealBuildAccelerator/Common/Private/UbaScheduler.cpp
--- a/Engine/Source/Programs/UnrealBuildAccelerator/Common/Private/UbaScheduler.cpp
+++ b/Engine/Source/Programs/UnrealBuildAccelerator/Common/Private/UbaScheduler.cpp
@@ -427,6 +427,11 @@ namespace uba
 		return index;
 	}
 
+	u32 Scheduler::GetTotalProcesses()
+	{
+		return m_totalProcesses;
+	}
+
 	void Scheduler::GetStats(u32& outQueued, u32& outActiveLocal, u32& outActiveRemote, u32& outFinished)
 	{
 		outActiveLocal = m_activeLocalProcesses;
diff --git a/Engine/Source/Programs/UnrealBuildAccelerator/Common/Public/UbaScheduler.h b/Engine/Source/Programs/UnrealBuildAccelerator/Common/Public/UbaScheduler.h
--- a/Engine/Source/Programs/UnrealBuildAccelerator/Common/Public/UbaScheduler.h
+++ b/Engine/Source/Programs/UnrealBuildAccelerator/Common/Public/UbaScheduler.h
@@ -91,6 +91,7 @@ namespace uba
 
 		u32 EnqueueProcess(const EnqueueProcessInfo& info); // Returns index of process. Index is a rolling number
 
+		u32 GetTotalProcesses();
 		void GetStats(u32& outQueued, u32& outActiveLocal, u32& outActiveRemote, u32& outFinished); // This is not "threadsafe".. which means that you are not guaranteed to get the exact correct state.
 
 		bool IsEmpty(); // Returns true if scheduler is entirely empty.. as in no processes are left in the system
diff --git a/Engine/Source/Programs/UnrealBuildAccelerator/Core/Public/UbaBinaryReaderWriter.h b/Engine/Source/Programs/UnrealBuildAccelerator/Core/Public/UbaBinaryReaderWriter.h
--- a/Engine/Source/Programs/UnrealBuildAccelerator/Core/Public/UbaBinaryReaderWriter.h
+++ b/Engine/Source/Programs/UnrealBuildAccelerator/Core/Public/UbaBinaryReaderWriter.h
@@ -349,7 +349,7 @@ namespace uba
 	u64 BinaryReader::ReadString(tchar* str, u64 strCapacity)
 	{
 		u64 charLen = Read7BitEncoded();
-		if (charLen < strCapacity - 1)
+		if (charLen < strCapacity)
 			return InternalReadString(str, charLen);
 
 		u64 newCharLen = strCapacity - 1;
diff --git a/Engine/Source/Programs/UnrealBuildAccelerator/Core/Public/UbaProcessUtils.h b/Engine/Source/Programs/UnrealBuildAccelerator/Core/Public/UbaProcessUtils.h
--- a/Engine/Source/Programs/UnrealBuildAccelerator/Core/Public/UbaProcessUtils.h
+++ b/Engine/Source/Programs/UnrealBuildAccelerator/Core/Public/UbaProcessUtils.h
@@ -23,7 +23,7 @@ namespace uba
 		bool isBackslashOwned = false;
 
 		
-		CharType temp[5];
+		CharType temp[4*1024];
 
 		CharType* argWriteBuffer = temp;
 		u64 argBufferSize = sizeof_array(temp);
@@ -46,6 +46,17 @@ namespace uba
 				if ((pastEnd || *argEnd == '\n') && argStart != argEnd && *(argEnd-1) == '\r')
 					--argEnd;
 
+				// Check once if this argument might overflow the buffer
+				u64 maxArgLen = argEnd - argStart;
+				if (maxArgLen + 1 > argBufferSize)
+				{
+					// This argument is too large for our buffer - allocate exactly what we need
+					if (argWriteBuffer != temp)
+						delete[] argWriteBuffer;
+					argBufferSize = maxArgLen + 16;
+					argWriteBuffer = new CharType[argBufferSize];
+				}
+
 				CharType* argWriteIt = argWriteBuffer;
 				CharType lastChar2 = 0;
 				isBackslashOwned = false;
@@ -67,16 +78,6 @@ namespace uba
 
 					*argWriteIt++ = *argIt;
 
-					u64 writtenCount = argWriteIt - argWriteBuffer;
-					if (writtenCount + 1 == argBufferSize)
-					{
-						UBA_ASSERTF(argWriteBuffer == temp, TC("argWriteBuffer already replaced for size %llu, and now we need %llu"), argBufferSize, writtenCount + 1);
-						argBufferSize = argumentsLen + 16;
-						argWriteBuffer = new CharType[argBufferSize];
-						argWriteIt = argWriteBuffer + writtenCount;
-						memcpy(argWriteBuffer, temp, writtenCount *sizeof(CharType));
-					}
-
 					lastChar2 = *argIt;
 					++argIt;
 				}
diff --git a/Engine/Source/Programs/UnrealBuildAccelerator/Detours/Private/Windows/UbaDetoursFunctionsWin.cpp b/Engine/Source/Programs/UnrealBuildAccelerator/Detours/Private/Windows/UbaDetoursFunctionsWin.cpp
--- a/Engine/Source/Programs/UnrealBuildAccelerator/Detours/Private/Windows/UbaDetoursFunctionsWin.cpp
+++ b/Engine/Source/Programs/UnrealBuildAccelerator/Detours/Private/Windows/UbaDetoursFunctionsWin.cpp
@@ -980,16 +980,13 @@ void Init(const DetoursPayload& payload, u64 startTime)
 		g_stdHandle[2] = makeDetouredHandle(new DetouredHandle(HandleType_StdIn)); // STD_IN
 	}
 
+	// These can actually not exist in certain programs
 	g_systemRoot.count = True_GetEnvironmentVariableW(L"SystemRoot", g_systemRoot.data, g_systemRoot.capacity);
 	g_systemRoot.MakeLower();
-	UBA_ASSERT(g_systemRoot.count);
-
 	wchar_t systemTemp[256];
-	DWORD tempLen = True_GetEnvironmentVariableW(L"TEMP", systemTemp, sizeof_array(systemTemp));(void)tempLen;
-	UBA_ASSERTF(tempLen, TC("TEMP environment variable not found?"));
-	UBA_ASSERTF(tempLen < sizeof_array(systemTemp), TC("systemTemp buffer not large enough"));
-
-	FixPath(g_systemTemp, systemTemp);
+	DWORD tempLen = True_GetEnvironmentVariableW(L"TEMP", systemTemp, sizeof_array(systemTemp));
+	if (tempLen)
+		FixPath(g_systemTemp, systemTemp);
 
 	StringBuffer<512> applicationBuffer;
 	StringBuffer<512> workingDirBuffer;
