[UBA]
* Added support for enqueing suspended processes to the scheduler. This way calling system can build up dependencies and resume processes once ready

diff --git a/Engine/Source/Programs/UnrealBuildAccelerator/Common/Private/UbaScheduler.cpp b/Engine/Source/Programs/UnrealBuildAccelerator/Common/Private/UbaScheduler.cpp
--- a/Engine/Source/Programs/UnrealBuildAccelerator/Common/Private/UbaScheduler.cpp
+++ b/Engine/Source/Programs/UnrealBuildAccelerator/Common/Private/UbaScheduler.cpp
@@ -327,36 +327,7 @@
 		m_session.CancelAllProcesses();
 	}
 
-	void Scheduler::SkipAllQueued()
+	Scheduler::ProcessEntry Scheduler::CreateProcessEntry(const EnqueueProcessInfo& info)
 	{
-		Vector<ProcessStartInfo2*> skipped;
-		SCOPED_FUTEX(m_processEntriesLock, lock);
-		for (auto& entry : m_processEntries)
-		{
-			if (entry.runStatus == RunStatus_QueuedForRun && (entry.cacheStatus == CacheStatus_QueuedForTest || entry.cacheStatus == CacheStatus_Failed))
-			{
-				entry.cacheStatus = CacheStatus_Failed;
-				entry.runStatus = RunStatus_Skipped;
-				skipped.push_back(entry.info);
-			}
-		}
-		lock.Leave();
-		for (auto pi : skipped)
-			SkipProcess(*pi);
-	}
-
-	void Scheduler::SetMaxLocalProcessors(u32 maxLocalProcessors)
-	{
-		m_maxLocalProcessors = maxLocalProcessors;
-		m_updateThreadWakeup.Set();
-	}
-
-	void Scheduler::SetAllowDisableRemoteExecution(bool allow)
-	{
-		m_allowDisableRemoteExecution = allow;
-	}
-
-	u32 Scheduler::EnqueueProcess(const EnqueueProcessInfo& info)
-	{
 		u8* ki = nullptr;
 		if (info.knownInputsCount)
@@ -404,10 +375,7 @@
 				processConfig->GetValueAsBool(canDetour, TC("CanDetour"));
 			}
 		}
-
-		SCOPED_FUTEX(m_processEntriesLock, lock);
-		u32 index = u32(m_processEntries.size());
-		auto& entry = m_processEntries.emplace_back();
+		ProcessEntry entry;
 		entry.info = info2;
 		entry.dependencies = dep;
 		entry.dependencyCount = info.dependencyCount;
@@ -417,6 +385,46 @@
 		entry.canExecuteRemotely = canExecuteRemotely;
 		entry.canCrossArchitecture = canCrossArchitecture;
 		entry.canCrossPlatform = canCrossPlatform;
+		return entry;
+	}
+
+
+	void Scheduler::SkipAllQueued()
+	{
+		Vector<ProcessStartInfo2*> skipped;
+		SCOPED_FUTEX(m_processEntriesLock, lock);
+		for (auto& entry : m_processEntries)
+		{
+			if (entry.runStatus == RunStatus_QueuedForRun && (entry.cacheStatus == CacheStatus_QueuedForTest || entry.cacheStatus == CacheStatus_Failed))
+			{
+				entry.cacheStatus = CacheStatus_Failed;
+				entry.runStatus = RunStatus_Skipped;
+				skipped.push_back(entry.info);
+			}
+		}
+		lock.Leave();
+		for (auto pi : skipped)
+			SkipProcess(*pi);
+	}
+
+	void Scheduler::SetMaxLocalProcessors(u32 maxLocalProcessors)
+	{
+		m_maxLocalProcessors = maxLocalProcessors;
+		m_updateThreadWakeup.Set();
+	}
+
+	void Scheduler::SetAllowDisableRemoteExecution(bool allow)
+	{
+		m_allowDisableRemoteExecution = allow;
+	}
+
+	u32 Scheduler::EnqueueProcess(const EnqueueProcessInfo& info)
+	{
+		ProcessEntry entry = CreateProcessEntry(info);
+
+		SCOPED_FUTEX(m_processEntriesLock, lock);
+		u32 index = u32(m_processEntries.size());
+		m_processEntries.emplace_back(entry);
 		lock.Leave();
 
 		++m_totalProcesses;
@@ -427,6 +435,26 @@
 		return index;
 	}
 
+	u32 Scheduler::EnqueueSuspendedProcess()
+	{
+		SCOPED_FUTEX(m_processEntriesLock, lock);
+		UBA_ASSERT(!m_cancelled);
+		u32 index = u32(m_processEntries.size());
+		auto& entry = m_processEntries.emplace_back();
+		entry.runStatus = RunStatus_Suspended;
+		entry.cacheStatus = CacheStatus_Suspended;
+		++m_totalProcesses;
+		UpdateQueueCounter(1);
+		return index;
+	}
+	void Scheduler::ResumeQueuedProcess(u32 index, const EnqueueProcessInfo& info)
+	{
+		ProcessEntry entry = CreateProcessEntry(info);
+		SCOPED_FUTEX(m_processEntriesLock, lock);
+		m_processEntries[index] = entry;
+		m_updateThreadWakeup.Set();
+	}
+
 	u32 Scheduler::GetTotalProcesses()
 	{
 		return m_totalProcesses;
@@ -876,7 +904,7 @@
 				return i;
 			}
 
-			if (cacheStatus == CacheStatus_Testing)
+			if (cacheStatus == CacheStatus_Suspended || cacheStatus == CacheStatus_Testing)
 				continue;
 
 			UBA_ASSERT(cacheStatus == CacheStatus_QueuedForDownload);
@@ -904,11 +932,11 @@
 			auto runStatus = entry.runStatus;
 
 			// Check if we can move search starting point forward
-			if (runStatus != RunStatus_QueuedForRun)
+			if (runStatus > RunStatus_QueuedForRun)
 			{
 				if (allFinished)
 				{
-					if (runStatus != RunStatus_Running)
+					if (runStatus > RunStatus_Running)
 						m_processEntriesRunStart = i;
 					else
 						allFinished = false;
diff --git a/Engine/Source/Programs/UnrealBuildAccelerator/Common/Public/UbaScheduler.h b/Engine/Source/Programs/UnrealBuildAccelerator/Common/Public/UbaScheduler.h
--- a/Engine/Source/Programs/UnrealBuildAccelerator/Common/Public/UbaScheduler.h
+++ b/Engine/Source/Programs/UnrealBuildAccelerator/Common/Public/UbaScheduler.h
@@ -91,6 +91,9 @@ namespace uba
 
 		u32 EnqueueProcess(const EnqueueProcessInfo& info); // Returns index of process. Index is a rolling number
 
+		u32 EnqueueSuspendedProcess(); // Can be used to put entries in the queue early in order to calculate process info afterwards
+		void ResumeQueuedProcess(u32 index, const EnqueueProcessInfo& info); // Resume process queued with EnqueueSuspended
+
 		u32 GetTotalProcesses();
 		void GetStats(u32& outQueued, u32& outActiveLocal, u32& outActiveRemote, u32& outFinished); // This is not "threadsafe".. which means that you are not guaranteed to get the exact correct state.
 
@@ -113,6 +116,7 @@ namespace uba
 
 		enum RunStatus : u8
 		{
+			RunStatus_Suspended,
 			RunStatus_QueuedForRun,
 			RunStatus_Running,
 			RunStatus_Success,
@@ -122,6 +126,7 @@ namespace uba
 
 		enum CacheStatus : u8
 		{
+			CacheStatus_Suspended,
 			CacheStatus_QueuedForTest,
 			CacheStatus_Testing,
 			CacheStatus_QueuedForDownload,
@@ -132,6 +137,8 @@ namespace uba
 
 		UBA_NOINLINE void ThreadLoop();
 
+		struct ProcessEntry;
+		ProcessEntry CreateProcessEntry(const EnqueueProcessInfo& info);
 		void SkipAllQueued();
 		bool ProcessReturned(Process& process, bool isLocal);
 		void HandleCacheMissed(u32 processIndex);
