[UBA]
* Added so StringView::Contains has an optional position out parameter

diff --git a/Engine/Source/Programs/UnrealBuildAccelerator/Core/Private/UbaStringBuffer.cpp b/Engine/Source/Programs/UnrealBuildAccelerator/Core/Private/UbaStringBuffer.cpp
--- a/Engine/Source/Programs/UnrealBuildAccelerator/Core/Private/UbaStringBuffer.cpp
+++ b/Engine/Source/Programs/UnrealBuildAccelerator/Core/Private/UbaStringBuffer.cpp
@@ -472,11 +472,16 @@ namespace uba
 
 	////////////////////////////////////////////////////////////////////////////////////////////////////
 
-	bool StringView::Contains(tchar c) const
+	bool StringView::Contains(tchar c, u32* outPos) const
 	{
 		for (const tchar* i = data, *e = data + count; i!=e; ++i)
-			if (*i == c)
-				return true;
+		{
+			if (*i != c)
+				continue;
+			if (outPos)
+				*outPos = u32(i - data);
+			return true;
+		}
 		return false;
 	}
 
diff --git a/Engine/Source/Programs/UnrealBuildAccelerator/Core/Public/UbaStringBuffer.h b/Engine/Source/Programs/UnrealBuildAccelerator/Core/Public/UbaStringBuffer.h
--- a/Engine/Source/Programs/UnrealBuildAccelerator/Core/Public/UbaStringBuffer.h
+++ b/Engine/Source/Programs/UnrealBuildAccelerator/Core/Public/UbaStringBuffer.h
@@ -129,7 +129,7 @@ namespace uba
 		bool StartsWith(const tchar* str, bool ignoreCase = true) const { return uba::StartsWith(data, str, ignoreCase); }
 		bool EndsWith(const StringView& value, bool ignoreCase = true) const { return uba::EndsWith(data, count, value.data, value.count, ignoreCase); }
 		bool EndsWith(const tchar* value, bool ignoreCase = true) const;
-		bool Contains(tchar c) const;
+		bool Contains(tchar c, u32* outPos = nullptr) const;
 		bool Contains(const tchar* str, bool ignoreCase = true) const { return uba::Contains(data, str, ignoreCase); }
 		bool Contains(const StringView& str, bool ignoreCase = true) const { return uba::Contains(data, str.data, ignoreCase); }
 		bool Equals(const tchar* str, bool ignoreCase = true) const { return uba::Equals(data, str, ignoreCase); }
