[UBA]
* Speculative fix for race condition in directory table (seen when compiling shaders). Added additional check for directory parse position and available position inside lock

diff --git a/Engine/Source/Programs/UnrealBuildAccelerator/Core/Public/UbaDirectoryTable.h b/Engine/Source/Programs/UnrealBuildAccelerator/Core/Public/UbaDirectoryTable.h
--- a/Engine/Source/Programs/UnrealBuildAccelerator/Core/Public/UbaDirectoryTable.h
+++ b/Engine/Source/Programs/UnrealBuildAccelerator/Core/Public/UbaDirectoryTable.h
@@ -265,8 +265,11 @@ namespace uba
 					if (parentDir.parseOffset != parentDir.tableOffset)
 					{
 						SCOPED_WRITE_LOCK(parentDir.lock, lock);
-						PopulateDirectoryRecursive(ancestorHasher, parentDir.tableOffset, parentDir.parseOffset, parentDir.files);
-						parentDir.parseOffset = parentDir.tableOffset;
+						if (parentDir.parseOffset != parentDir.tableOffset)
+						{
+							PopulateDirectoryRecursive(ancestorHasher, parentDir.tableOffset, parentDir.parseOffset, parentDir.files);
+							parentDir.parseOffset = parentDir.tableOffset;
+						}
 					}
 
 					SCOPED_READ_LOCK(parentDir.lock, lock);
