[UBA]
* Improved log output in visualizer when log has lots of escape characters (like changing colors of text etc)

diff --git a/Engine/Source/Programs/UnrealBuildAccelerator/Visualizer/Private/UbaVisualizer.cpp b/Engine/Source/Programs/UnrealBuildAccelerator/Visualizer/Private/UbaVisualizer.cpp
--- a/Engine/Source/Programs/UnrealBuildAccelerator/Visualizer/Private/UbaVisualizer.cpp
+++ b/Engine/Source/Programs/UnrealBuildAccelerator/Visualizer/Private/UbaVisualizer.cpp
@@ -1935,8 +1935,38 @@ namespace uba
 					logger.Info(TC(""));
 					logger.Info(TC("  ---------------- Log ----------------"));
 					logger.AddTextOffset(14);
-					for (auto& line : logLines)
-						logger.Log(LogEntryType_Info, line);
+					for (auto line : logLines)
+					{
+						if (line.empty())
+						{
+							logger.Log(LogEntryType_Info, line);
+							continue;
+						}
+						const tchar* segStart = line.c_str();
+						for (const tchar* it=segStart;; ++it)
+						{
+							if (*it != '\n' && *it != 0)
+								continue;
+							const tchar* segEnd = it;
+							if (it != segStart && it[-1] == '\r')
+								--segEnd;
+							TString str(segStart, u32(segEnd - segStart));
+							while (true)
+							{
+								auto escapeStart = str.find_first_of(TC("\x1b["));
+								if (escapeStart == -1)
+									break;
+								auto escapeEnd = str.find_first_of(TC("m"), escapeStart);
+								if(escapeEnd == -1)
+									break;
+								str.erase(escapeStart, escapeEnd + 1);
+							}
+							logger.Log(LogEntryType_Info, str);
+							if (!*it)
+								break;
+							segStart = it + 1;
+						}
+					}
 				}
 			}
 			logger.AddSpace(3);
