[UBA]
* Added so Logger::Logf returns itself

diff --git a/Engine/Source/Programs/UnrealBuildAccelerator/Common/Private/UbaLogger.cpp b/Engine/Source/Programs/UnrealBuildAccelerator/Common/Private/UbaLogger.cpp
--- a/Engine/Source/Programs/UnrealBuildAccelerator/Common/Private/UbaLogger.cpp
+++ b/Engine/Source/Programs/UnrealBuildAccelerator/Common/Private/UbaLogger.cpp
@@ -173,12 +173,13 @@ namespace uba
 		return false;
 	}
 
-	void Logger::Logf(LogEntryType type, const tchar* format, ...)
+	Logger& Logger::Logf(LogEntryType type, const tchar* format, ...)
 	{
 		va_list arg;
 		va_start(arg, format);
 		LogArg(type, format, arg);
 		va_end(arg);
+		return *this;
 	}
 
 	void Logger::LogArg(LogEntryType type, const tchar* format, va_list& args)
diff --git a/Engine/Source/Programs/UnrealBuildAccelerator/Common/Public/UbaLogger.h b/Engine/Source/Programs/UnrealBuildAccelerator/Common/Public/UbaLogger.h
--- a/Engine/Source/Programs/UnrealBuildAccelerator/Common/Public/UbaLogger.h
+++ b/Engine/Source/Programs/UnrealBuildAccelerator/Common/Public/UbaLogger.h
@@ -37,7 +37,7 @@ namespace uba
 		Logger& Info(const tchar* format, ...);
 		Logger& Detail(const tchar* format, ...);
 		Logger& Debug(const tchar* format, ...);
-		void Logf(LogEntryType type, const tchar* format, ...);
+		Logger& Logf(LogEntryType type, const tchar* format, ...);
 		void LogArg(LogEntryType type, const tchar* format, va_list& args);
 		void Log(LogEntryType type, const StringView& str);
 
