[UBA]
* Added code to UbaCli to be able to easily loop same process without shutting down session

diff --git a/Engine/Source/Programs/UnrealBuildAccelerator/Cli/Private/UbaCli.cpp b/Engine/Source/Programs/UnrealBuildAccelerator/Cli/Private/UbaCli.cpp
--- a/Engine/Source/Programs/UnrealBuildAccelerator/Cli/Private/UbaCli.cpp
+++ b/Engine/Source/Programs/UnrealBuildAccelerator/Cli/Private/UbaCli.cpp
@@ -1333,6 +1333,7 @@ namespace uba
 			//info.traceName.Append(TCV("TESTTRACE"));
 			//info.storeIntermediateFilesCompressed = false;
 			info.readIntermediateFilesCompressed = true;
+			//info.treatTempDirAsEmpty = false;
 			//info.extractObjFilesSymbols = true;
 			#if UBA_DEBUG_LOG_ENABLED
 			info.remoteLogEnabled = true;
@@ -1730,6 +1731,7 @@ namespace uba
 #endif
 
 		bool success = true;
+		u32 procLoopCount = 1;
 
 		for (u32 i=0; i!=loopCount; ++i)
 		{
@@ -1745,21 +1747,28 @@ namespace uba
 						return RunScheduler(sessionServer, application, arguments);
 					if (!createCoordinator(nullptr))
 						return false;
-					switch (commandType)
+					bool res = false;
+
+					for (u32 procIt=0; procIt!=procLoopCount; ++procIt)
 					{
-					case CommandType_Native:
-						return RunLocal(sessionServer, application, arguments, false);
-						break;
-					case CommandType_Local:
-						return RunLocal(sessionServer, application, arguments, true);
-						break;
-					case CommandType_Remote:
-						return RunRemote(sessionServer, application, arguments);
-						break;
-					case CommandType_Agent:
-						return RunAgent(sessionServer, application, arguments);
+						switch (commandType)
+						{
+						case CommandType_Native:
+							res = RunLocal(sessionServer, application, arguments, false);
+							break;
+						case CommandType_Local:
+							res = RunLocal(sessionServer, application, arguments, true);
+							break;
+						case CommandType_Remote:
+							res = RunRemote(sessionServer, application, arguments);
+							break;
+						case CommandType_Agent:
+							res = RunAgent(sessionServer, application, arguments);
+						}
+						if (!res)
+							return false;
 					}
-					return false;
+					return true;
 				});
 			if (!success)
 				break;
