[UBA]
* Changed so process startinfo holder uses WriteLongString/ReadLongString (with threshold 1000) for arguments since arguments can be insanely long

diff --git a/Engine/Source/Programs/UnrealBuildAccelerator/Common/Public/UbaProcessStartInfoHolder.h b/Engine/Source/Programs/UnrealBuildAccelerator/Common/Public/UbaProcessStartInfoHolder.h
--- a/Engine/Source/Programs/UnrealBuildAccelerator/Common/Public/UbaProcessStartInfoHolder.h
+++ b/Engine/Source/Programs/UnrealBuildAccelerator/Common/Public/UbaProcessStartInfoHolder.h
@@ -63,7 +63,7 @@ namespace uba
 				writer.WriteString(applicationOverride);
 			else
 				writer.WriteString(applicationStr);
-			writer.WriteString(argumentsStr);
+			writer.WriteLongString(argumentsStr, 1000);
 			writer.WriteString(workingDirStr);
 			writer.WriteString(logFileStr);
 			// ignore breadcrumbs here
@@ -77,7 +77,7 @@ namespace uba
 		{
 			descriptionStr = reader.ReadString();
 			applicationStr = reader.ReadString();
-			argumentsStr = reader.ReadString();
+			argumentsStr = reader.ReadLongString();
 			workingDirStr = reader.ReadString();
 			logFileStr = reader.ReadString();
 			// ignore breadcrumbs here
