[UBA]
* Fixed bug in RtlAnsiStringToUnicodeString and RtlUnicodeStringToAnsiString when using mimalloc (which is default).

diff --git a/Engine/Source/Programs/UnrealBuildAccelerator/Detours/Private/Windows/UbaDetoursFunctionsNtDll.inl b/Engine/Source/Programs/UnrealBuildAccelerator/Detours/Private/Windows/UbaDetoursFunctionsNtDll.inl
--- a/Engine/Source/Programs/UnrealBuildAccelerator/Detours/Private/Windows/UbaDetoursFunctionsNtDll.inl
+++ b/Engine/Source/Programs/UnrealBuildAccelerator/Detours/Private/Windows/UbaDetoursFunctionsNtDll.inl
@@ -658,7 +658,7 @@ NTSTATUS Detoured_RtlAnsiStringToUnicodeString(PUNICODE_STRING DestinationString
 #if UBA_USE_MIMALLOC
 	if (AllocateDestinationString && g_useMiMalloc)
 	{
-		DestinationString->MaximumLength = SourceString->MaximumLength * 2;
+		DestinationString->MaximumLength = (USHORT)((SourceString->Length+1) * 2);
 		DestinationString->Buffer = (wchar_t*)mi_malloc(DestinationString->MaximumLength);
 		AllocateDestinationString = false;
 	}
@@ -672,8 +672,11 @@ NTSTATUS Detoured_RtlUnicodeStringToAnsiString(PANSI_STRING DestinationString, P
 #if UBA_USE_MIMALLOC
 	if (AllocateDestinationString && g_useMiMalloc)
 	{
-		DestinationString->MaximumLength = SourceString->MaximumLength / 2;
-		DestinationString->Buffer = (char*)mi_malloc(DestinationString->MaximumLength);
+		DWORD ret;
+		RtlUnicodeToMultiByteSize( &ret, SourceString->Buffer, SourceString->Length );
+		DWORD len = ret + 1;
+		DestinationString->MaximumLength = (USHORT)len;
+		DestinationString->Buffer = (char*)mi_malloc(len);
 		AllocateDestinationString = false;
 	}
 #endif
