[UBA]
* Fixed bug where GetModuleHandleNoSyscall could assume dll was loaded but it had been unloaded.

Finally fixed the annoying "fatal error C1356: unable to find mspdbcore.dll".. thanks Phillip.Kavan for helping me debug this

diff --git a/Engine/Source/Programs/UnrealBuildAccelerator/Detours/Private/Windows/UbaDetoursFunctionsKernelBase.inl b/Engine/Source/Programs/UnrealBuildAccelerator/Detours/Private/Windows/UbaDetoursFunctionsKernelBase.inl
--- a/Engine/Source/Programs/UnrealBuildAccelerator/Detours/Private/Windows/UbaDetoursFunctionsKernelBase.inl
+++ b/Engine/Source/Programs/UnrealBuildAccelerator/Detours/Private/Windows/UbaDetoursFunctionsKernelBase.inl
@@ -2926,27 +2926,35 @@ HMODULE Recursive_LoadLibraryExW(const StringView& filePath, LPCWSTR originalNam
 	#define ReadTebOffset(offset) __readgsqword(offset)
 #endif
 
-HMODULE GetModuleHandleNoSyscall(const tchar* moduleName)
+HMODULE LoadModuleNoSyscall(StringView moduleName)
 {
-#if 1
-	const tchar* moduleFileName = GetFileName(moduleName);
+	StringView moduleFileName = moduleName.GetFileName();
 	PPEB peb = (PPEB)ReadTebOffset(0x60);
 	PPEB_LDR_DATA ldr = peb->Ldr;
 	PLIST_ENTRY head = &ldr->InMemoryOrderModuleList;
-	PLIST_ENTRY entry = head->Flink;
-	while (entry != head)
+	for (PLIST_ENTRY entry = head->Flink; entry != head; entry = entry->Flink)
 	{
 		PLDR_DATA_TABLE_ENTRY mod = CONTAINING_RECORD(entry, LDR_DATA_TABLE_ENTRY, InMemoryOrderLinks);
-		StringView view(mod->FullDllName.Buffer, mod->FullDllName.Length/2);
-		if (view.EndsWith(moduleFileName))
-			return (HMODULE)mod->DllBase;
-		entry = entry->Flink;
+		StringView fullDllName(mod->FullDllName.Buffer, mod->FullDllName.Length/2);
+		StringView dllName = fullDllName.GetFileName();
+		if (!dllName.Equals(moduleFileName))
+			continue;
+
+		// Check if memory is still committed.. if not it might have been unloaded
+		MEMORY_BASIC_INFORMATION mbi;
+		bool res = VirtualQuery(mod->DllBase, &mbi, sizeof(mbi));
+		UBA_ASSERT(res);
+		if (!res)
+			return NULL;
+		if (mbi.State != MEM_COMMIT)
+			return nullptr;
+		NTSTATUS ntRes = LdrAddRefDll(0, mod->DllBase);
+		UBA_ASSERT(ntRes == STATUS_SUCCESS);
+		if (ntRes != STATUS_SUCCESS)
+			return NULL;
+		return (HMODULE)mod->DllBase;
 	}
 	return NULL;  // Module not found
-#else
-	SuppressCreateFileDetourScope cfs;
-	return GetModuleHandle(moduleName);
-#endif
 }
 
 HMODULE Shared_LoadLibrary(LPCWSTR lpLibFileName, HANDLE hFile, DWORD dwFlags, bool& mightExist)
@@ -2958,14 +2966,16 @@ HMODULE Shared_LoadLibrary(LPCWSTR lpLibFileName, HANDLE hFile, DWORD dwFlags, b
 	path.Append(lpLibFileName).FixPathSeparators();
 
 	bool detourDll = path.EndsWith(L".exe") || path.EndsWith(L".dll");
-	if (detourDll && path.StartsWith(g_systemRoot.data))
+	if (detourDll && g_systemRoot.count && path.StartsWith(g_systemRoot.data))
 		detourDll = g_runningRemote && GetFileAttributes(path.data) == INVALID_FILE_ATTRIBUTES; // It might be that remote machine actually doesn't have the file in system32. then we need to detour
 
-	if (detourDll)
-		detourDll = !GetModuleHandleNoSyscall(lpLibFileName); // This internally can end up calling NtCreate and we don't want NtCreate to handle the download of the file because of paths
-
 	if (!detourDll)
 		return NULL;
+		
+	// When running remote the path will not be correct (since the dll is loaded from a temporary local folder)
+	// This code assumes that we never load the same dll (same exact name) from different paths and will break if we do.
+	if (HMODULE res = LoadModuleNoSyscall(path))
+		return res;
 
 	u64 pathLen = path.count;
 	u64 toSkip = 0;
