[UBA]
* Added missing file for previous submit

diff --git a/Engine/Source/Programs/UnrealBuildAccelerator/Detours/Private/Windows/UbaDetoursFunctionsWin.h b/Engine/Source/Programs/UnrealBuildAccelerator/Detours/Private/Windows/UbaDetoursFunctionsWin.h
--- a/Engine/Source/Programs/UnrealBuildAccelerator/Detours/Private/Windows/UbaDetoursFunctionsWin.h
+++ b/Engine/Source/Programs/UnrealBuildAccelerator/Detours/Private/Windows/UbaDetoursFunctionsWin.h
@@ -423,6 +423,7 @@ extern "C" {
 	NTSTATUS NTAPI ZwCreateSection(PHANDLE SectionHandle, ACCESS_MASK DesiredAccess, POBJECT_ATTRIBUTES ObjectAttributes, PLARGE_INTEGER MaximumSize, ULONG SectionPageProtection, ULONG AllocationAttributes, HANDLE FileHandle);
 	NTSTATUS NTAPI ZwQueryDirectoryFile(HANDLE FileHandle, HANDLE Event, PIO_APC_ROUTINE ApcRoutine, PVOID ApcContext, PIO_STATUS_BLOCK IoStatusBlock, PVOID FileInformation, ULONG Length, FILE_INFORMATION_CLASS FileInformationClass, BOOLEAN ReturnSingleEntry, PUNICODE_STRING FileName, BOOLEAN RestartScan);
 	NTSTATUS NTAPI ZwSetInformationFile(HANDLE FileHandle,PIO_STATUS_BLOCK IoStatusBlock, PVOID FileInformation, ULONG Length, FILE_INFORMATION_CLASS FileInformationClass);
+	NTSTATUS NTAPI LdrAddRefDll(ULONG Flags, PVOID BaseAddress);
 	PVOID WINAPI ResolveDelayLoadedAPI(PVOID ParentModuleBase, PCIMAGE_DELAYLOAD_DESCRIPTOR DelayloadDescriptor, void* FailureDllHook, void* FailureSystemHook, PIMAGE_THUNK_DATA ThunkAddress, ULONG Flags);
 	void WINAPI RtlExitUserThread(ULONG);
 	BOOLEAN NTAPI RtlFreeHeap(PVOID HeapHandle, ULONG Flags, PVOID HeapBase);
