[UBA]
* Fixed so uba FileAccessor can create memory writes of size 0

diff --git a/Engine/Source/Programs/UnrealBuildAccelerator/Common/Private/UbaFileAccessor.cpp b/Engine/Source/Programs/UnrealBuildAccelerator/Common/Private/UbaFileAccessor.cpp
--- a/Engine/Source/Programs/UnrealBuildAccelerator/Common/Private/UbaFileAccessor.cpp
+++ b/Engine/Source/Programs/UnrealBuildAccelerator/Common/Private/UbaFileAccessor.cpp
@@ -200,6 +200,13 @@ namespace uba
 			realFileName = tempFile.Append(m_tempPath).Append("Temp_").AppendValue(m_tempFileIndex).data;
 		#endif
 
+		// Size 0 can never be written to so just return a pointer to some memory
+		if (size == 0)
+		{
+			m_data = (u8*)"";
+			return true;
+		}
+
 		m_mappingHandle = uba::CreateFileMappingW(m_logger, m_fileHandle, PAGE_READWRITE, size, realFileName);
 		if (!m_mappingHandle.IsValid())
 			return m_logger.Error(TC("Failed to create memory map %s with size %llu (%s)"), realFileName, size, LastErrorToText().data);
