[UBA]
* Added UbaWinTypes.h

diff --git a/Engine/Source/Programs/UnrealBuildAccelerator/Common/Public/UbaDirectoryIterator.h b/Engine/Source/Programs/UnrealBuildAccelerator/Common/Public/UbaDirectoryIterator.h
--- a/Engine/Source/Programs/UnrealBuildAccelerator/Common/Public/UbaDirectoryIterator.h
+++ b/Engine/Source/Programs/UnrealBuildAccelerator/Common/Public/UbaDirectoryIterator.h
@@ -34,41 +34,9 @@ namespace uba
 
 
 #if PLATFORM_WINDOWS
-#include <winternl.h>
-
-struct FILE_DIRECTORY_INFORMATION {
-  ULONG         NextEntryOffset;
-  ULONG         FileIndex;
-  LARGE_INTEGER CreationTime;
-  LARGE_INTEGER LastAccessTime;
-  LARGE_INTEGER LastWriteTime;
-  LARGE_INTEGER ChangeTime;
-  LARGE_INTEGER EndOfFile;
-  LARGE_INTEGER AllocationSize;
-  ULONG         FileAttributes;
-  ULONG         FileNameLength;
-  WCHAR         FileName[1];
-};
-extern "C" NTSTATUS NtQueryDirectoryFile(HANDLE FileHandle, HANDLE Event, PIO_APC_ROUTINE ApcRoutine, PVOID ApcContext, PIO_STATUS_BLOCK IoStatusBlock, PVOID FileInformation, ULONG Length, FILE_INFORMATION_CLASS FileInformationClass, BOOLEAN ReturnSingleEntry, PUNICODE_STRING FileName, BOOLEAN RestartScan);
+#include "UbaWinTypes.h"
 
-constexpr int FileIdBothDirectoryInformation = 37;
-struct FILE_ID_BOTH_DIR_INFORMATION {
-  ULONG         NextEntryOffset;
-  ULONG         FileIndex;
-  LARGE_INTEGER CreationTime;
-  LARGE_INTEGER LastAccessTime;
-  LARGE_INTEGER LastWriteTime;
-  LARGE_INTEGER ChangeTime;
-  LARGE_INTEGER EndOfFile;
-  LARGE_INTEGER AllocationSize;
-  ULONG         FileAttributes;
-  ULONG         FileNameLength;
-  ULONG         EaSize;
-  CCHAR         ShortNameLength;
-  WCHAR         ShortName[12];
-  LARGE_INTEGER FileId;
-  WCHAR         FileName[1];
-};
+extern "C" NTSTATUS NtQueryDirectoryFile(HANDLE FileHandle, HANDLE Event, PIO_APC_ROUTINE ApcRoutine, PVOID ApcContext, PIO_STATUS_BLOCK IoStatusBlock, PVOID FileInformation, ULONG Length, FILE_INFORMATION_CLASS FileInformationClass, BOOLEAN ReturnSingleEntry, PUNICODE_STRING FileName, BOOLEAN RestartScan);
 
 namespace uba
 {
@@ -149,7 +117,7 @@ namespace uba
 		while (true)
 		{
 			startTime = GetTime();
-			res = NtQueryDirectoryFile(handle, 0, NULL, NULL, &IoStatusBlock, buff, sizeof(buff) - 2, (FILE_INFORMATION_CLASS)FileIdBothDirectoryInformation, FALSE, NULL, FALSE);
+			res = NtQueryDirectoryFile(handle, 0, NULL, NULL, &IoStatusBlock, buff, sizeof(buff) - 2, FileIdBothDirectoryInformation, FALSE, NULL, FALSE);
 			timer.time += GetTime() - startTime;
 
 			if (res != STATUS_SUCCESS)
diff --git a/Engine/Source/Programs/UnrealBuildAccelerator/Core/Public/UbaWinTypes.h b/Engine/Source/Programs/UnrealBuildAccelerator/Core/Public/UbaWinTypes.h
new file mode 100644
--- /dev/null
+++ b/Engine/Source/Programs/UnrealBuildAccelerator/Core/Public/UbaWinTypes.h
@@ -0,0 +1,68 @@
+// Copyright Epic Games, Inc. All Rights Reserved.
+
+#pragma once
+
+#include <winternl.h>
+
+struct FILE_DIRECTORY_INFORMATION {
+	ULONG         NextEntryOffset;
+	ULONG         FileIndex;
+	LARGE_INTEGER CreationTime;
+	LARGE_INTEGER LastAccessTime;
+	LARGE_INTEGER LastWriteTime;
+	LARGE_INTEGER ChangeTime;
+	LARGE_INTEGER EndOfFile;
+	LARGE_INTEGER AllocationSize;
+	ULONG         FileAttributes;
+	ULONG         FileNameLength;
+	WCHAR         FileName[1];
+};
+
+struct FILE_FULL_DIR_INFORMATION {
+	ULONG         NextEntryOffset;
+	ULONG         FileIndex;
+	LARGE_INTEGER CreationTime;
+	LARGE_INTEGER LastAccessTime;
+	LARGE_INTEGER LastWriteTime;
+	LARGE_INTEGER ChangeTime;
+	LARGE_INTEGER EndOfFile;
+	LARGE_INTEGER AllocationSize;
+	ULONG         FileAttributes;
+	ULONG         FileNameLength;
+	ULONG         EaSize;
+	WCHAR         FileName[1];
+};
+
+struct FILE_ID_BOTH_DIR_INFORMATION {
+	ULONG         NextEntryOffset;
+	ULONG         FileIndex;
+	LARGE_INTEGER CreationTime;
+	LARGE_INTEGER LastAccessTime;
+	LARGE_INTEGER LastWriteTime;
+	LARGE_INTEGER ChangeTime;
+	LARGE_INTEGER EndOfFile;
+	LARGE_INTEGER AllocationSize;
+	ULONG         FileAttributes;
+	ULONG         FileNameLength;
+	ULONG         EaSize;
+	CCHAR         ShortNameLength;
+	WCHAR         ShortName[12];
+	LARGE_INTEGER FileId;
+	WCHAR         FileName[1];
+};
+
+struct FILE_BASIC_INFORMATION {
+	LARGE_INTEGER CreationTime;
+	LARGE_INTEGER LastAccessTime;
+	LARGE_INTEGER LastWriteTime;
+	LARGE_INTEGER ChangeTime;
+	ULONG FileAttributes;
+};
+
+constexpr FILE_INFORMATION_CLASS FileRenameInformation = (FILE_INFORMATION_CLASS)10;
+constexpr FILE_INFORMATION_CLASS FileDispositionInformation = (FILE_INFORMATION_CLASS)13;
+constexpr FILE_INFORMATION_CLASS FileIdBothDirectoryInformation = (FILE_INFORMATION_CLASS)37;
+constexpr FILE_INFORMATION_CLASS FileRenameInformationEx = (FILE_INFORMATION_CLASS)65;
+
+
+constexpr OBJECT_INFORMATION_CLASS ObjectNameInformation = (OBJECT_INFORMATION_CLASS)1;
\ No newline at end of file
diff --git a/Engine/Source/Programs/UnrealBuildAccelerator/Detours/Private/Windows/UbaDetoursFunctionsWin.h b/Engine/Source/Programs/UnrealBuildAccelerator/Detours/Private/Windows/UbaDetoursFunctionsWin.h
--- a/Engine/Source/Programs/UnrealBuildAccelerator/Detours/Private/Windows/UbaDetoursFunctionsWin.h
+++ b/Engine/Source/Programs/UnrealBuildAccelerator/Detours/Private/Windows/UbaDetoursFunctionsWin.h
@@ -3,7 +3,7 @@
 #pragma once
 
 #include "UbaPlatform.h"
-#include <winternl.h>
+#include "UbaWinTypes.h"
 #include <tlhelp32.h>
 #include <psapi.h>
 #include <io.h>
@@ -442,36 +442,6 @@ extern "C" {
 DETOURED_FUNCTIONS
 #undef DETOURED_FUNCTION
 
-
-struct FILE_DIRECTORY_INFORMATION {
-  ULONG         NextEntryOffset;
-  ULONG         FileIndex;
-  LARGE_INTEGER CreationTime;
-  LARGE_INTEGER LastAccessTime;
-  LARGE_INTEGER LastWriteTime;
-  LARGE_INTEGER ChangeTime;
-  LARGE_INTEGER EndOfFile;
-  LARGE_INTEGER AllocationSize;
-  ULONG         FileAttributes;
-  ULONG         FileNameLength;
-  WCHAR         FileName[1];
-};
-
-struct FILE_FULL_DIR_INFORMATION {
-	ULONG         NextEntryOffset;
-	ULONG         FileIndex;
-	LARGE_INTEGER CreationTime;
-	LARGE_INTEGER LastAccessTime;
-	LARGE_INTEGER LastWriteTime;
-	LARGE_INTEGER ChangeTime;
-	LARGE_INTEGER EndOfFile;
-	LARGE_INTEGER AllocationSize;
-	ULONG         FileAttributes;
-	ULONG         FileNameLength;
-	ULONG         EaSize;
-	WCHAR         FileName[1];
-};
-
 struct FILE_RENAME_INFORMATION {
     union {
         BOOLEAN ReplaceIfExists;  // FileRenameInformation
@@ -496,14 +466,6 @@ struct FILE_IS_REMOTE_DEVICE_INFORMATION {
 	WCHAR FileName[1];
 	};
 
-	struct FILE_BASIC_INFORMATION {
-	LARGE_INTEGER CreationTime;
-	LARGE_INTEGER LastAccessTime;
-	LARGE_INTEGER LastWriteTime;
-	LARGE_INTEGER ChangeTime;
-	DWORD FileAttributes;
-	};
-
 	struct FILE_STANDARD_INFORMATION {
 	LARGE_INTEGER AllocationSize;
 	LARGE_INTEGER EndOfFile;
