[UBA]
* Fixed so GetApplication message can handle relative paths
* Bumped session network version since we changed some network traffic

diff --git a/Engine/Source/Programs/UnrealBuildAccelerator/Common/Private/UbaSessionClient.cpp b/Engine/Source/Programs/UnrealBuildAccelerator/Common/Private/UbaSessionClient.cpp
--- a/Engine/Source/Programs/UnrealBuildAccelerator/Common/Private/UbaSessionClient.cpp
+++ b/Engine/Source/Programs/UnrealBuildAccelerator/Common/Private/UbaSessionClient.cpp
@@ -378,17 +378,16 @@ namespace uba
 		ProcessStartInfoHolder& startInfo = process.m_startInfo;
 		outRealWorkingDir = m_processWorkingDir.data;
 
-		const tchar* application = startInfo.application;
-		StringView applicationView = ToView(application);
+		StringView application = startInfo.applicationStr;
 		
-		if (applicationView.StartsWith(TC("ubacopy")))
+		if (application.StartsWith(TC("ubacopy")))
 			return true;
 
 		#if PLATFORM_LINUX
-		//if (applicationView.EndsWith(TStrlen(startInfo.application), "/sh")) // For some reason the local shell application on the helpers hang (need to investigate aws shell application)
+		//if (application.EndsWith(TStrlen(startInfo.application), "/sh")) // For some reason the local shell application on the helpers hang (need to investigate aws shell application)
 		//	return true;
 		#elif PLATFORM_WINDOWS
-		if (applicationView.Equals(TCV("cmd.exe")) || applicationView.EndsWith(TCV("\\cmd.exe"))) // We want to use local cmd.exe.. wine does not support kernel calls called by windows 11's cmd.exe for example
+		if (application.Equals(TCV("cmd.exe")) || application.EndsWith(TCV("\\cmd.exe"))) // We want to use local cmd.exe.. wine does not support kernel calls called by windows 11's cmd.exe for example
 		{
 			outRealApplication.Clear().Append(m_systemPath).Append(TCV("\\System32\\cmd.exe"));
 			return true;
@@ -400,7 +399,7 @@ namespace uba
 		bool isAbsolute = IsAbsolutePath(startInfo.application);
 
 		SCOPED_FUTEX(m_handledApplicationEnvironmentsLock, environmentslock);
-		auto insres = m_handledApplicationEnvironments.try_emplace(application);
+		auto insres = m_handledApplicationEnvironments.try_emplace(application.data);
 		environmentslock.Leave();
 
 		ApplicationEnvironment& appEnv = insres.first->second;
@@ -419,11 +418,11 @@ namespace uba
 
 		List<ModuleInfo> modules;
 		StringBuffer<512> virtualApplication;
-		if (!ReadModules(modules, virtualApplication, process.GetId(), application))
+		if (!ReadModules(modules, virtualApplication, process.GetId(), application, startInfo.workingDirStr))
 			return false;
 
 		StringBuffer<MaxPath> applicationDir;
-		applicationDir.AppendDir(application);
+		applicationDir.AppendDir(modules.begin()->name.c_str());
 		KeyToString keyStr(ToStringKeyLower(applicationDir));
 
 		Atomic<bool> success = true;
@@ -454,7 +453,7 @@ namespace uba
 		if (!success)
 			return false;
 		
-		outRealApplication.Append(m_sessionBinDir).Append(keyStr).Append(PathSeparator).AppendFileName(application);
+		outRealApplication.Append(m_sessionBinDir).Append(keyStr).Append(PathSeparator).AppendFileName(application.data);
 		appEnv.realApplication = outRealApplication.data;
 
 		if (!isAbsolute)
@@ -467,16 +466,17 @@ namespace uba
 		return true;
 	}
 
-	bool SessionClient::ReadModules(List<ModuleInfo>& outModules, StringBufferBase& outVirtualApplication, u32 processId, const tchar* application)
+	bool SessionClient::ReadModules(List<ModuleInfo>& outModules, StringBufferBase& outVirtualApplication, u32 processId, StringView application, StringView workingDir)
 	{
 		TrackWorkScope tws(m_client, AsView(TC("ReadModules")), ColorWork);
 
-		StackBinaryReader<16*1024> reader;
+		StackBinaryReader<SendMaxSize> reader;
 		{
-			StackBinaryWriter<256> writer;
+			StackBinaryWriter<1024> writer;
 			NetworkMessage msg(m_client, ServiceId, SessionMessageType_GetApplication, writer);
 			writer.WriteBool(IsLinux);
 			writer.WriteString(application);
+			writer.WriteString(workingDir);
 			if (!msg.Send(reader, m_stats.getApplicationMsg))
 				return false;
 		}
@@ -484,7 +484,7 @@ namespace uba
 		u32 serverSystemPathLen = reader.ReadU32();
 		u32 moduleCount = reader.ReadU32();
 		if (moduleCount == 0)
-			return m_logger.Error(TC("Application %s not found"), application);
+			return m_logger.Error(TC("Application %s not found"), application.data);
 
 		while (moduleCount--)
 		{
@@ -855,7 +855,7 @@ namespace uba
 
 		if (!sendDelete)
 		{
-			if (!m_allowMemoryMaps && isTemp)
+			if (isTemp && (!m_allowMemoryMaps || !msg.process.m_startInfo.rules->KeepInMemory(msg.fileName, m_tempPath, true, true)))
 			{
 				out.result = uba::DeleteFileW(msg.fileName.data);
 				out.errorCode = GetLastError();
diff --git a/Engine/Source/Programs/UnrealBuildAccelerator/Common/Private/UbaSessionServer.cpp b/Engine/Source/Programs/UnrealBuildAccelerator/Common/Private/UbaSessionServer.cpp
--- a/Engine/Source/Programs/UnrealBuildAccelerator/Common/Private/UbaSessionServer.cpp
+++ b/Engine/Source/Programs/UnrealBuildAccelerator/Common/Private/UbaSessionServer.cpp
@@ -1231,12 +1231,22 @@ namespace uba
 		bool isClientLinux = reader.ReadBool();
 		StringBuffer<> applicationName;
 		reader.ReadString(applicationName);
+		StringBuffer<> workingDir;
+		reader.ReadString(workingDir);
 
 		#if UBA_TEST_LINUX_HELPING_WINDOWS
 		if (isClientLinux && applicationName.Equals(TCV("/usr/bin/gcc")))
 			applicationName.Clear().Append(TCV("d:\\temp\\gcc"));
 		#endif
 
+		if (!IsAbsolutePath(applicationName.data))
+		{
+			StringBuffer<> temp;
+			if (!SearchPathForFile(m_logger, temp, applicationName.data, workingDir, {}))
+				return false;
+			applicationName.Clear().Append(temp);
+		}
+
 		StringKey applicationKey = ToStringKeyLower(applicationName);
 
 		SCOPED_FUTEX(m_applicationDataLock, lock);
@@ -1483,6 +1493,7 @@ namespace uba
 			if (startInfo.trackInputs)
 			{
 				u64 bytes = GetStringWriteSize(destination.data, destination.count);
+				SCOPED_WRITE_LOCK(m_receivedFilesLock, lock); // Abuse of existsing lock (could use its own lock but not worth it)
 				u64 prevSize = process->m_trackedOutputs.size();
 				process->m_trackedOutputs.resize(prevSize + bytes);
 				BinaryWriter w2(process->m_trackedOutputs.data(), prevSize, prevSize + bytes);
@@ -1707,19 +1718,34 @@ namespace uba
 			RemoteProcess* process = DequeueProcess(session, sessionId, connectionInfo.GetId());
 			if (process == nullptr)
 				break;
-			auto& startInfo = process->m_startInfo;
 
-			StringBuffer<> applicationOverride;
-
-			if (isClientCrossArchitecture || isClientCrossPlatform)
-			{
-				auto returnProcess = [&]()
+			auto returnProcess = [&]()
 				{
 					m_queuedRemoteProcesses.push_front({process});
 					if (m_remoteProcessReturnedEvent)
 						m_remoteProcessReturnedEvent(*process);
 				};
-			
+
+			auto& startInfo = process->m_startInfo;
+
+			if (!IsAbsolutePath(startInfo.application))
+			{
+				StringBuffer<> temp;
+				if (!SearchPathForFile(m_logger, temp, startInfo.application, startInfo.workingDirStr, {}))
+				{
+					// TODO: This should be seen as a failed process
+					m_logger.Error(TC("BAD PATH TO APPLICATION %s"), startInfo.application);
+					returnProcess();
+					continue;
+				}
+				startInfo.applicationStr = temp.ToString();
+				startInfo.application = startInfo.applicationStr.data();
+			}
+
+			StringBuffer<> applicationOverride;
+
+			if (isClientCrossArchitecture || isClientCrossPlatform)
+			{
 				if ((!process->m_canCrossArchitecture && isClientCrossArchitecture) || (!process->m_canCrossPlatform && isClientCrossPlatform))
 				{
 					// There is a risk cross architecture client dequeues another client's process that is not cross architecture
diff --git a/Engine/Source/Programs/UnrealBuildAccelerator/Common/Public/UbaNetwork.h b/Engine/Source/Programs/UnrealBuildAccelerator/Common/Public/UbaNetwork.h
--- a/Engine/Source/Programs/UnrealBuildAccelerator/Common/Public/UbaNetwork.h
+++ b/Engine/Source/Programs/UnrealBuildAccelerator/Common/Public/UbaNetwork.h
@@ -16,7 +16,7 @@ namespace uba
 	
 	static constexpr u32 SystemNetworkVersion = 1339;
 	static constexpr u32 StorageNetworkVersion = 4;
-	static constexpr u32 SessionNetworkVersion = 49;
+	static constexpr u32 SessionNetworkVersion = 50;
 	static constexpr u32 CacheNetworkVersion = 5;
 	static constexpr u32 CacheBucketVersion = 3;
 
diff --git a/Engine/Source/Programs/UnrealBuildAccelerator/Common/Public/UbaSessionClient.h b/Engine/Source/Programs/UnrealBuildAccelerator/Common/Public/UbaSessionClient.h
--- a/Engine/Source/Programs/UnrealBuildAccelerator/Common/Public/UbaSessionClient.h
+++ b/Engine/Source/Programs/UnrealBuildAccelerator/Common/Public/UbaSessionClient.h
@@ -88,7 +88,7 @@ namespace uba
 		struct ModuleInfo;
 
 		bool GetCasKeyForFile(CasKey& out, u32 processId, const StringView& fileName, const StringKey& fileNameKey);
-		bool ReadModules(List<ModuleInfo>& outModules, StringBufferBase& outVirtualApplication, u32 processId, const tchar* application);
+		bool ReadModules(List<ModuleInfo>& outModules, StringBufferBase& outVirtualApplication, u32 processId, StringView application, StringView workingDir);
 		bool EnsureBinaryFile(StringBufferBase& out, StringBufferBase& outVirtual, u32 processId, StringView fileName, const StringKey& fileNameKey, StringView applicationDir, StringView workingDir, const u8* loaderPaths, u32 loaderPathsSize);
 		bool WriteBinFile(StringBufferBase& out, const StringView& binaryName, const CasKey& casKey, const KeyToString& applicationDir, u32 fileAttributes);
 		bool SendFiles(ProcessImpl& process, Timer& sendFiles);
