[UBA]
* Removed assert that shouldn't be there. SetHandleInformation calls detoured Nt versions

diff --git a/Engine/Source/Programs/UnrealBuildAccelerator/Detours/Private/Windows/UbaDetoursFunctionsKernelBase.inl b/Engine/Source/Programs/UnrealBuildAccelerator/Detours/Private/Windows/UbaDetoursFunctionsKernelBase.inl
--- a/Engine/Source/Programs/UnrealBuildAccelerator/Detours/Private/Windows/UbaDetoursFunctionsKernelBase.inl
+++ b/Engine/Source/Programs/UnrealBuildAccelerator/Detours/Private/Windows/UbaDetoursFunctionsKernelBase.inl
@@ -4197,7 +4197,6 @@ BOOL Detoured_SetHandleInformation(HANDLE hObject, DWORD dwMask, DWORD dwFlags)
 {
 	DETOURED_CALL(SetHandleInformation);
 	DEBUG_LOG_TRUE(L"SetHandleInformation", L"%llu", uintptr_t(hObject));
-	UBA_ASSERT(isDetouredHandle(hObject));
 	return True_SetHandleInformation(hObject, dwMask, dwFlags); // Calls NtQueryObject and NtSetInformationObject internally
 }
 
