[UBA]
* Changed AddWork to use move assignment for function

diff --git a/Engine/Source/Programs/UnrealBuildAccelerator/Common/Private/UbaNetworkServer.cpp b/Engine/Source/Programs/UnrealBuildAccelerator/Common/Private/UbaNetworkServer.cpp
--- a/Engine/Source/Programs/UnrealBuildAccelerator/Common/Private/UbaNetworkServer.cpp
+++ b/Engine/Source/Programs/UnrealBuildAccelerator/Common/Private/UbaNetworkServer.cpp
@@ -1061,9 +1061,9 @@ namespace uba
 		}
 	}
 
-	void NetworkServer::AddWork(const WorkFunction& work, u32 count, const tchar* desc, const Color& color, WorkPriority priority, Bottleneck* bottleneck)
+	void NetworkServer::AddWork(WorkFunction&& work, u32 count, const tchar* desc, const Color& color, WorkPriority priority, Bottleneck* bottleneck)
 	{
-		m_additionalWorkQueue.Enqueue(work, count, desc, color, priority, bottleneck, m_workTracker != nullptr);
+		m_additionalWorkQueue.Enqueue(std::move(work), count, desc, color, priority, bottleneck, m_workTracker != nullptr);
 
 		SCOPED_FUTEX(m_availableWorkersLock, lock2);
 		if (!m_workersEnabled)
diff --git a/Engine/Source/Programs/UnrealBuildAccelerator/Common/Private/UbaSession.cpp b/Engine/Source/Programs/UnrealBuildAccelerator/Common/Private/UbaSession.cpp
--- a/Engine/Source/Programs/UnrealBuildAccelerator/Common/Private/UbaSession.cpp
+++ b/Engine/Source/Programs/UnrealBuildAccelerator/Common/Private/UbaSession.cpp
@@ -2384,6 +2384,7 @@ namespace uba
 	Storage& Session::GetStorage() { return m_storage; }
 	MutableLogger& Session::GetLogger() { return m_logger; }
 	LogWriter& Session::GetLogWriter() { return m_logger.m_writer; }
+	WorkManager& Session::GetWorkManager() { return m_workManager; }
 	Trace& Session::GetTrace() { return m_trace; }
 
 	const ApplicationRules* Session::GetRules(const ProcessStartInfo& si)
diff --git a/Engine/Source/Programs/UnrealBuildAccelerator/Common/Private/UbaStorageUtils.cpp b/Engine/Source/Programs/UnrealBuildAccelerator/Common/Private/UbaStorageUtils.cpp
--- a/Engine/Source/Programs/UnrealBuildAccelerator/Common/Private/UbaStorageUtils.cpp
+++ b/Engine/Source/Programs/UnrealBuildAccelerator/Common/Private/UbaStorageUtils.cpp
@@ -115,7 +115,7 @@ namespace uba
 				workerCount = Min(workCount, workManager->GetWorkerCount()-1); // We are a worker ourselves
 				workerCount = Min(workerCount, MaxWorkItemsPerAction2); // Cap this to not starve other things
 				rec->refCount += workerCount;
-				workManager->AddWork(work, workerCount, TC("CalculateKey"));
+				workManager->AddWork(WorkFunction(work), workerCount, TC("CalculateKey"));
 			}
 
 			{
@@ -1208,7 +1208,7 @@ namespace uba
 		workerCount = Min(workerCount, MaxWorkItemsPerAction2);
 
 		rec.refCount = workerCount + 1; // We need to keep refcount up 1 to make sure it is not deleted before we read rec->written
-		m_workManager.AddWork(work, workerCount-1, TC("Compress")); // We are a worker ourselves
+		m_workManager.AddWork(WorkFunction(work), workerCount-1, TC("Compress")); // We are a worker ourselves
 		{
 			TrackWorkScope tws;
 			work({tws});
@@ -1447,7 +1447,7 @@ namespace uba
 			u32 workerCount = Min(workCount, m_workManager.GetWorkerCount() - 1); // We are a worker ourselves
 			workerCount = Min(workerCount, MaxWorkItemsPerAction2); // Cap this to not starve other things
 			rec->refCount += workerCount;
-			m_workManager.AddWork(work, workerCount, TC("DecompressMemToMem"));
+			m_workManager.AddWork(WorkFunction(work), workerCount, TC("DecompressMemToMem"));
 
 			TimerScope ts(stats.decompressToMem);
 			{
diff --git a/Engine/Source/Programs/UnrealBuildAccelerator/Common/Private/UbaWorkManager.cpp b/Engine/Source/Programs/UnrealBuildAccelerator/Common/Private/UbaWorkManager.cpp
--- a/Engine/Source/Programs/UnrealBuildAccelerator/Common/Private/UbaWorkManager.cpp
+++ b/Engine/Source/Programs/UnrealBuildAccelerator/Common/Private/UbaWorkManager.cpp
@@ -82,10 +82,11 @@ namespace uba
 		}
 	}
 
-	void WorkQueue::Enqueue(const WorkFunction& func, u32 count, const tchar* desc, const Color& color, WorkPriority priority, Bottleneck* bottleneck, bool trackWork)
+	void WorkQueue::Enqueue(WorkFunction&& func, u32 count, const tchar* desc, const Color& color, WorkPriority priority, Bottleneck* bottleneck, bool trackWork)
 	{
 		UBA_ASSERT(*desc);
 		SCOPED_FUTEX(lock, lk);
+		u32 lastIndex = count - 1;
 		for (u32 i = 0; i != count; ++i)
 		{
 			Work* work;
@@ -98,7 +99,12 @@ namespace uba
 				work = new Work();
 		
 			work->next = nullptr;
-			work->func = func;
+
+			if (i == lastIndex)
+				work->func = std::move(func);
+			else
+				work->func = func;
+
 			work->bottleneck = bottleneck;
 			if (trackWork)
 				work->desc = desc;
@@ -268,9 +274,9 @@ namespace uba
 	}
 
 
-	void WorkManagerImpl::AddWork(const WorkFunction& work, u32 count, const tchar* desc, const Color& color, WorkPriority priority, Bottleneck* bottleneck)
+	void WorkManagerImpl::AddWork(WorkFunction&& work, u32 count, const tchar* desc, const Color& color, WorkPriority priority, Bottleneck* bottleneck)
 	{
-		m_workQueue.Enqueue(work, count, desc, color, priority, bottleneck, m_workTracker.load());
+		m_workQueue.Enqueue(std::move(work), count, desc, color, priority, bottleneck, m_workTracker.load());
 
 		SCOPED_FUTEX(m_availableWorkersLock, lock2);
 		while (count--)
diff --git a/Engine/Source/Programs/UnrealBuildAccelerator/Common/Public/UbaDependencyCrawler.h b/Engine/Source/Programs/UnrealBuildAccelerator/Common/Public/UbaDependencyCrawler.h
--- a/Engine/Source/Programs/UnrealBuildAccelerator/Common/Public/UbaDependencyCrawler.h
+++ b/Engine/Source/Programs/UnrealBuildAccelerator/Common/Public/UbaDependencyCrawler.h
@@ -815,7 +815,7 @@ namespace uba
 						if (parseDefines)
 							work({tracker});
 						else
-							m_workManager.AddWork(work, 1, TC("CrawlIncludes"), ColorWork, WorkPriority_Low, &m_bottleneck);
+							m_workManager.AddWork(std::move(work), 1, TC("CrawlIncludes"), ColorWork, WorkPriority_Low, &m_bottleneck);
 					}
 				}
 				else if (parseDefines && strncmp(it, "define", 6) == 0)
diff --git a/Engine/Source/Programs/UnrealBuildAccelerator/Common/Public/UbaNetworkServer.h b/Engine/Source/Programs/UnrealBuildAccelerator/Common/Public/UbaNetworkServer.h
--- a/Engine/Source/Programs/UnrealBuildAccelerator/Common/Public/UbaNetworkServer.h
+++ b/Engine/Source/Programs/UnrealBuildAccelerator/Common/Public/UbaNetworkServer.h
@@ -95,7 +95,7 @@ namespace uba
 		//void RegisterOnConnection(u8 id, const OnConnectionFunction& func);
 		//void UnregisterOnConnection(u8 id);
 
-		virtual void AddWork(const WorkFunction& work, u32 count, const tchar* desc, const Color& color = ColorWork, WorkPriority priority = WorkPriority_Normal, Bottleneck* bottleneck = nullptr) override final;
+		virtual void AddWork(WorkFunction&& work, u32 count, const tchar* desc, const Color& color = ColorWork, WorkPriority priority = WorkPriority_Normal, Bottleneck* bottleneck = nullptr) override final;
 		virtual void DoWork(u32 count = 1) override final;
 		virtual u32 GetWorkerCount() override final;
 		virtual bool FlushWork(u32 timeoutMs = 0) override final;
diff --git a/Engine/Source/Programs/UnrealBuildAccelerator/Common/Public/UbaSession.h b/Engine/Source/Programs/UnrealBuildAccelerator/Common/Public/UbaSession.h
--- a/Engine/Source/Programs/UnrealBuildAccelerator/Common/Public/UbaSession.h
+++ b/Engine/Source/Programs/UnrealBuildAccelerator/Common/Public/UbaSession.h
@@ -91,6 +91,7 @@ namespace uba
 		Storage& GetStorage(); // Storage (only used when remote machines are connected)
 		MutableLogger& GetLogger(); // Logger used for logging 
 		LogWriter& GetLogWriter(); // LogWriter used by logger
+		WorkManager& GetWorkManager();
 		Trace& GetTrace(); // Trace written to be session
 		const ApplicationRules* GetRules(const ProcessStartInfo& si); // Get application rules used for process
 		const tchar* GetTempPath(); // Path for temp files used for current session
diff --git a/Engine/Source/Programs/UnrealBuildAccelerator/Common/Public/UbaWorkManager.h b/Engine/Source/Programs/UnrealBuildAccelerator/Common/Public/UbaWorkManager.h
--- a/Engine/Source/Programs/UnrealBuildAccelerator/Common/Public/UbaWorkManager.h
+++ b/Engine/Source/Programs/UnrealBuildAccelerator/Common/Public/UbaWorkManager.h
@@ -57,7 +57,7 @@ namespace uba
 	class WorkManager : public WorkTracker
 	{
 	public:
-		virtual void AddWork(const WorkFunction& work, u32 count, const tchar* desc, const Color& color = ColorWork, WorkPriority priority = WorkPriority_Normal, Bottleneck* bottleneck = nullptr) = 0;
+		virtual void AddWork(WorkFunction&& work, u32 count, const tchar* desc, const Color& color = ColorWork, WorkPriority priority = WorkPriority_Normal, Bottleneck* bottleneck = nullptr) = 0;
 		virtual u32 GetWorkerCount() = 0;
 		virtual void DoWork(u32 count = 1) = 0;
 		virtual bool FlushWork(u32 timeoutMs = 0) = 0;
@@ -83,7 +83,7 @@ namespace uba
 	struct WorkQueue
 	{
 		~WorkQueue();
-		void Enqueue(const WorkFunction& func, u32 count, const tchar* desc, const Color& color, WorkPriority priority, Bottleneck* bottleneck, bool trackWork);
+		void Enqueue(WorkFunction&& func, u32 count, const tchar* desc, const Color& color, WorkPriority priority, Bottleneck* bottleneck, bool trackWork);
 		Work* Dequeue(BottleneckTicket& ticket);
 		void Finish(Work* work, BottleneckTicket& ticket);
 		bool IsEmpty();
@@ -103,7 +103,7 @@ namespace uba
 	public:
 		WorkManagerImpl(u32 maxWorkerCount, const tchar* workerDesc = TC("UbaWrk"), u32 precreateWorkerCount = 1);
 		virtual ~WorkManagerImpl();
-		virtual void AddWork(const WorkFunction& work, u32 count, const tchar* desc, const Color& color = ColorWork, WorkPriority priority = WorkPriority_Normal, Bottleneck* bottleneck = nullptr) override final;
+		virtual void AddWork(WorkFunction&& work, u32 count, const tchar* desc, const Color& color = ColorWork, WorkPriority priority = WorkPriority_Normal, Bottleneck* bottleneck = nullptr) override final;
 		virtual u32 GetWorkerCount() override final;
 		virtual void DoWork(u32 count = 1) override final;
 		virtual bool FlushWork(u32 timeoutMs = 0) override final;
@@ -237,7 +237,7 @@ namespace uba
 				}
 			};
 		
-		AddWork(work, workCount, description.data, ColorWork, priority, bottleneck);
+		AddWork(std::move(work), workCount, description.data, ColorWork, priority, bottleneck);
 
 		{
 			#if UBA_TRACK_WORK
