[UBA]
* Attempt at fixing mac remote crashes that I can't repro

diff --git a/Engine/Source/Programs/UnrealBuildAccelerator/Common/Private/UbaFile.cpp b/Engine/Source/Programs/UnrealBuildAccelerator/Common/Private/UbaFile.cpp
--- a/Engine/Source/Programs/UnrealBuildAccelerator/Common/Private/UbaFile.cpp
+++ b/Engine/Source/Programs/UnrealBuildAccelerator/Common/Private/UbaFile.cpp
@@ -502,7 +502,7 @@ namespace uba
 
 		StringBuffer<> fullPath;
 
-		StringView fileView = ToView(file);
+		StringView fileView = file ? ToView(file) : {};
 		StringView ext;
 
 		#if PLATFORM_WINDOWS
