[UBA]
* Fixed clang compile error

diff --git a/Engine/Source/Programs/UnrealBuildAccelerator/Common/Private/UbaFile.cpp b/Engine/Source/Programs/UnrealBuildAccelerator/Common/Private/UbaFile.cpp
--- a/Engine/Source/Programs/UnrealBuildAccelerator/Common/Private/UbaFile.cpp
+++ b/Engine/Source/Programs/UnrealBuildAccelerator/Common/Private/UbaFile.cpp
@@ -502,7 +502,7 @@ namespace uba
 
 		StringBuffer<> fullPath;
 
-		StringView fileView = file ? ToView(file) : {};
+		StringView fileView = ToView(file ? file : TC(""));
 		StringView ext;
 
 		#if PLATFORM_WINDOWS
