[UBA]
* Another attempt at trying to fix remote mac crashes that I can't repro

diff --git a/Engine/Source/Programs/UnrealBuildAccelerator/Common/Private/UbaSessionClient.cpp b/Engine/Source/Programs/UnrealBuildAccelerator/Common/Private/UbaSessionClient.cpp
--- a/Engine/Source/Programs/UnrealBuildAccelerator/Common/Private/UbaSessionClient.cpp
+++ b/Engine/Source/Programs/UnrealBuildAccelerator/Common/Private/UbaSessionClient.cpp
@@ -378,7 +378,7 @@ namespace uba
 		ProcessStartInfoHolder& startInfo = process.m_startInfo;
 		outRealWorkingDir = m_processWorkingDir.data;
 
-		StringView application = startInfo.applicationStr;
+		StringView application = ToView(startInfo.application); 
 		
 		if (application.StartsWith(TC("ubacopy")))
 			return true;
