[UBA]
* Added validation of time stamp of deferred source files
* Added handling of when remote upload cas and there is a deferred source that doesn't exist anymore

diff --git a/Engine/Source/Programs/UnrealBuildAccelerator/Common/Private/UbaCacheClient.cpp b/Engine/Source/Programs/UnrealBuildAccelerator/Common/Private/UbaCacheClient.cpp
--- a/Engine/Source/Programs/UnrealBuildAccelerator/Common/Private/UbaCacheClient.cpp
+++ b/Engine/Source/Programs/UnrealBuildAccelerator/Common/Private/UbaCacheClient.cpp
@@ -1357,7 +1357,7 @@ namespace uba
 			TimerScope createCasTs(stats.createCas);
 			StorageImpl::CasEntry* casEntry;
 
-			if (m_storage.HasCasFile(casKey, &casEntry, m_useMappingForCreatedCas))
+			if (m_storage.HasCasFile(casKey, &casEntry, m_useMappingForCreatedCas, false))
 			{
 				createCasTs.Leave();
 
diff --git a/Engine/Source/Programs/UnrealBuildAccelerator/Common/Private/UbaStorage.cpp b/Engine/Source/Programs/UnrealBuildAccelerator/Common/Private/UbaStorage.cpp
--- a/Engine/Source/Programs/UnrealBuildAccelerator/Common/Private/UbaStorage.cpp
+++ b/Engine/Source/Programs/UnrealBuildAccelerator/Common/Private/UbaStorage.cpp
@@ -401,7 +401,7 @@ namespace uba
 		return w.CompressToFile(out.storedSize, from, readHandle, readMem, fileSize, to, header, headerSize, lastWriteTime, m_tempPath.data);
 	}
 
-	bool StorageImpl::WriteCasFileNoCheck(WriteCasResult& out, const StringKey& fileNameKey, const tchar* fileName, const CasKey& casKey, const tchar* casFile, bool storeCompressed, bool useMapping)
+	bool StorageImpl::WriteCasFileNoCheck(WriteCasResult& out, const StringKey& fileNameKey, const tchar* fileName, const CasKey& casKey, const tchar* casFile, bool storeCompressed, bool useMapping, DeferredSource* deferredSource, bool allowFail)
 	{
 		StorageStats& stats = Stats();
 		TimerScope ts(stats.createCas);
@@ -434,33 +434,33 @@ namespace uba
 			}
 		}
 
-		// Returns the number of failed open file attempts. If this is equal to maxOpenFileAttempts no file has been opened successfully.
-		auto TryOpenFileSequentialRead = [this](const tchar* fileName, FileHandle& outHandle, int maxOpenFileAttempts) -> int
-			{
-				UBA_ASSERT(maxOpenFileAttempts > 0);
-				int numFailedOpenFileAttempts = 0;
-				while (!OpenFileSequentialRead(m_logger, fileName, outHandle) && ++numFailedOpenFileAttempts < maxOpenFileAttempts)
-				{
-					// Delay next attempt for 100 milliseconds to open the file in case of a race condition
-					Sleep(100);
-				}
-				return numFailedOpenFileAttempts;
-			};
-
 		FileHandle readHandle;
 		constexpr int maxOpenFileAttempts = 5;
-		if (int numFailedOpenFileAttempts = TryOpenFileSequentialRead(fileName, readHandle, maxOpenFileAttempts))
+		int numFailedOpenFileAttempts = 0;
+
+		while (!OpenFileSequentialRead(m_logger, fileName, readHandle) && ++numFailedOpenFileAttempts < maxOpenFileAttempts)
+			Sleep(100); // Delay next attempt for 100 milliseconds to open the file in case of a race condition
+
+		if (numFailedOpenFileAttempts == maxOpenFileAttempts)
+			return m_logger.Logf(allowFail?LogEntryType_Info:LogEntryType_Error, TC("[WriteCasFileNoCheck] Failed to open file %s for read (%s)"), fileName, LastErrorToText().data).ToFalse();
+
+		auto fileGuard = MakeGuard([&](){ CloseFile(fileName, readHandle); });
+
+		if (numFailedOpenFileAttempts)
+			m_logger.Warning(TC("[WriteCasFileNoCheck] First attempts (%d) to open file %s for read failed; This may indicate a race condition"), numFailedOpenFileAttempts, fileName);
+
+		FileBasicInformation basicInfo;
+		if (!GetFileBasicInformationByHandle(basicInfo, m_logger, fileName, readHandle, false))
+			return m_logger.Error(TC("[WriteCasFileNoCheck] GetFileBasicInformationByHandle failed for %s (%s)"), fileName, LastErrorToText().data);
+	
+		if (deferredSource && deferredSource->lastWriteTime != basicInfo.lastWriteTime)
 		{
-			if (numFailedOpenFileAttempts == maxOpenFileAttempts)
-				return m_logger.Error(TC("[WriteCasFileNoCheck] Failed to open file %s for read (%s)"), fileName, LastErrorToText().data);
-			else
-				m_logger.Warning(TC("[WriteCasFileNoCheck] First attempts (%d) to open file %s for read failed; This may indicate a race condition"), numFailedOpenFileAttempts, fileName);
+			m_logger.Warning(TC("[WriteCasFileNoCheck] Last write time has changed since %s was deferred for cas creation."), fileName);
+			if (allowFail)
+				return false;
 		}
-		auto fileGuard = MakeGuard([&](){ CloseFile(fileName, readHandle); });
+		u64 readFileSize = basicInfo.size;
 
-		u64 readFileSize = 0;
-		if (!uba::GetFileSizeEx(readFileSize, readHandle))
-			return m_logger.Error(TC("[WriteCasFileNoCheck] GetFileSize failed for %s (%s)"), fileName, LastErrorToText().data);
 		out.sourceStoredSize = readFileSize;
 
 		bool readIsCompressed = false;
@@ -545,7 +545,7 @@ namespace uba
 		return true;
 	}
 
-	bool StorageImpl::WriteCasFile(WriteCasResult& out, const tchar* fileName, const CasKey& casKey, bool useMapping)
+	bool StorageImpl::WriteCasFile(WriteCasResult& out, const tchar* fileName, const CasKey& casKey, bool useMapping, DeferredSource* deferredSource, bool allowFail)
 	{
 		UBA_ASSERT(IsCompressed(casKey) == m_storeCompressed);
 		StringBuffer<> casFile;
@@ -559,7 +559,7 @@ namespace uba
 		if (CaseInsensitiveFs)
 			temp.MakeLower();
 		StringKey fileNameKey = ToStringKey(temp);
-		return WriteCasFileNoCheck(out, fileNameKey, fileName, casKey, casFile.data, IsCompressed(casKey), useMapping);
+		return WriteCasFileNoCheck(out, fileNameKey, fileName, casKey, casFile.data, IsCompressed(casKey), useMapping, deferredSource, allowFail);
 	}
 
 	void StorageImpl::CasEntryAccessed(const CasKey& casKey)
@@ -765,7 +765,7 @@ namespace uba
 		casEntry.exists = false;
 
 		WriteCasResult res;
-		if (!WriteCasFileNoCheck(res, fileNameKey, fileName, casKey, casFile.data, IsCompressed(casKey), false))
+		if (!WriteCasFileNoCheck(res, fileNameKey, fileName, casKey, casFile.data, IsCompressed(casKey), false, nullptr, false))
 			return false;
 		casEntry.memoryHandle = res.memoryHandle;
 		casEntry.memoryOffset = res.memoryOffset;
@@ -2120,7 +2120,7 @@ namespace uba
 		return false;
 	}
 
-	bool StorageImpl::HasCasFile(const CasKey& casKey, CasEntry** out, bool useMapping)
+	bool StorageImpl::HasCasFile(const CasKey& casKey, CasEntry** out, bool useMapping, bool allowFail)
 	{
 		SCOPED_READ_LOCK(m_casLookupLock, lookupLock);
 		auto it = m_casLookup.find(casKey);
@@ -2190,8 +2190,8 @@ namespace uba
 
 		casEntry.verified = true;
 		WriteCasResult res;
-		if (!WriteCasFile(res, deferredCreation.data, casKey, useMapping))
-			return m_logger.Error(TC("Failed to write deferred cas %s (%s)"), CasKeyString(casKey).str, deferredCreation.data);
+		if (!WriteCasFile(res, deferredCreation.data, casKey, useMapping, &sourceIt->second, allowFail))
+			return m_logger.Logf(allowFail?LogEntryType_Info:LogEntryType_Error, TC("Failed to write deferred cas %s (%s)"), CasKeyString(casKey).str, deferredCreation.data).ToFalse();
 
 		dccGuard.Execute();
 
@@ -2272,7 +2272,7 @@ namespace uba
 		if (!fileName)
 			return false;
 		WriteCasResult res;
-		if (!WriteCasFile(res, fileName, casKey, false))
+		if (!WriteCasFile(res, fileName, casKey, false, nullptr, false))
 			return false;
 		casEntry.exists = true;
 		casEntryLock.Leave();
@@ -2470,7 +2470,7 @@ namespace uba
 		while (true)
 		{
 			CasEntry* casEntryPtr = nullptr;
-			if (!HasCasFile(actualKey, &casEntryPtr))
+			if (!HasCasFile(actualKey, &casEntryPtr, false, false))
 			{
 				if (!testCompressed)
 				{
@@ -2798,7 +2798,7 @@ namespace uba
 		CasKey key = findIt->second.key;
 		m_logger.Logf(IsDebug ? LogEntryType_Warning : LogEntryType_Info, TC("Flush cas creation of file %s at %s (This should be rare)"), CasKeyString(key).str, fileName);
 		deferredLock.Leave();
-		HasCasFile(key);
+		HasCasFile(key, nullptr, false, false);
 	}
 	/*
 	void StorageImpl::PushStats(StorageStats* stats)
diff --git a/Engine/Source/Programs/UnrealBuildAccelerator/Common/Private/UbaStorageClient.cpp b/Engine/Source/Programs/UnrealBuildAccelerator/Common/Private/UbaStorageClient.cpp
--- a/Engine/Source/Programs/UnrealBuildAccelerator/Common/Private/UbaStorageClient.cpp
+++ b/Engine/Source/Programs/UnrealBuildAccelerator/Common/Private/UbaStorageClient.cpp
@@ -896,7 +896,7 @@ namespace uba
 		return true;
 	}
 
-	bool StorageClient::HasCasFile(const CasKey& casKey, CasEntry** out, bool useMapping)
+	bool StorageClient::HasCasFile(const CasKey& casKey, CasEntry** out, bool useMapping, bool allowFail)
 	{
 		CasKey localKey = AsCompressed(casKey, false);
 		SCOPED_READ_LOCK(m_localStorageFilesLock, lock);
@@ -915,7 +915,7 @@ namespace uba
 			}
 		}
 		lock.Leave();
-		return StorageImpl::HasCasFile(casKey, out);
+		return StorageImpl::HasCasFile(casKey, out, useMapping, allowFail);
 	}
 
 	bool StorageClient::StoreCasFileClient(CasKey& out, StringKey fileNameKey, const tchar* fileName, SharedMemoryHandle memoryHandle, u64 memoryOffset, u64 fileSize, const tchar* hint, bool keepMappingInMemory, bool storeCompressed)
diff --git a/Engine/Source/Programs/UnrealBuildAccelerator/Common/Private/UbaStorageServer.cpp b/Engine/Source/Programs/UnrealBuildAccelerator/Common/Private/UbaStorageServer.cpp
--- a/Engine/Source/Programs/UnrealBuildAccelerator/Common/Private/UbaStorageServer.cpp
+++ b/Engine/Source/Programs/UnrealBuildAccelerator/Common/Private/UbaStorageServer.cpp
@@ -108,7 +108,7 @@ namespace uba
 					m_waitEntries.erase(actualKey);
 			});
 
-		if (HasCasFile(actualKey))
+		if (HasCasFile(actualKey, nullptr, false, false))
 			return true;
 
 		u64 startTime = GetTime();
@@ -604,7 +604,8 @@ namespace uba
 
 
 				CasEntry* casEntryPtr = nullptr;
-				bool has = HasCasFile(casKey, &casEntryPtr); // HasCasFile writes deferred cas entries if in queue. Also increase CasEntry refcount
+				bool allowFail = false;
+				bool has = HasCasFile(casKey, &casEntryPtr, false, allowFail); // HasCasFile writes deferred cas entries if in queue. Also increase CasEntry refcount
 
 				todoRemoveMe(3);
 
@@ -1013,7 +1014,8 @@ namespace uba
 					{
 						deferredLock.Leave();
 						casEntryLock.Leave();
-						if (HasCasFile(casKey))
+						bool allowFail = true; // We can afford to fail here since we are about to receive the cas file
+						if (HasCasFile(casKey, nullptr, false, allowFail))
 						{
 							writer.WriteU16(u16(~0));
 							writer.WriteBool(m_traceStore);
diff --git a/Engine/Source/Programs/UnrealBuildAccelerator/Common/Public/UbaStorage.h b/Engine/Source/Programs/UnrealBuildAccelerator/Common/Public/UbaStorage.h
--- a/Engine/Source/Programs/UnrealBuildAccelerator/Common/Public/UbaStorage.h
+++ b/Engine/Source/Programs/UnrealBuildAccelerator/Common/Public/UbaStorage.h
@@ -199,9 +199,10 @@ namespace uba
 
 		struct CasEntry;
 		struct CasEntryWriteLock;
+		struct DeferredSource;
 
 		virtual bool StoreCasFile(CasKey& out, const tchar* fileName, const CasKey& casKeyOverride, bool deferCreation, CasEntry** outEntry);
-		virtual bool HasCasFile(const CasKey& casKey, CasEntry** out = nullptr, bool useMapping = false);
+		virtual bool HasCasFile(const CasKey& casKey, CasEntry** out, bool useMapping, bool allowFail);
 
 		bool EnsureCasFile(const CasKey& casKey, const tchar* fileName);
 		CasKey CalculateCasKey(const tchar* fileName, FileHandle fileHandle, u64 fileSize, bool storeCompressed);
@@ -222,8 +223,8 @@ namespace uba
 		virtual void RegisterExternalFileMappingsProvider(ExternalFileMappingsProvider&& provider) override;
 		virtual void DeleteExternalFileMapping(StringKey fileNameKey, const tchar* fileName) override;
 		virtual bool WriteCompressedFile(WriteResult& out, const tchar* from, FileHandle readHandle, u8* readMem, u64 fileSize, const tchar* toFile, const void* header, u64 headerSize, u64 lastWriteTime = 0) final;
-		bool WriteCasFileNoCheck(WriteCasResult& out, const StringKey& fileNameKey, const tchar* fileName, const CasKey& casKey, const tchar* casFile, bool storeCompressed, bool useMapping);
-		bool WriteCasFile(WriteCasResult& out, const tchar* fileName, const CasKey& casKey, bool useMapping);
+		bool WriteCasFileNoCheck(WriteCasResult& out, const StringKey& fileNameKey, const tchar* fileName, const CasKey& casKey, const tchar* casFile, bool storeCompressed, bool useMapping, DeferredSource* deferredSource, bool allowFail);
+		bool WriteCasFile(WriteCasResult& out, const tchar* fileName, const CasKey& casKey, bool useMapping, DeferredSource* deferredSouirce, bool allowFail);
 		bool VerifyExisting(bool& outReturnValue, CasEntryWriteLock& casEntryLock, const CasKey& casKey, CasEntry& casEntry, StringBufferBase& casFile);
 		bool AddCasFile(StringKey fileNameKey, const tchar* fileName, const CasKey& casKey, bool deferCreation, u64 fileSize, u64 lastWriteTime, CasEntry** outEntry, bool dropCasAfterUse);
 		void CasEntryAccessed(const CasKey& casKey);
diff --git a/Engine/Source/Programs/UnrealBuildAccelerator/Common/Public/UbaStorageClient.h b/Engine/Source/Programs/UnrealBuildAccelerator/Common/Public/UbaStorageClient.h
--- a/Engine/Source/Programs/UnrealBuildAccelerator/Common/Public/UbaStorageClient.h
+++ b/Engine/Source/Programs/UnrealBuildAccelerator/Common/Public/UbaStorageClient.h
@@ -59,7 +59,7 @@ namespace uba
 		virtual bool RetrieveCasFile(RetrieveResult& out, const CasKey& casKey, StringKey keyHint, const tchar* hint, FileMappingBuffer* mappingBuffer = nullptr, bool createIndependentMapping = false, u64 memoryMapAlignment = 1, bool allowProxy = true, u32 clientId = 0) override;
 		virtual bool StoreCasFile(CasKey& out, const tchar* fileName, const CasKey& casKeyOverride, bool deferCreation) override;
 		virtual bool StoreCasFileClient(CasKey& out, StringKey fileNameKey, const tchar* fileName, SharedMemoryHandle mappingFile, u64 mappingOffset, u64 fileSize, const tchar* hint, bool keepMappingInMemory = false, bool storeCompressed = true) override;
-		virtual bool HasCasFile(const CasKey& casKey, CasEntry** out = nullptr, bool useMapping = false) override;
+		virtual bool HasCasFile(const CasKey& casKey, CasEntry** out, bool useMapping, bool allowFail) override;
 		virtual void Ping() override;
 		virtual void PrintSummary(Logger& logger) override;
 
