[UBA]
* Fixed mac bug that was actually overflow of stack. solved by allocating big buffer each time (does not matter since it is rarely called)

diff --git a/Engine/Source/Programs/UnrealBuildAccelerator/Common/Private/UbaSessionClient.cpp b/Engine/Source/Programs/UnrealBuildAccelerator/Common/Private/UbaSessionClient.cpp
--- a/Engine/Source/Programs/UnrealBuildAccelerator/Common/Private/UbaSessionClient.cpp
+++ b/Engine/Source/Programs/UnrealBuildAccelerator/Common/Private/UbaSessionClient.cpp
@@ -470,7 +470,14 @@ namespace uba
 	{
 		TrackWorkScope tws(m_client, AsView(TC("ReadModules")), ColorWork);
 
-		StackBinaryReader<SendMaxSize> reader;
+		#if PLATFORM_MAC // Stack overflow on mac
+		u8* mem = new u8[SendMaxSize];
+		auto g = MakeGuard([mem]() { delete[] mem; });
+		#else
+		u8 mem[SendMaxSize];
+		#endif
+		BinaryReader reader(mem, 0, SendMaxSize);
+
 		{
 			StackBinaryWriter<1024> writer;
 			NetworkMessage msg(m_client, ServiceId, SessionMessageType_GetApplication, writer);
