[UBA]
* Added detouring support for jobs killing child processes. Right now it does not handle a job handle being closed in the middle of a process

diff --git a/Engine/Source/Programs/UnrealBuildAccelerator/Common/Private/UbaProcess.cpp b/Engine/Source/Programs/UnrealBuildAccelerator/Common/Private/UbaProcess.cpp
--- a/Engine/Source/Programs/UnrealBuildAccelerator/Common/Private/UbaProcess.cpp
+++ b/Engine/Source/Programs/UnrealBuildAccelerator/Common/Private/UbaProcess.cpp
@@ -398,6 +398,9 @@ namespace uba
 					err.Appendf(TC(" (GetExitCodeProcess returned 0x%x (%s Size: %llu, CasKey: %s)"), exitCode, m_realApplication.c_str(), imageInfo.fileSize, CasKeyString(imageInfo.key).str);
 			}
 
+				if (m_parentReportedExit)
+					return false;
+					
 			if (err.IsEmpty())
 			{
 				// Ok, this is a bit hacky putting this here.. but theory is that process could have exited and exit message is still in the shared memory
@@ -1421,10 +1424,12 @@ namespace uba
 
 		if (nativeProcessHandle)
 		{
-			DuplicateHandle((HANDLE)m_nativeProcessHandle, nativeProcessHandle, GetCurrentProcess(), (HANDLE*)&process.m_nativeProcessHandle, 0, false, DUPLICATE_SAME_ACCESS);
+			if (!DuplicateHandle((HANDLE)m_nativeProcessHandle, nativeProcessHandle, GetCurrentProcess(), (HANDLE*)&process.m_nativeProcessHandle, 0, false, DUPLICATE_SAME_ACCESS))
+				return CancelWithError(m_session.m_logger.Error(TC("Failed to duplicate handle for child process")));
 			if (!process.m_nativeProcessHandle || process.m_nativeProcessHandle == InvalidProcHandle)
 				return CancelWithError(m_session.m_logger.Error(TC("Failed to duplicate handle for child process")));
-			DuplicateHandle((HANDLE)m_nativeProcessHandle, nativeThreadHandle, GetCurrentProcess(), &process.m_nativeThreadHandle, 0, false, DUPLICATE_SAME_ACCESS);
+			if (!DuplicateHandle((HANDLE)m_nativeProcessHandle, nativeThreadHandle, GetCurrentProcess(), &process.m_nativeThreadHandle, 0, false, DUPLICATE_SAME_ACCESS))
+				return CancelWithError(m_session.m_logger.Error(TC("Failed to duplicate handle for child thread")));
 			if (!process.m_nativeThreadHandle || process.m_nativeThreadHandle == INVALID_HANDLE_VALUE)
 				return CancelWithError(m_session.m_logger.Error(TC("Failed to duplicate handle for child thread")));
 			process.m_nativeProcessId = nativeProcessId;
@@ -1460,6 +1465,7 @@ namespace uba
 
 	bool ProcessImpl::HandleExitChildProcess(BinaryReader& reader, BinaryWriter& writer)
 	{
+		// TODO: Should this traverse down? This is used for when child processes are terminated through job being closed. used by node.exe (building chrome)
 		u32 nativeProcessId = reader.ReadU32();
 		for (auto& child : m_childProcesses)
 		{
@@ -1469,7 +1475,6 @@ namespace uba
 			process.m_parentReportedExit = true;
 			return true;
 		}
-		UBA_ASSERT(false);
 		return true;
 	}
 
@@ -1530,6 +1535,15 @@ namespace uba
 		AtomicMax(m_processStats.detoursMemory, stats.detoursMemory.load());
 		AtomicMax(m_processStats.peakMemory, stats.peakMemory.load());
 
+		// Traverse children that has been tagged to be exited. This is a for-now solution until we found use cases where parent process
+		// kills job object before it exits itself
+		for (auto& child : m_childProcesses)
+		{
+			auto& process = *(ProcessImpl*)child.m_process;
+			if (process.m_parentReportedExit)
+				process.Cancel();
+		}
+
 		return false;
 	}
 
@@ -2486,7 +2500,7 @@ namespace uba
 		if (cancel)
 		{
 			if (!TerminateProcess(handle, ProcessCancelExitCode))
-				logger.Info(TC("TerminateProcess for process %llu (%s) failed when cancelling (%s)"), u64(handle), m_startInfo.GetDescription(), LastErrorToText().data);
+				logger.Info(TC("TerminateProcess for process %llu (%s) failed when cancelling (%s)%s"), u64(handle), m_startInfo.GetDescription(), LastErrorToText().data, (m_parentReportedExit ? TC("(ParentReportedExit)") : TC("")));
 			DWORD res = WaitForSingleObject(handle, 1000); // On windows this should return in milliseconds. On wine it might take a little while. We need to wait to make sure process is reaped
 			if (res == WAIT_FAILED)
 				logger.Info(TC("WaitForSingleObject for process %llu (%s) failed when cancelling (%s)"), u64(handle), m_startInfo.GetDescription(), LastErrorToText().data);
@@ -2509,7 +2523,7 @@ namespace uba
 						const tchar* gotMessage = m_gotExitMessage ? TC("Got") : TC("Did not get");
 						logger.Info(TC("WaitForSingleObject timed out after 120 seconds waiting for process %llu (%s) to exit (Exit code %u, %s ExitMessage and wrote %u files. Runtime: %s). Will terminate and wait again"), u64(handle), m_startInfo.GetDescription(), m_nativeProcessExitCode, gotMessage, u32(m_shared.writtenFiles.size()), TimeToText(GetTime() - m_startTime).str);
 						if (!TerminateProcess(handle, m_nativeProcessExitCode))
-							logger.Info(TC("TerminateProcess for process %llu (%s) failed when waiting timed out (%s)"), u64(handle), m_startInfo.GetDescription(), LastErrorToText().data);
+							logger.Info(TC("TerminateProcess for process %llu (%s) failed when waiting timed out (%s)%s"), u64(handle), m_startInfo.GetDescription(), LastErrorToText().data, (m_parentReportedExit ? TC("(ParentReportedExit)") : TC("")));
 						continue;
 					}
 					logger.Error(TC("WaitForSingleObject failed while waiting for process %llu (%s) to exit even after terminating it (%s)"), u64(handle), m_startInfo.GetDescription(), LastErrorToText().data);
diff --git a/Engine/Source/Programs/UnrealBuildAccelerator/Common/Public/UbaProcess.h b/Engine/Source/Programs/UnrealBuildAccelerator/Common/Public/UbaProcess.h
--- a/Engine/Source/Programs/UnrealBuildAccelerator/Common/Public/UbaProcess.h
+++ b/Engine/Source/Programs/UnrealBuildAccelerator/Common/Public/UbaProcess.h
@@ -163,10 +163,10 @@ namespace uba
 		Atomic<bool> m_cancelAllowed = true;
 		Atomic<bool> m_cancelled = false;
 		Atomic<bool> m_hasExited = false;
+		Atomic<bool> m_parentReportedExit = false;
 		MessageType m_firstMessageError = MessageType_None;
 		bool m_runningRemote = false;
 		bool m_gotExitMessage = false;
-		bool m_parentReportedExit = false;
 		bool m_detourEnabled = true;
 		bool m_isArmBinary = false;
 		TString m_realApplication;
diff --git a/Engine/Source/Programs/UnrealBuildAccelerator/Detours/Private/Windows/UbaDetoursFunctionsKernelBase.inl b/Engine/Source/Programs/UnrealBuildAccelerator/Detours/Private/Windows/UbaDetoursFunctionsKernelBase.inl
--- a/Engine/Source/Programs/UnrealBuildAccelerator/Detours/Private/Windows/UbaDetoursFunctionsKernelBase.inl
+++ b/Engine/Source/Programs/UnrealBuildAccelerator/Detours/Private/Windows/UbaDetoursFunctionsKernelBase.inl
@@ -3555,16 +3555,78 @@ BOOL Detoured_TerminateProcess(HANDLE hProcess, UINT uExitCode)
 
 BOOL Detoured_AssignProcessToJobObject(HANDLE hJob, HANDLE hProcess)
 {
-	{
-		//SuppressDetourScope _;
-		//MessageBoxW(GetConsoleWindow(), L"DFWQFWEFE", TC("Assert"), MB_ABORTRETRYIGNORE);
-	}
-
 	HANDLE trueHandle = hProcess;
 	if (isDetouredHandle(hProcess))
 		trueHandle = asDetouredHandle(hProcess).trueHandle;
 	BOOL res = True_AssignProcessToJobObject(hJob, trueHandle);
-	DEBUG_LOG_DETOURED(L"AssignProcessToJobObject", L"%llu -> %ls", uintptr_t(trueHandle), ToString(res));
+	DEBUG_LOG_DETOURED(L"AssignProcessToJobObject", L"Process %llu assigned to %llu -> %ls", uintptr_t(trueHandle), u64(hJob), ToString(res));
+
+	if (res)
+	{
+		JOBOBJECT_EXTENDED_LIMIT_INFORMATION info = {};
+		if (QueryInformationJobObject(hJob, JobObjectExtendedLimitInformation, &info, sizeof(info), nullptr))
+		{
+			if (info.BasicLimitInformation.LimitFlags & JOB_OBJECT_LIMIT_KILL_ON_JOB_CLOSE)
+			{
+				DWORD pid = GetProcessId(hProcess);
+				DEBUG_LOG(L"PROCESS %u WILL BE TERMINATED from job %llu", pid, u64(hJob));
+				RPC_MESSAGE(ExitChildProcess, detach) // Let it be part of detach timer for now
+				writer.WriteU32(pid);
+				writer.Flush();
+			}
+		}
+		else
+		{
+			UBA_ASSERT(false);
+		}
+	}
+
+	return res;
+}
+
+BOOL Detoured_SetInformationJobObject(HANDLE hJob, JOBOBJECTINFOCLASS JobObjectInformationClass, LPVOID lpJobObjectInformation, DWORD cbJobObjectInformationLength)
+{
+	DEBUG_LOG_TRUE(L"SetInformationJobObject", L"%u", JobObjectInformationClass);
+	bool res = True_SetInformationJobObject(hJob, JobObjectInformationClass, lpJobObjectInformation, cbJobObjectInformationLength);
+
+	if (!res)
+		return res;
+
+	bool killOnJobClose = false;
+	if (JobObjectInformationClass == JobObjectBasicLimitInformation)
+	{
+		auto info = *(JOBOBJECT_BASIC_LIMIT_INFORMATION*)lpJobObjectInformation;
+		killOnJobClose = (info.LimitFlags & JOB_OBJECT_LIMIT_KILL_ON_JOB_CLOSE) != 0;
+	}
+	if (JobObjectInformationClass == JobObjectExtendedLimitInformation)
+	{
+		auto info = *(JOBOBJECT_EXTENDED_LIMIT_INFORMATION*)lpJobObjectInformation;
+		killOnJobClose = (info.BasicLimitInformation.LimitFlags & JOB_OBJECT_LIMIT_KILL_ON_JOB_CLOSE) != 0;
+	}
+
+	if (killOnJobClose)
+	{
+		DEBUG_LOG(L"PROCESSES WILL BE TERMINATED ON CLOSE OF %llu", u64(hJob));
+
+		u8 buffer[sizeof(JOBOBJECT_BASIC_PROCESS_ID_LIST) + 1024 * sizeof(ULONG_PTR)];
+		auto idList = (JOBOBJECT_BASIC_PROCESS_ID_LIST*)buffer;
+		if (QueryInformationJobObject(hJob, JobObjectBasicProcessIdList, buffer, sizeof(buffer), NULL))
+		{
+			for (DWORD i = 0; i < idList->NumberOfProcessIdsInList; ++i)
+			{
+				DWORD pid = (DWORD)idList->ProcessIdList[i];
+				DEBUG_LOG(L"PROCESS %u WILL BE TERMINATED from job %llu", pid, u64(hJob));
+				// TODO: This should be tracked and sent when hJob handle closes.. for now we just tell all processes that they might go away at any time
+				RPC_MESSAGE(ExitChildProcess, detach) // Let it be part of detach timer for now
+				writer.WriteU32(pid);
+				writer.Flush();
+			}
+		}
+		else
+		{
+			UBA_ASSERT(false);
+		}
+	}
 	return res;
 }
 
diff --git a/Engine/Source/Programs/UnrealBuildAccelerator/Detours/Private/Windows/UbaDetoursFunctionsWin.h b/Engine/Source/Programs/UnrealBuildAccelerator/Detours/Private/Windows/UbaDetoursFunctionsWin.h
--- a/Engine/Source/Programs/UnrealBuildAccelerator/Detours/Private/Windows/UbaDetoursFunctionsWin.h
+++ b/Engine/Source/Programs/UnrealBuildAccelerator/Detours/Private/Windows/UbaDetoursFunctionsWin.h
@@ -129,6 +129,7 @@
 
 #define DETOURED_FUNCTIONS_KERNEL32 \
 	DETOURED_FUNCTION(AssignProcessToJobObject) \
+	DETOURED_FUNCTION(SetInformationJobObject) \
 	DETOURED_FUNCTION(RegisterWaitForSingleObject) \
 	DETOURED_FUNCTION(CreateFileMappingA) \
 	DETOURED_FUNCTION(GetExitCodeProcess) \
