[UBA]
* Fixed code checking if image is 32/64 bit

diff --git a/Engine/Source/Programs/UnrealBuildAccelerator/Common/Private/UbaProcess.cpp b/Engine/Source/Programs/UnrealBuildAccelerator/Common/Private/UbaProcess.cpp
--- a/Engine/Source/Programs/UnrealBuildAccelerator/Common/Private/UbaProcess.cpp
+++ b/Engine/Source/Programs/UnrealBuildAccelerator/Common/Private/UbaProcess.cpp
@@ -111,6 +111,7 @@ namespace uba
 	{
 		CasKey key;
 		u64 fileSize = 0;
+		bool isPEImage = false;
 		bool is64Bit = true;
 		bool isX64 = false;
 		bool isArm64 = false;
@@ -136,12 +137,20 @@ namespace uba
 		{
 			u32 offset = *(u32*)(data + 0x3c);
 			u32* signaturePos = (u32*)(data + offset);
-			out.is64Bit = *signaturePos == 0x00004550;
-			u16 machine = *(u16*)(signaturePos + 1);
-			out.isX64 = machine == 0x8664;
-			out.isArm64 = machine == 0xaa64;
-			if (out.fileSize > offset + 0x18 + 0x70 + 4)
-				out.isDotnet = *(u32*)(data + offset + 0x18 + 0x70);
+
+			out.isPEImage = *signaturePos == 0x00004550;
+			if (out.isPEImage)
+			{
+				u16 machine = *(u16*)(signaturePos + 1);
+				out.isX64 = machine == 0x8664;
+				out.isArm64 = machine == 0xaa64;
+
+				u16 magic = *(u16*)(data + offset + 0x18);
+				out.is64Bit = (magic == 0x20B);
+
+				if (out.fileSize > offset + 0x18 + 0x70 + 4)
+					out.isDotnet = *(u32*)(data + offset + 0x18 + 0x70);
+			}
 		}
 		return fa.Close();
 	}
@@ -385,6 +394,8 @@ namespace uba
 					err.Appendf(TC("UbaDetours.dll has a different version than " UBA_BINARY));
 				else if (!GetImageInfo(imageInfo, m_session.m_logger, m_realApplication.c_str(), true))
 					err.Appendf(TC("Failed to load %s"), m_realApplication.c_str());
+				else if (!imageInfo.isPEImage)
+					err.Append(TCV("Doesn't seem to be a PE executable"));
 				else if (!imageInfo.is64Bit)
 					err.Append(TCV("Doesn't seem to be a 64-bit executable"));
 				else if (imageInfo.isDotnet)
@@ -395,7 +406,7 @@ namespace uba
 					err.Appendf(TC("Machine is x64 and image is arm64"));
 
 				if (exitCode != 1398)
-					err.Appendf(TC(" (GetExitCodeProcess returned 0x%x (%s Size: %llu, CasKey: %s)"), exitCode, m_realApplication.c_str(), imageInfo.fileSize, CasKeyString(imageInfo.key).str);
+					err.Appendf(TC(" (GetExitCodeProcess returned 0x%x - %s Size: %llu, CasKey: %s)"), exitCode, m_realApplication.c_str(), imageInfo.fileSize, CasKeyString(imageInfo.key).str);
 			}
 
 				if (m_parentReportedExit)
@@ -2076,6 +2087,16 @@ namespace uba
 			m_nativeProcessId = processInfo.dwProcessId;
 			m_nativeThreadHandle = processInfo.hThread;
 
+			#if UBA_DEBUG
+			USHORT processMachine = 0;
+			USHORT nativeMachine = 0;
+			if (IsWow64Process2(processInfo.hProcess, &processMachine, &nativeMachine))
+			{
+				bool is64Bit = processMachine == IMAGE_FILE_MACHINE_UNKNOWN;
+				UBA_ASSERTF(is64Bit, TC("Can't detour 32-bit application %s"), m_realApplication.c_str());
+			}
+			#endif
+
 			FILETIME dummy;
 			if (processInfo.hProcess && !GetProcessTimes(processInfo.hProcess, (FILETIME*)&m_nativeProcessCreationTime, &dummy, &dummy, &dummy))
 				logger.Error(TC("GetProcessTimes failed (%s)"), LastErrorToText().data);
