[UBA]
* Improved assert
* Minor cosmetic improvement

diff --git a/Engine/Source/Programs/UnrealBuildAccelerator/Detours/Private/Windows/UbaDetoursFunctionsWin.cpp b/Engine/Source/Programs/UnrealBuildAccelerator/Detours/Private/Windows/UbaDetoursFunctionsWin.cpp
--- a/Engine/Source/Programs/UnrealBuildAccelerator/Detours/Private/Windows/UbaDetoursFunctionsWin.cpp
+++ b/Engine/Source/Programs/UnrealBuildAccelerator/Detours/Private/Windows/UbaDetoursFunctionsWin.cpp
@@ -336,7 +336,10 @@ void MemoryFile::EnsureCommitted(const DetouredHandle& handle, u64 size)
 	if (committedSize >= size)
 		return;
 	if (size > reserveSize)
+	{
+		UBA_ASSERTF(false, L"Reserved size of %ls is smaller than what is requested to be. ReserveSize: %llu Written: %llu Requested: %llu", HandleToName(handle), reserveSize, writtenSize, size);
 		FatalError(1347, L"Reserved size of %ls is smaller than what is requested to be. ReserveSize: %llu Written: %llu Requested: %llu", HandleToName(handle), reserveSize, writtenSize, size);
+	}
 	/*
 	if (size > mappedSize)
 	{
@@ -951,7 +954,7 @@ void Init(const DetoursPayload& payload, u64 startTime)
 			auto& objName0 = *reinterpret_cast<UNICODE_STRING*>(buffer0);
 			if (!g_stdHandle[0])
 				objName0.Buffer = const_cast<tchar*>(L"Unset");
-			else if (NT_SUCCESS(True_NtQueryObject(g_stdHandle[0], (OBJECT_INFORMATION_CLASS)1, buffer0, sizeof(buffer0), &neededSize)) && objName0.Buffer)
+			else if (NT_SUCCESS(True_NtQueryObject(g_stdHandle[0], ObjectNameInformation, buffer0, sizeof(buffer0), &neededSize)) && objName0.Buffer)
 				g_conEnabled[0] = !Equals(objName0.Buffer, L"\\Device\\Null");
 			else
 				objName0.Buffer = const_cast<tchar*>(L"???");
@@ -960,7 +963,7 @@ void Init(const DetoursPayload& payload, u64 startTime)
 			auto& objName1 = *reinterpret_cast<UNICODE_STRING*>(buffer1);
 			if (!g_stdHandle[1])
 				objName1.Buffer = const_cast<tchar*>(L"Unset");
-			else if (NT_SUCCESS(True_NtQueryObject(g_stdHandle[1], (OBJECT_INFORMATION_CLASS)1, buffer1, sizeof(buffer1), &neededSize)) && objName1.Buffer)
+			else if (NT_SUCCESS(True_NtQueryObject(g_stdHandle[1], ObjectNameInformation, buffer1, sizeof(buffer1), &neededSize)) && objName1.Buffer)
 				g_conEnabled[1] = !Equals(objName1.Buffer, L"\\Device\\Null");
 			else
 				objName1.Buffer = const_cast<tchar*>(L"???");
