[UBA]
* Made no-syscall load library call thread safe
* Improved logging messages
* Fixed so DuplicatedHandle can handle detoured target process handle

diff --git a/Engine/Source/Programs/UnrealBuildAccelerator/Detours/Private/Windows/UbaDetoursFunctionsKernelBase.inl b/Engine/Source/Programs/UnrealBuildAccelerator/Detours/Private/Windows/UbaDetoursFunctionsKernelBase.inl
--- a/Engine/Source/Programs/UnrealBuildAccelerator/Detours/Private/Windows/UbaDetoursFunctionsKernelBase.inl
+++ b/Engine/Source/Programs/UnrealBuildAccelerator/Detours/Private/Windows/UbaDetoursFunctionsKernelBase.inl
@@ -148,11 +148,14 @@ BOOL Detoured_SetCurrentDirectoryW(LPCWSTR lpPathName)
 
 BOOL Detoured_DuplicateHandle(HANDLE hSourceProcessHandle, HANDLE hSourceHandle, HANDLE hTargetProcessHandle, LPHANDLE lpTargetHandle, DWORD dwDesiredAccess, BOOL bInheritHandle, DWORD dwOptions)
 {
+	if (isDetouredHandle(hTargetProcessHandle))
+		hTargetProcessHandle = asDetouredHandle(hTargetProcessHandle).trueHandle;
+
 	DETOURED_CALL(DuplicateHandle);
 	if (hSourceHandle == PseudoHandle || !isDetouredHandle(hSourceHandle))
 	{
 		auto res = True_DuplicateHandle(hSourceProcessHandle, hSourceHandle, hTargetProcessHandle, lpTargetHandle, dwDesiredAccess, bInheritHandle, dwOptions);
-		DEBUG_LOG_TRUE(L"DuplicateHandle", L"%llu to %llu -> %ls", uintptr_t(hSourceHandle), lpTargetHandle ? uintptr_t(*lpTargetHandle) : 0ull, ToString(res));
+		DEBUG_LOG_TRUE(L"DuplicateHandle",L"%llu (%llu) to %llu (%llu) -> %ls", uintptr_t(hSourceHandle), uintptr_t(hSourceProcessHandle), lpTargetHandle ? uintptr_t(*lpTargetHandle) : 0ull, uintptr_t(hTargetProcessHandle), ToString(res));
 		return res;
 	}
 
@@ -179,7 +182,7 @@ BOOL Detoured_DuplicateHandle(HANDLE hSourceProcessHandle, HANDLE hSourceHandle,
 		InterlockedIncrement(&fo->refCount);
 	}
 	*lpTargetHandle = makeDetouredHandle(newDh);
-	DEBUG_LOG_DETOURED(L"DuplicateHandle", L"%llu to %llu -> %ls", uintptr_t(hSourceHandle), uintptr_t(*lpTargetHandle), ToString(res));
+	DEBUG_LOG_DETOURED(L"DuplicateHandle", L"%llu (%llu) to %llu/%llu (%llu) -> %ls", u64(hSourceHandle), u64(hSourceProcessHandle), u64(*lpTargetHandle), u64(targetHandle), u64(hTargetProcessHandle), ToString(res));
 	return res;
 }
 
@@ -1771,10 +1774,10 @@ bool Shared_MoveFile(const tchar* func, LPCWSTR lpExistingFileName, LPCWSTR lpNe
 		errorCode = reader.ReadU32();
 		directoryTableSize = reader.ReadU32();
 		pcs.Leave();
-		DEBUG_LOG_PIPE(L"MoveFile", L"%ls to %ls", lpExistingFileName, lpNewFileName);
+		DEBUG_LOG_PIPE(func, L"%ls to %ls", lpExistingFileName, lpNewFileName);
 	}
 
-	DEBUG_LOG_DETOURED(L"MoveFileExW", L"(PIPE) (%ls to %ls) -> %ls (%u)", lpExistingFileName, lpNewFileName, ToString(result), errorCode);
+	DEBUG_LOG_DETOURED(func, L"(PIPE) (%ls to %ls) -> %ls (%u)", lpExistingFileName, lpNewFileName, ToString(result), errorCode);
 
 	if (result)
 	{
@@ -2256,7 +2259,7 @@ BOOL Detoured_GetFileInformationByHandleEx(HANDLE hFile, FILE_INFO_BY_HANDLE_CLA
 			UBA_ASSERTF(trueHandle != INVALID_HANDLE_VALUE, L"GetFileInformationByHandleEx with class %u not Implemented (%ls)", fileInformationClass, HandleToName(hFile));
 		}
 	}
-	DEBUG_LOG_TRUE(L"GetFileInformationByHandleEx", L"(%ls)", HandleToName(hFile));
+	DEBUG_LOG_TRUE(L"GetFileInformationByHandleEx", L"(%u) %llu (%ls)", u32(fileInformationClass), uintptr_t(hFile), HandleToName(hFile));
 	TimerScope ts(g_kernelStats.getFileInfo);
 	return True_GetFileInformationByHandleEx(trueHandle, fileInformationClass, lpFileInformation, dwBufferSize); /// calls GetFileInformationByHandleEx
 }
@@ -2326,6 +2329,7 @@ BOOL Detoured_GetFileInformationByHandle(HANDLE hFile, LPBY_HANDLE_FILE_INFORMAT
 			lpFileInformation->nFileSizeHigh = li.HighPart;
 			lpFileInformation->nFileSizeLow = li.LowPart;
 			DEBUG_LOG_DETOURED(L"GetFileInformationByHandle", L"(file) %llu (%ls) -> Success (size: %llu)", uintptr_t(hFile), HandleToName(hFile), fileSize);
+			SetLastError(ERROR_SUCCESS);
 			return TRUE;
 		}
 
@@ -2344,6 +2348,7 @@ BOOL Detoured_GetFileInformationByHandle(HANDLE hFile, LPBY_HANDLE_FILE_INFORMAT
 			li = ToLargeInteger(mf->writtenSize);
 			lpFileInformation->nFileSizeHigh = li.HighPart;
 			lpFileInformation->nFileSizeLow = li.LowPart;
+			SetLastError(ERROR_SUCCESS);
 			return TRUE;
 		}
 
@@ -2363,7 +2368,7 @@ BOOL Detoured_GetFileInformationByHandle(HANDLE hFile, LPBY_HANDLE_FILE_INFORMAT
 
 
 					SetLastError(ERROR_FILE_NOT_FOUND);
-					DEBUG_LOG_DETOURED(L"GetFileInformationByHandle", L"remote %llu (%ls) -> NotFound", uintptr_t(hFile), HandleToName(hFile));
+					DEBUG_LOG_DETOURED(L"GetFileInformationByHandle", L"(remote) %llu (%ls) -> NotFound", uintptr_t(hFile), HandleToName(hFile));
 					return false;
 				}
 
@@ -2401,7 +2406,7 @@ BOOL Detoured_GetFileInformationByHandle(HANDLE hFile, LPBY_HANDLE_FILE_INFORMAT
 				lpFileInformation->nNumberOfLinks = 1;//~u32(0); // TODO
 				lpFileInformation->nFileSizeHigh = attr.data.nFileSizeHigh;
 				lpFileInformation->nFileSizeLow = attr.data.nFileSizeLow;
-				DEBUG_LOG_DETOURED(L"GetFileInformationByHandle", L"remote %llu (%ls) -> Success", uintptr_t(hFile), HandleToName(hFile));
+				DEBUG_LOG_DETOURED(L"GetFileInformationByHandle", L"(remote) %llu (%ls) -> Success", uintptr_t(hFile), HandleToName(hFile));
 				return TRUE;
 			}
 		}
@@ -2928,6 +2933,11 @@ HMODULE Recursive_LoadLibraryExW(const StringView& filePath, LPCWSTR originalNam
 
 HMODULE LoadModuleNoSyscall(StringView moduleName)
 {
+	ULONG disp;
+	PVOID cookie;
+	LdrLockLoaderLock(0, &disp, &cookie);
+	auto lockGuard = MakeGuard([&](){ LdrUnlockLoaderLock(0, cookie); });
+
 	StringView moduleFileName = moduleName.GetFileName();
 	PPEB peb = (PPEB)ReadTebOffset(0x60);
 	PPEB_LDR_DATA ldr = peb->Ldr;
@@ -3181,7 +3191,7 @@ BOOL Detoured_CreateProcessW(LPCWSTR lpApplicationName, LPWSTR lpCommandLine, LP
 	StringView originalCmd = ToView(lpCommandLine ? lpCommandLine : TC(""));
 
 	DETOURED_CALL(CreateProcessW);
-	DEBUG_LOG_DETOURED(L"CreateProcessW", L"%ls %ls CreationFlags: 0x%x StartupFlags: 0x%u Stdin: %llu Stdout: %llu Stderr: %llu WorkDir: %s", lpApplicationName, originalCmd.data, dwCreationFlags, lpStartupInfo->dwFlags, u64(lpStartupInfo->hStdInput), u64(lpStartupInfo->hStdOutput), u64(lpStartupInfo->hStdError), (lpCurrentDirectory ? lpCurrentDirectory : L""));
+	DEBUG_LOG_DETOURED(L"CreateProcessW", L"%ls %ls CreationFlags: 0x%x StartupFlags: 0x%u Stdin: %llu Stdout: %llu Stderr: %llu WorkDir: %s Inherit: %s Env: %llu", lpApplicationName, originalCmd.data, dwCreationFlags, lpStartupInfo->dwFlags, u64(lpStartupInfo->hStdInput), u64(lpStartupInfo->hStdOutput), u64(lpStartupInfo->hStdError), (lpCurrentDirectory ? lpCurrentDirectory : L""), bInheritHandles?L"true":L"false", u64(lpEnvironment));
 
 	StringView testName = ToView(lpApplicationName ? lpApplicationName : TC(""));
 
@@ -3385,7 +3395,7 @@ BOOL Detoured_CreateProcessW(LPCWSTR lpApplicationName, LPWSTR lpCommandLine, LP
 		else
 		{
 			#if UBA_ASSERT_ON_DETOURED_STD_PIPE
-			UBA_ASSERTF(false, L"hStdInput is detoured (%s)", lpApplicationName);
+			UBA_ASSERTF(false, L"hStdInput is detoured (type %i, true: %llu) (%s)", int(dh.type), dh.trueHandle, lpApplicationName);
 			#endif
 			
 			lpStartupInfo->hStdInput = dh.trueHandle;
@@ -3457,7 +3467,7 @@ BOOL Detoured_CreateProcessW(LPCWSTR lpApplicationName, LPWSTR lpCommandLine, LP
 		DEBUG_LOG_PIPE(L"StartProcess", L"%ls %ls", application.data, originalCmd.data);
 	}
 
-	UBA_ASSERTF(res, L"Failed to spawn process %ls (Error code: %u)", originalCmd.data, lastError);
+	UBA_ASSERTF(res, L"Failed to spawn process %ls (Error code: %u)", finalCommandLine, lastError);
 
 	HANDLE trueHandle = lpProcessInformation->hProcess;
 
diff --git a/Engine/Source/Programs/UnrealBuildAccelerator/Detours/Private/Windows/UbaDetoursFunctionsWin.h b/Engine/Source/Programs/UnrealBuildAccelerator/Detours/Private/Windows/UbaDetoursFunctionsWin.h
--- a/Engine/Source/Programs/UnrealBuildAccelerator/Detours/Private/Windows/UbaDetoursFunctionsWin.h
+++ b/Engine/Source/Programs/UnrealBuildAccelerator/Detours/Private/Windows/UbaDetoursFunctionsWin.h
@@ -424,6 +424,8 @@ extern "C" {
 	NTSTATUS NTAPI ZwCreateSection(PHANDLE SectionHandle, ACCESS_MASK DesiredAccess, POBJECT_ATTRIBUTES ObjectAttributes, PLARGE_INTEGER MaximumSize, ULONG SectionPageProtection, ULONG AllocationAttributes, HANDLE FileHandle);
 	NTSTATUS NTAPI ZwQueryDirectoryFile(HANDLE FileHandle, HANDLE Event, PIO_APC_ROUTINE ApcRoutine, PVOID ApcContext, PIO_STATUS_BLOCK IoStatusBlock, PVOID FileInformation, ULONG Length, FILE_INFORMATION_CLASS FileInformationClass, BOOLEAN ReturnSingleEntry, PUNICODE_STRING FileName, BOOLEAN RestartScan);
 	NTSTATUS NTAPI ZwSetInformationFile(HANDLE FileHandle,PIO_STATUS_BLOCK IoStatusBlock, PVOID FileInformation, ULONG Length, FILE_INFORMATION_CLASS FileInformationClass);
+	NTSTATUS NTAPI LdrLockLoaderLock(ULONG Flags, ULONG* Disposition, PVOID* Cookie);
+	NTSTATUS NTAPI LdrUnlockLoaderLock(ULONG Flags, PVOID Cookie);
 	NTSTATUS NTAPI LdrAddRefDll(ULONG Flags, PVOID BaseAddress);
 	PVOID WINAPI ResolveDelayLoadedAPI(PVOID ParentModuleBase, PCIMAGE_DELAYLOAD_DESCRIPTOR DelayloadDescriptor, void* FailureDllHook, void* FailureSystemHook, PIMAGE_THUNK_DATA ThunkAddress, ULONG Flags);
 	void WINAPI RtlExitUserThread(ULONG);
