[UBA]
* Fixed bug where if file is rewritten after it was used as a file mapping with a memory handle

diff --git a/Engine/Source/Programs/UnrealBuildAccelerator/Common/Private/UbaSession.cpp b/Engine/Source/Programs/UnrealBuildAccelerator/Common/Private/UbaSession.cpp
--- a/Engine/Source/Programs/UnrealBuildAccelerator/Common/Private/UbaSession.cpp
+++ b/Engine/Source/Programs/UnrealBuildAccelerator/Common/Private/UbaSession.cpp
@@ -605,6 +605,7 @@ namespace uba
 					if (!res.second)
 						return;
 					UBA_ASSERT(e.attributes);
+					UBA_ASSERT(e.nameLen < 512);
 					memoryWriter.WriteString(e.name, e.nameLen);
 
 					#if UBA_DEBUG_TRACK_DIR
@@ -3386,6 +3387,43 @@ namespace uba
 				}
 			}
 
+			{
+				// We need to check if there are any memory maps for this file from before. If it is we need to update it (for example, build steps that writes .h files might be here)
+				// TODO For now we just turn that mapping to a normal file but maybe we should recreate the file mapping with a new mapping
+				SCOPED_FUTEX(m_fileMappingTableLookupLock, lookupLock);
+				auto findMappingIt = m_fileMappingTableLookup.find(key);
+				if (findMappingIt != m_fileMappingTableLookup.end())
+				{
+					FileMappingEntry& entry = findMappingIt->second;
+					SCOPED_FUTEX(entry.lock, entryLock);
+					if (entry.memoryHandle.IsValid()) // If there is a memory map we need to make sure to update it..
+					{
+						//m_logger.Warning(TC("MOVING %s to file"), name.data), 
+						writtenFile.memoryHandle = {};
+						UBA_ASSERT(!msg.memoryHandle.IsValid());
+						entry.contentSize = InvalidValue;
+						entry.isDir = false;
+						entry.success = true;
+						entry.memoryHandle = {};
+						entry.handled = true;
+						#if UBA_DEBUG_TRACK_MAPPING
+						entry.name = TC("#");
+						#endif
+						SCOPED_WRITE_LOCK(m_fileMappingTableMemLock, memLock);
+						BinaryWriter writer(m_fileMappingTableMem, m_fileMappingTableSize);
+						writer.WriteStringKey(key);
+						if (m_runningRemote)
+							writer.WriteString(writtenFile.backedName);
+						else
+							writer.WriteString(TCV("#"));
+						writer.Write7BitEncoded(InvalidValue);
+						u32 newSize = (u32)writer.GetPosition();
+						m_fileMappingTableSize = (u32)newSize;
+						addMapping = false;
+					}
+				}
+			}
+
 			if (addMapping)
 			{
 				msg.process.m_shared.sharedMemoryHandles.erase(msg.memoryHandle); // memory handle transferred to written files
@@ -4436,7 +4474,7 @@ namespace uba
 				continue;
 			}
 			temp.Clear();
-			temp.Append(start, u64(it - start));
+			temp.Append(start, Min(256ull, u64(it - start)));
 			if (!temp.Contains(TC(".rsp")) && !temp.Contains(TC(".bat")))
 			{
 				if (*it == 0)
diff --git a/Engine/Source/Programs/UnrealBuildAccelerator/Detours/Private/UbaDetoursFileMappingTable.cpp b/Engine/Source/Programs/UnrealBuildAccelerator/Detours/Private/UbaDetoursFileMappingTable.cpp
--- a/Engine/Source/Programs/UnrealBuildAccelerator/Detours/Private/UbaDetoursFileMappingTable.cpp
+++ b/Engine/Source/Programs/UnrealBuildAccelerator/Detours/Private/UbaDetoursFileMappingTable.cpp
@@ -40,10 +40,24 @@ namespace uba
 			FileInfo& info = insres.first->second;
 			if (!insres.second)
 			{
-				if (info.name && IsMemoryHandle(info.name) && !mappedFileName.Equals(info.name, false)) // Mapped file has been re-mapped.
+				if (!info.name)
+					continue;
+				if (IsMemoryHandle(info.name) && !mappedFileName.Equals(info.name, false)) // Mapped file has been re-mapped.
 				{
-					UBA_ASSERTF(!info.memoryFile, TC("Mapped file %s has changed mapping (%s to %s) while being in use"), info.originalName, info.name, mappedFileName.data);
-					info.name = m_memoryBlock.Strdup(mappedFileName).data;
+					if (!IsMemoryHandle(mappedFileName.data))
+					{
+						info.name = m_memoryBlock.Strdup(mappedFileName).data;
+						info.size = size;
+						//UBA_ASSERT(!info.isFileMap);
+						info.isFileMap = false;
+					}
+					else
+					{
+						UBA_ASSERTF(!info.memoryFile, TC("Mapped file %s has changed mapping (%s to %s) while being in use"), info.originalName, info.name, mappedFileName.data);
+						info.name = m_memoryBlock.Strdup(mappedFileName).data;
+						UBA_ASSERTF(size != InvalidValue, TC("File mapping of %s (%s)"), info.name, mappedFileName.data);
+						info.size = size;
+					}
 				}
 				continue;
 			}
