[UBA]
* Changed how detoured process waits for child processes to make sure exit time is properly calculated

diff --git a/Engine/Source/Programs/UnrealBuildAccelerator/Common/Private/UbaProcess.cpp b/Engine/Source/Programs/UnrealBuildAccelerator/Common/Private/UbaProcess.cpp
--- a/Engine/Source/Programs/UnrealBuildAccelerator/Common/Private/UbaProcess.cpp
+++ b/Engine/Source/Programs/UnrealBuildAccelerator/Common/Private/UbaProcess.cpp
@@ -2521,7 +2521,8 @@ namespace uba
 		if (cancel)
 		{
 			if (!TerminateProcess(handle, ProcessCancelExitCode))
-				logger.Info(TC("TerminateProcess for process %llu (%s) failed when cancelling (%s)%s"), u64(handle), m_startInfo.GetDescription(), LastErrorToText().data, (m_parentReportedExit ? TC("(ParentReportedExit)") : TC("")));
+				if (!m_parentReportedExit || GetLastError() != ERROR_ACCESS_DENIED)
+					logger.Info(TC("TerminateProcess for process %llu (%s) failed when cancelling (%s)%s"), u64(handle), m_startInfo.GetDescription(), LastErrorToText().data, (m_parentReportedExit ? TC("(ParentReportedExit)") : TC("")));
 			DWORD res = WaitForSingleObject(handle, 1000); // On windows this should return in milliseconds. On wine it might take a little while. We need to wait to make sure process is reaped
 			if (res == WAIT_FAILED)
 				logger.Info(TC("WaitForSingleObject for process %llu (%s) failed when cancelling (%s)"), u64(handle), m_startInfo.GetDescription(), LastErrorToText().data);
@@ -2766,12 +2767,11 @@ namespace uba
 				child.Cancel();
 
 		for (auto& child : m_childProcesses)
-		{
-			auto& childProcess = *(ProcessImpl*)child.m_process;
-			childProcess.m_waitForParent.Set();
-			while (!childProcess.m_hasExited)
-				Sleep(10);
-		}
+			((ProcessImpl*)child.m_process)->m_waitForParent.Set();
+
+		for (auto& child : m_childProcesses)
+			if (!((ProcessImpl*)child.m_process)->m_messageThreadDone.IsSet(10000))
+				m_session.m_logger.Warning(TC("Process timed out waiting for its child process to finish (%s)"), m_startInfo.GetDescription());
 	}
 
 	bool ProcessImpl::CancelWithError(bool dummy)
