[UBA]
* Fixed so RegisterDeleteFile does not try to update dirtable for directories that are invalid
* Improved assert descriptions

diff --git a/Engine/Source/Programs/UnrealBuildAccelerator/Common/Private/UbaSession.cpp b/Engine/Source/Programs/UnrealBuildAccelerator/Common/Private/UbaSession.cpp
--- a/Engine/Source/Programs/UnrealBuildAccelerator/Common/Private/UbaSession.cpp
+++ b/Engine/Source/Programs/UnrealBuildAccelerator/Common/Private/UbaSession.cpp
@@ -1495,6 +1495,10 @@ namespace uba
 			Sleep(1);
 			dirLock.Enter();
 		}
+		
+		if (dir.parseOffset == 2) // Invalid directory, don't need to register this
+			return 0;
+
 		UBA_ASSERTF(dir.parseOffset == 1, TC("Registering deleted file %s with unexpect dir.parseOffset %u "), fileName.data, dir.parseOffset);
 
 		if (fileNameKey == StringKeyZero)
@@ -3644,13 +3648,13 @@ namespace uba
 
 	bool Session::GetFullFileName(GetFullFileNameResponse& out, const GetFullFileNameMessage& msg)
 	{
-		UBA_ASSERTF(false, TC("SHOULD NOT HAPPEN (only remote).. %s"), msg.fileName.data);
+		UBA_ASSERTF(false, TC("SHOULD NOT HAPPEN (only remote).. %s (%s)"), msg.fileName.data, msg.process.m_startInfo.GetDescription());
 		return false;
 	}
 
 	bool Session::GetLongPathName(GetLongPathNameResponse& out, const GetLongPathNameMessage& msg)
 	{
-		UBA_ASSERTF(false, TC("SHOULD NOT HAPPEN (only remote).. %s"), msg.fileName.data);
+		UBA_ASSERTF(false, TC("SHOULD NOT HAPPEN (only remote).. %s (%s)"), msg.fileName.data, msg.process.m_startInfo.GetDescription());
 		return false;
 	}
 
