[UBA]
* Added more detoured functions in debug to be able to track down future detouring issues

diff --git a/Engine/Source/Programs/UnrealBuildAccelerator/Common/Private/UbaFile.cpp b/Engine/Source/Programs/UnrealBuildAccelerator/Common/Private/UbaFile.cpp
--- a/Engine/Source/Programs/UnrealBuildAccelerator/Common/Private/UbaFile.cpp
+++ b/Engine/Source/Programs/UnrealBuildAccelerator/Common/Private/UbaFile.cpp
@@ -14,7 +14,6 @@
 #include <WinCon.h>
 #include <RestartManager.h>
 #pragma comment(lib, "Rstrtmgr.lib")
-struct FILE_NETWORK_OPEN_INFORMATION { LARGE_INTEGER CreationTime; LARGE_INTEGER LastAccessTime; LARGE_INTEGER LastWriteTime; LARGE_INTEGER ChangeTime; LARGE_INTEGER AllocationSize; LARGE_INTEGER EndOfFile; ULONG FileAttributes; };
 extern "C" NTSTATUS NTAPI NtQueryInformationFile(HANDLE FileHandle, PIO_STATUS_BLOCK IoStatusBlock, PVOID FileInformation, ULONG Length, FILE_INFORMATION_CLASS FileInformationClass);
 extern "C" NTSTATUS NTAPI NtQueryFullAttributesFile(const OBJECT_ATTRIBUTES *attr, FILE_NETWORK_OPEN_INFORMATION *info);
 #else
diff --git a/Engine/Source/Programs/UnrealBuildAccelerator/Core/Public/UbaWinTypes.h b/Engine/Source/Programs/UnrealBuildAccelerator/Core/Public/UbaWinTypes.h
--- a/Engine/Source/Programs/UnrealBuildAccelerator/Core/Public/UbaWinTypes.h
+++ b/Engine/Source/Programs/UnrealBuildAccelerator/Core/Public/UbaWinTypes.h
@@ -59,10 +59,32 @@ struct FILE_BASIC_INFORMATION {
 	ULONG FileAttributes;
 };
 
+struct FILE_NETWORK_OPEN_INFORMATION
+{
+	LARGE_INTEGER CreationTime;
+	LARGE_INTEGER LastAccessTime;
+	LARGE_INTEGER LastWriteTime;
+	LARGE_INTEGER ChangeTime;
+	LARGE_INTEGER AllocationSize;
+	LARGE_INTEGER EndOfFile;
+	ULONG FileAttributes;
+};
+
+constexpr FILE_INFORMATION_CLASS FileBasicInformation = (FILE_INFORMATION_CLASS)4;
+constexpr FILE_INFORMATION_CLASS FileNameInformation = (FILE_INFORMATION_CLASS)9;
 constexpr FILE_INFORMATION_CLASS FileRenameInformation = (FILE_INFORMATION_CLASS)10;
 constexpr FILE_INFORMATION_CLASS FileDispositionInformation = (FILE_INFORMATION_CLASS)13;
+constexpr FILE_INFORMATION_CLASS FileAllInformation = (FILE_INFORMATION_CLASS)18;
+constexpr FILE_INFORMATION_CLASS FileNetworkOpenInformation = (FILE_INFORMATION_CLASS)34;
+constexpr FILE_INFORMATION_CLASS FileNormalizedNameInformation = (FILE_INFORMATION_CLASS)35;
 constexpr FILE_INFORMATION_CLASS FileIdBothDirectoryInformation = (FILE_INFORMATION_CLASS)37;
+constexpr FILE_INFORMATION_CLASS FileIsRemoteDeviceInformation = (FILE_INFORMATION_CLASS)51;
+constexpr FILE_INFORMATION_CLASS FileIdInformation = (FILE_INFORMATION_CLASS)59;
 constexpr FILE_INFORMATION_CLASS FileRenameInformationEx = (FILE_INFORMATION_CLASS)65;
 
+constexpr OBJECT_INFORMATION_CLASS ObjectNameInformation = (OBJECT_INFORMATION_CLASS)1;
+
+enum FS_INFORMATION_CLASS {};
 
-constexpr OBJECT_INFORMATION_CLASS ObjectNameInformation = (OBJECT_INFORMATION_CLASS)1;
\ No newline at end of file
+constexpr FS_INFORMATION_CLASS FileFsVolumeInformation = (FS_INFORMATION_CLASS)1;
+constexpr FS_INFORMATION_CLASS FileFsDeviceInformation = (FS_INFORMATION_CLASS)4;
diff --git a/Engine/Source/Programs/UnrealBuildAccelerator/Detours/Private/Windows/UbaDetoursFunctionsKernelBase.inl b/Engine/Source/Programs/UnrealBuildAccelerator/Detours/Private/Windows/UbaDetoursFunctionsKernelBase.inl
--- a/Engine/Source/Programs/UnrealBuildAccelerator/Detours/Private/Windows/UbaDetoursFunctionsKernelBase.inl
+++ b/Engine/Source/Programs/UnrealBuildAccelerator/Detours/Private/Windows/UbaDetoursFunctionsKernelBase.inl
@@ -2,6 +2,15 @@
 
 #define UBA_ASSERT_ON_DETOURED_STD_PIPE 0 // Errors on clang on llvm-symbolizer.exe.. needs investigation
 
+static inline BOOL SetLastErrorFromNtStatus(NTSTATUS status)
+{
+	if (status)
+		SetLastError(RtlNtStatusToDosError(status));
+	else
+		SetLastError(ERROR_SUCCESS);
+	return status == 0;
+}
+
 DWORD Local_GetLongPathNameW(LPCWSTR lpszShortPath, LPWSTR lpszLongPath, DWORD cchBuffer)
 {
 	SCOPED_WRITE_LOCK(g_longPathNameCacheLock, lock);
@@ -643,6 +652,16 @@ DWORD Detoured_GetSecurityInfo(HANDLE handle, SE_OBJECT_TYPE ObjectType, SECURIT
 	return True_GetSecurityInfo(handle, ObjectType, SecurityInfo, ppsidOwner, ppsidGroup, ppDacl, ppSacl, ppSecurityDescriptor);
 }
 
+BOOL Detoured_PathSearchAndQualifyW(LPCWSTR pszPath, LPWSTR pszBuf, UINT cchBuf)
+{
+	FixPath2(pszPath, TStrlen(pszPath), g_virtualWorkingDir.data, g_virtualWorkingDir.count, pszBuf, cchBuf, nullptr);
+	DEBUG_LOG_DETOURED(L"PathSearchAndQualifyW", L"(%ls) -> %s (%s)", pszPath, ToString(true), pszBuf);
+	return true;
+	//BOOL res = True_PathSearchAndQualifyW(pszPath, pszBuf, cchBuf);
+	//DEBUG_LOG_TRUE(L"PathSearchAndQualifyW", L"(%ls) -> %s (%s)", pszPath, ToString(res), pszBuf);
+	//return res;
+}
+
 BOOL Detoured_InitializeSynchronizationBarrier(LPSYNCHRONIZATION_BARRIER lpBarrier, LONG lTotalThreads, LONG lSpinCount)
 {
 	Rpc_WriteLogf(TC("You need a newer version of wine to be able to call InitializeSynchronizationBarrier (%s)"), g_virtualApplication.data);
@@ -2893,7 +2912,8 @@ HMODULE Recursive_LoadLibraryExW(const StringView& filePath, LPCWSTR originalNam
 				path = lastSlash + 1;
 		u64 pathLen = wcslen(path);
 		StringBuffer<512> tempBuf;
-		Rpc_GetFullFileName(path, pathLen, tempBuf, false);
+		if (g_runningRemote)
+			Rpc_GetFullFileName(path, pathLen, tempBuf, false);
 
 		if (HMODULE r = Recursive_LoadLibraryExW(ToView(path), importedModule.name, dwFlags, additionalLoads, visitedModules))
 			additionalLoads.push_back(r);
@@ -4136,6 +4156,13 @@ LPWCH Detoured_GetEnvironmentStringsW()
 	return res;
 }
 
+DWORD Detoured_GetTempPathW(DWORD nBufferLength, LPWSTR lpBuffer)
+{
+	DETOURED_CALL(GetTempPathW);
+	DWORD res = True_GetTempPathW(nBufferLength, lpBuffer);
+	DEBUG_LOG_TRUE(L"GetTempPathW", L"%s -> %u", lpBuffer, res);
+	return res;
+}
 DWORD Detoured_ExpandEnvironmentStringsW(LPCWSTR lpSrc, LPWSTR lpDst, DWORD nSize)
 {
 	DETOURED_CALL(ExpandEnvironmentStringsW);
@@ -4143,6 +4170,13 @@ DWORD Detoured_ExpandEnvironmentStringsW(LPCWSTR lpSrc, LPWSTR lpDst, DWORD nSiz
 	return True_ExpandEnvironmentStringsW(lpSrc, lpDst, nSize);
 }
 
+HANDLE Detoured_FindFirstFileNameW(LPCWSTR lpFileName, DWORD dwFlags, LPDWORD StringLength, PWSTR LinkName)
+{
+	DETOURED_CALL(FindFirstFileNameW);
+	DEBUG_LOG_TRUE(L"FindFirstFileNameW", L"%ls", lpFileName);
+	return True_FindFirstFileNameW(lpFileName, dwFlags, StringLength, LinkName);
+}
+
 UINT Detoured_GetTempFileNameW(LPCWSTR lpPathName, LPCWSTR lpPrefixString, UINT uUnique, LPTSTR lpTempFileName)
 {
 	DETOURED_CALL(GetTempFileNameW);
@@ -4150,6 +4184,28 @@ UINT Detoured_GetTempFileNameW(LPCWSTR lpPathName, LPCWSTR lpPrefixString, UINT
 	return True_GetTempFileNameW(lpPathName, lpPrefixString, uUnique, lpTempFileName);
 }
 
+DWORD Detoured_GetShortPathNameA(LPCSTR lpszLongPath, LPSTR lpszShortPath, DWORD cchBuffer)
+{
+	DETOURED_CALL(GetShortPathNameA);
+	DWORD res = True_GetShortPathNameA(lpszLongPath, lpszShortPath, cchBuffer);
+	DEBUG_LOG_TRUE(L"GetShortPathNameA", L"%hs (%u) -> %u (%hs)", lpszLongPath, cchBuffer, res, res ? lpszShortPath : "error");
+	return res;
+}
+
+DWORD Detoured_GetTempPathA(DWORD nBufferLength, LPSTR lpBuffer)
+{
+	DETOURED_CALL(GetTempPathA);
+	DEBUG_LOG_TRUE(L"GetTempPathA", L"");
+	return True_GetTempPathA(nBufferLength, lpBuffer);
+}
+
+UINT Detoured_GetTempFileNameA(LPCSTR lpPathName, LPCSTR lpPrefixString, UINT uUnique, LPSTR lpTempFileName)
+{
+	DETOURED_CALL(GetTempFileNameA);
+	DEBUG_LOG_TRUE(L"GetTempFileNameA", L"");
+	return True_GetTempFileNameA(lpPathName, lpPrefixString, uUnique, lpTempFileName);
+}
+
 BOOL Detoured_CreateDirectoryExW(LPCWSTR lpTemplateDirectory, LPCWSTR lpNewDirectory, LPSECURITY_ATTRIBUTES lpSecurityAttributes)
 {
 	DETOURED_CALL(CreateDirectoryExW);
@@ -4157,6 +4213,7 @@ BOOL Detoured_CreateDirectoryExW(LPCWSTR lpTemplateDirectory, LPCWSTR lpNewDirec
 	return True_CreateDirectoryExW(lpTemplateDirectory, lpNewDirectory, lpSecurityAttributes);
 }
 
+/*
 BOOL Detoured_DecryptFileW(LPCWSTR lpFileName, DWORD dwReserved)
 {
 	DETOURED_CALL(DecryptFileW);
@@ -4198,7 +4255,7 @@ DWORD Detoured_OpenEncryptedFileRawA(LPCSTR lpFileName, ULONG ulFlags, PVOID* pv
 	DEBUG_LOG_TRUE(L"OpenEncryptedFileRawA", L"");
 	return True_OpenEncryptedFileRawA(lpFileName, ulFlags, pvContext);
 }
-
+*/
 HANDLE Detoured_OpenFileById(HANDLE hFile, LPFILE_ID_DESCRIPTOR lpFileID, DWORD dwDesiredAccess, DWORD dwShareMode, LPSECURITY_ATTRIBUTES lpSecurityAttributes, DWORD dwFlags)
 {
 	DETOURED_CALL(OpenFileById);
@@ -4261,7 +4318,9 @@ BOOL Detoured_CreatePipe(PHANDLE hReadPipe, PHANDLE hWritePipe, LPSECURITY_ATTRI
 {
 	DETOURED_CALL(CreatePipe);
 	DEBUG_LOG_TRUE(L"CreatePipe", L"");
-	return True_CreatePipe(hReadPipe, hWritePipe, lpPipeAttributes, nSize);
+	bool res = True_CreatePipe(hReadPipe, hWritePipe, lpPipeAttributes, nSize);
+	DEBUG_LOG_TRUE(L"CreatePipe", L"(POST) %llu %llu -> %s", u64(*hReadPipe), u64(*hWritePipe), ToString(res));
+	return res;
 }
 
 
@@ -4510,6 +4569,43 @@ HANDLE Detoured_CreateConsoleScreenBuffer(DWORD dwDesiredAccess, DWORD dwShareMo
 	return True_CreateConsoleScreenBuffer(dwDesiredAccess, dwShareMode, lpSecurityAttributes, dwFlags, lpScreenBufferData);
 }
 
+HANDLE Detoured_OpenProcess(DWORD dwDesiredAccess, BOOL bInheritHandle, DWORD dwProcessId)
+{
+	DETOURED_CALL(OpenProcess);
+	HANDLE h = True_OpenProcess(dwDesiredAccess, bInheritHandle, dwProcessId);
+	DEBUG_LOG_TRUE(L"OpenProcess", L"%u -> %llu", dwProcessId, u64(h));
+	return h;
+}
+
+PVOID Detoured_AddVectoredExceptionHandler(ULONG First, PVECTORED_EXCEPTION_HANDLER Handler)
+{
+	DEBUG_LOG_TRUE(L"AddVectoredExceptionHandler", L"");
+	return nullptr;
+}
+
+ULONG Detoured_RemoveVectoredExceptionHandler(PVOID Handle)
+{
+	DEBUG_LOG_TRUE(L"RemoveVectoredExceptionHandler", L"");
+	return 0;
+}
+
+VOID Detoured_RaiseFailFastException(PEXCEPTION_RECORD pExceptionRecord, PCONTEXT pContextRecord, DWORD dwFlags)
+{
+	DEBUG_LOG_TRUE(L"RaiseFailFastException", L"");
+	return True_RaiseFailFastException(pExceptionRecord, pContextRecord, dwFlags);
+}
+
+BOOL Detoured_TerminateJobObject(HANDLE hJob, UINT uExitCode)
+{
+	DEBUG_LOG_TRUE(L"TerminateJobObject", L"");
+	return True_TerminateJobObject(hJob, uExitCode);
+}
+
+UINT Detoured_SetErrorMode(UINT uMode)
+{
+	DEBUG_LOG_TRUE(L"SetErrorMode", L"%u", uMode);
+	return True_SetErrorMode(uMode);
+}
 BOOL Detoured_CreateProcessAsUserW(HANDLE hToken, LPCWSTR lpApplicationName, LPWSTR lpCommandLine, LPSECURITY_ATTRIBUTES lpProcessAttributes, LPSECURITY_ATTRIBUTES lpThreadAttributes, BOOL bInheritHandles, DWORD dwCreationFlags, LPVOID lpEnvironment, LPCWSTR lpCurrentDirectory, LPSTARTUPINFOW lpStartupInfo, LPPROCESS_INFORMATION lpProcessInformation)
 {
 	DETOURED_CALL(CreateProcessAsUserW);
@@ -4678,6 +4774,18 @@ BOOL Detoured_PathIsDirectoryEmptyW(LPCWSTR pszPath)
 	return res;
 }
 
+void Detoured_UrlCreateFromPathW(PCWSTR pszPath, PWSTR pszUrl, DWORD* pcchUrl, DWORD dwFlags)
+{
+	DEBUG_LOG_TRUE(L"UrlCreateFromPathW", L"(%ls)", pszPath);
+	True_UrlCreateFromPathW(pszPath, pszUrl, pcchUrl, dwFlags);
+}
+
+void Detoured_PathCreateFromUrlW(PCWSTR pszUrl, PWSTR pszPath, DWORD* pcchPath, DWORD dwFlags)
+{
+	DEBUG_LOG_TRUE(L"PathCreateFromUrlW", L"(%ls)", pszUrl);
+	True_PathCreateFromUrlW(pszUrl, pszPath, pcchPath, dwFlags);
+}
+
 HRESULT Detoured_SHCreateStreamOnFileW(LPCWSTR pszFile, DWORD grfMode, IStream** ppstm)
 {
 	//UBA_ASSERTF(!g_runningRemote, L"%ls", pszFile);
diff --git a/Engine/Source/Programs/UnrealBuildAccelerator/Detours/Private/Windows/UbaDetoursFunctionsNtDll.inl b/Engine/Source/Programs/UnrealBuildAccelerator/Detours/Private/Windows/UbaDetoursFunctionsNtDll.inl
--- a/Engine/Source/Programs/UnrealBuildAccelerator/Detours/Private/Windows/UbaDetoursFunctionsNtDll.inl
+++ b/Engine/Source/Programs/UnrealBuildAccelerator/Detours/Private/Windows/UbaDetoursFunctionsNtDll.inl
@@ -1,16 +1,5 @@
 // Copyright Epic Games, Inc. All Rights Reserved.
 
-struct FILE_NETWORK_OPEN_INFORMATION
-{
-	LARGE_INTEGER CreationTime;
-	LARGE_INTEGER LastAccessTime;
-	LARGE_INTEGER LastWriteTime;
-	LARGE_INTEGER ChangeTime;
-	LARGE_INTEGER AllocationSize;
-	LARGE_INTEGER EndOfFile;
-	ULONG FileAttributes;
-};
-
 bool IsContentWrite(u32 desiredAccess, u32 createDisposition)
 {
 	if (desiredAccess & (FILE_WRITE_DATA | FILE_APPEND_DATA | GENERIC_WRITE))
diff --git a/Engine/Source/Programs/UnrealBuildAccelerator/Detours/Private/Windows/UbaDetoursFunctionsWin.h b/Engine/Source/Programs/UnrealBuildAccelerator/Detours/Private/Windows/UbaDetoursFunctionsWin.h
--- a/Engine/Source/Programs/UnrealBuildAccelerator/Detours/Private/Windows/UbaDetoursFunctionsWin.h
+++ b/Engine/Source/Programs/UnrealBuildAccelerator/Detours/Private/Windows/UbaDetoursFunctionsWin.h
@@ -125,6 +125,7 @@
 	DETOURED_FUNCTION(VirtualAlloc) \
 	DETOURED_FUNCTION(GetQueuedCompletionStatusEx) \
 	DETOURED_FUNCTION(GetSecurityInfo) \
+	DETOURED_FUNCTION(PathSearchAndQualifyW) \
 	DETOURED_FUNCTIONS_KERNELBASE_DEBUG \
 
 #define DETOURED_FUNCTIONS_KERNEL32 \
@@ -257,6 +258,11 @@
 	DETOURED_FUNCTION(SetConsoleTextAttribute) \
 	DETOURED_FUNCTION(SetConsoleTitleW) \
 	DETOURED_FUNCTION(CreateConsoleScreenBuffer) \
+	DETOURED_FUNCTION(OpenProcess) \
+	DETOURED_FUNCTION(AddVectoredExceptionHandler) \
+	DETOURED_FUNCTION(RemoveVectoredExceptionHandler) \
+	DETOURED_FUNCTION(RaiseFailFastException) \
+	DETOURED_FUNCTION(SetErrorMode) \
 	DETOURED_FUNCTION(CreateProcessAsUserW) \
 	DETOURED_FUNCTION(SetConsoleCtrlHandler) \
 	DETOURED_FUNCTION(GetConsoleOutputCP) \
@@ -269,8 +275,6 @@
 	DETOURED_FUNCTION(GetFileAttributesA) \
 	DETOURED_FUNCTION(GetFileAttributesExA) \
 	DETOURED_FUNCTION_X64(LoadLibraryW) \
-	DETOURED_FUNCTION(SetDllDirectoryW) \
-	DETOURED_FUNCTION(GetDllDirectoryW) \
 	DETOURED_FUNCTION(GetModuleBaseNameA) \
 	DETOURED_FUNCTION(GetModuleBaseNameW) \
 	DETOURED_FUNCTION(SetUnhandledExceptionFilter) \
@@ -292,8 +296,12 @@
 	DETOURED_FUNCTION(GetEnvironmentVariableW) \
 	DETOURED_FUNCTION(GetEnvironmentVariableA) \
 	DETOURED_FUNCTION(GetEnvironmentStringsW) \
+	DETOURED_FUNCTION(GetTempPathW) \
 	DETOURED_FUNCTION(ExpandEnvironmentStringsW) \
+	DETOURED_FUNCTION(FindFirstFileNameW) \
 	DETOURED_FUNCTION(GetTempFileNameW) \
+	DETOURED_FUNCTION(GetTempPathA) \
+	DETOURED_FUNCTION(GetTempFileNameA) \
 	DETOURED_FUNCTION(CreateDirectoryExW) \
 	DETOURED_FUNCTION(CreateEventW) \
 	DETOURED_FUNCTION(CreateEventExW) \
@@ -313,23 +321,18 @@
 	DETOURED_FUNCTION(GetNamedPipeHandleStateW) \
 	DETOURED_FUNCTION(GetNamedPipeServerProcessId) \
 	DETOURED_FUNCTION(GetNamedPipeServerSessionId) \
-	DETOURED_FUNCTION(DecryptFileW) \
-	DETOURED_FUNCTION(DecryptFileA) \
-	DETOURED_FUNCTION(EncryptFileW) \
-	DETOURED_FUNCTION(EncryptFileA) \
-	DETOURED_FUNCTION(OpenEncryptedFileRawW) \
-	DETOURED_FUNCTION(OpenEncryptedFileRawA) \
 	DETOURED_FUNCTION(OpenFileById) \
 	DETOURED_FUNCTION(OpenFileMappingA) \
 	DETOURED_FUNCTION(GetMappedFileNameW) \
 	DETOURED_FUNCTION(IsProcessorFeaturePresent) \
 	DETOURED_FUNCTION(UnmapViewOfFile2) \
-	DETOURED_FUNCTION(GetComputerNameA) \
 	DETOURED_FUNCTION(GetModuleHandleW) \
 	DETOURED_FUNCTION(GetStartupInfoW) \
 	DETOURED_FUNCTION(GetSystemInfo) \
 	DETOURED_FUNCTION(GetSystemTimeAsFileTime) \
 	DETOURED_FUNCTION(GetLargePageMinimum) \
+	DETOURED_FUNCTION(UrlCreateFromPathW) \
+	DETOURED_FUNCTION(PathCreateFromUrlW) \
 	//DETOURED_FUNCTION(VirtualFree) \
 	//DETOURED_FUNCTION(BaseThreadInitThunk) \
 	//DETOURED_FUNCTION(VirtualAllocEx) \
@@ -339,6 +342,11 @@
 	//DETOURED_FUNCTION(CreateFileMappingNumaW)
 
 #define DETOURED_FUNCTIONS_KERNEL32_DEBUG \
+	DETOURED_FUNCTION(SetDllDirectoryW) \
+	DETOURED_FUNCTION(GetDllDirectoryW) \
+	DETOURED_FUNCTION(TerminateJobObject) \
+	DETOURED_FUNCTION(GetComputerNameA) \
+	DETOURED_FUNCTION(GetShortPathNameA) \
 
 #define DETOURED_FUNCTIONS_NTDLL_DEBUG \
 	DETOURED_FUNCTION(RtlAllocateHeap) \
@@ -378,6 +386,16 @@
 	//DETOURED_FUNCTION(_wsopen_s) \
 	//DETOURED_FUNCTION(_fileno)
 
+// Advapi32.dll
+/*
+DETOURED_FUNCTION(DecryptFileW) \
+DETOURED_FUNCTION(DecryptFileA) \
+DETOURED_FUNCTION(EncryptFileW) \
+DETOURED_FUNCTION(EncryptFileA) \
+DETOURED_FUNCTION(OpenEncryptedFileRawW) \
+DETOURED_FUNCTION(OpenEncryptedFileRawA) \
+*/
+
 #if UBA_USE_MIMALLOC
 #define DETOURED_FUNCTIONS_MEMORY_DEBUG \
 	DETOURED_FUNCTION(_aligned_msize) \
@@ -398,7 +416,6 @@ extern "C" {
 	using PALPC_PORT_ATTRIBUTES = void*;
 	using PALPC_MESSAGE_ATTRIBUTES = void*;
 	using PPORT_MESSAGE = void*;
-	enum FS_INFORMATION_CLASS {};
 	NTSTATUS NTAPI NtQueryVolumeInformationFile(HANDLE FileHandle, PIO_STATUS_BLOCK IoStatusBlock, PVOID FsInformation, ULONG Length, FS_INFORMATION_CLASS FsInformationClass);
 	NTSTATUS NTAPI NtQueryFullAttributesFile(POBJECT_ATTRIBUTES ObjectAttributes, PVOID Attributes);
 	NTSTATUS NTAPI NtQueryInformationFile(HANDLE FileHandle, PIO_STATUS_BLOCK IoStatusBlock, PVOID FileInformation, ULONG Length, FILE_INFORMATION_CLASS FileInformationClass);
