[UBA]
* Added detouring of DeviceIOControl

diff --git a/Engine/Source/Programs/UnrealBuildAccelerator/Detours/Private/Windows/UbaDetoursFunctionsKernelBase.inl b/Engine/Source/Programs/UnrealBuildAccelerator/Detours/Private/Windows/UbaDetoursFunctionsKernelBase.inl
--- a/Engine/Source/Programs/UnrealBuildAccelerator/Detours/Private/Windows/UbaDetoursFunctionsKernelBase.inl
+++ b/Engine/Source/Programs/UnrealBuildAccelerator/Detours/Private/Windows/UbaDetoursFunctionsKernelBase.inl
@@ -826,6 +826,33 @@ BOOL Detoured_WriteFileEx(HANDLE hFile, LPCVOID lpBuffer, DWORD nNumberOfBytesTo
 	return True_WriteFileEx(h.trueHandle, lpBuffer, nNumberOfBytesToWrite, lpOverlapped, lpCompletionRoutine);
 }
 
+BOOL Detoured_DeviceIoControl(HANDLE hDevice, DWORD dwIoControlCode, LPVOID lpInBuffer, DWORD nInBufferSize, LPVOID lpOutBuffer, DWORD  nOutBufferSize, LPDWORD lpBytesReturned, LPOVERLAPPED lpOverlapped)
+{
+	HANDLE trueHandle = hDevice;
+	if (isDetouredHandle(hDevice))
+	{
+		DetouredHandle& h = asDetouredHandle(hDevice);
+		if (h.trueHandle == INVALID_HANDLE_VALUE)
+		{
+			if (dwIoControlCode == FSCTL_GET_REPARSE_POINT)
+			{
+				//UBA_ASSERTF(false, TC("FSCTL_GET_REPARSE_POINT"));
+				return STATUS_NOT_A_REPARSE_POINT;
+			}
+			else if (dwIoControlCode == FSCTL_SET_SPARSE)
+			{
+				UBA_ASSERTF(false, TC("FSCTL_SET_SPARSE"));
+			}
+			else
+			{
+				UBA_ASSERTF(false, TC("DeviceIoControl on detoured handle with code %u"), dwIoControlCode);
+			}
+		}
+		trueHandle = h.trueHandle;
+	}
+	return True_DeviceIoControl(trueHandle, dwIoControlCode, lpInBuffer, nInBufferSize, lpOutBuffer, nOutBufferSize, lpBytesReturned, lpOverlapped);
+}
+
 BOOL Detoured_FlushFileBuffers(HANDLE hFile)
 {
 	DETOURED_CALL(FlushFileBuffers);
diff --git a/Engine/Source/Programs/UnrealBuildAccelerator/Detours/Private/Windows/UbaDetoursFunctionsWin.h b/Engine/Source/Programs/UnrealBuildAccelerator/Detours/Private/Windows/UbaDetoursFunctionsWin.h
--- a/Engine/Source/Programs/UnrealBuildAccelerator/Detours/Private/Windows/UbaDetoursFunctionsWin.h
+++ b/Engine/Source/Programs/UnrealBuildAccelerator/Detours/Private/Windows/UbaDetoursFunctionsWin.h
@@ -54,6 +54,7 @@
 	DETOURED_FUNCTION(ReadFile) \
 	DETOURED_FUNCTION(WriteFile) \
 	DETOURED_FUNCTION(WriteFileEx) \
+	DETOURED_FUNCTION(DeviceIoControl) \
 	DETOURED_FUNCTION(FlushFileBuffers) \
 	DETOURED_FUNCTION(GetFileSize) \
 	DETOURED_FUNCTION(GetFileSizeEx) \
