[UBA]
* Added validation of DeleteFileA and RemoveDirectoryA to check so they always go through W versions

diff --git a/Engine/Source/Programs/UnrealBuildAccelerator/Detours/Private/Windows/UbaDetoursFunctionsKernelBase.inl b/Engine/Source/Programs/UnrealBuildAccelerator/Detours/Private/Windows/UbaDetoursFunctionsKernelBase.inl
--- a/Engine/Source/Programs/UnrealBuildAccelerator/Detours/Private/Windows/UbaDetoursFunctionsKernelBase.inl
+++ b/Engine/Source/Programs/UnrealBuildAccelerator/Detours/Private/Windows/UbaDetoursFunctionsKernelBase.inl
@@ -287,10 +287,18 @@ BOOL Detoured_CreateDirectoryW(LPCWSTR lpPathName, LPSECURITY_ATTRIBUTES lpSecur
 	return res;
 }
 
+#if UBA_DEBUG
+thread_local bool t_calledRemoveDirectoryW;
+#endif
+
 BOOL Detoured_RemoveDirectoryW(LPCWSTR lpPathName)
 {
 	DETOURED_CALL(RemoveDirectoryW);
 
+	#if UBA_DEBUG
+	t_calledRemoveDirectoryW = 1;
+	#endif
+
 	StringBuffer<> pathName;
 	FixPath(pathName, lpPathName);
 
@@ -1677,9 +1685,17 @@ BOOL Detoured_CreateHardLinkW(LPCWSTR lpFileName, LPCWSTR lpExistingFileName, LP
 	return res;
 }
 
+#if UBA_DEBUG
+thread_local bool t_calledDeleteFileW;
+#endif
+
 BOOL Detoured_DeleteFileW(LPCWSTR lpFileName)
 {
 	DETOURED_CALL(DeleteFileW);
+	#if UBA_DEBUG
+	t_calledDeleteFileW = true;
+	#endif
+
 	LPCWSTR original = lpFileName;
 
 	if (!CanDetour(lpFileName))
@@ -4113,10 +4129,12 @@ BOOL Detoured_CreateHardLinkA(LPCSTR lpFileName, LPCSTR lpExistingFileName, LPSE
 
 BOOL Detoured_DeleteFileA(LPCSTR lpFileName)
 {
-	UBA_ASSERT(!g_runningRemote);
 	DETOURED_CALL(DeleteFileA);
 	DEBUG_LOG_TRUE(L"DeleteFileA", L"");
-	return True_DeleteFileA(lpFileName);
+	t_calledDeleteFileW = false;
+	BOOL res = True_DeleteFileA(lpFileName);
+	UBA_ASSERT(t_calledDeleteFileW);
+	return res;
 }
 
 BOOL Detoured_SetCurrentDirectoryA(LPCSTR lpPathName)
@@ -4240,6 +4258,13 @@ BOOL Detoured_CreateDirectoryExW(LPCWSTR lpTemplateDirectory, LPCWSTR lpNewDirec
 	return True_CreateDirectoryExW(lpTemplateDirectory, lpNewDirectory, lpSecurityAttributes);
 }
 
+BOOL Detoured_RemoveDirectoryA(LPCSTR lpPathName)
+{
+	t_calledRemoveDirectoryW = false;
+	BOOL res = True_RemoveDirectoryA(lpPathName);
+	UBA_ASSERT(t_calledRemoveDirectoryW);
+	return res;
+}
 /*
 BOOL Detoured_DecryptFileW(LPCWSTR lpFileName, DWORD dwReserved)
 {
diff --git a/Engine/Source/Programs/UnrealBuildAccelerator/Detours/Private/Windows/UbaDetoursFunctionsWin.h b/Engine/Source/Programs/UnrealBuildAccelerator/Detours/Private/Windows/UbaDetoursFunctionsWin.h
--- a/Engine/Source/Programs/UnrealBuildAccelerator/Detours/Private/Windows/UbaDetoursFunctionsWin.h
+++ b/Engine/Source/Programs/UnrealBuildAccelerator/Detours/Private/Windows/UbaDetoursFunctionsWin.h
@@ -304,6 +304,7 @@
 	DETOURED_FUNCTION(GetTempPathA) \
 	DETOURED_FUNCTION(GetTempFileNameA) \
 	DETOURED_FUNCTION(CreateDirectoryExW) \
+	DETOURED_FUNCTION(RemoveDirectoryA) \
 	DETOURED_FUNCTION(CreateEventW) \
 	DETOURED_FUNCTION(CreateEventExW) \
 	DETOURED_FUNCTION(CreateMutexExW) \
