[UBA]
* Added more debug detours in ntdll
* Fixed so we detour RtlExitUserProcess instead of ExitProcess since it is one step lower level

diff --git a/Engine/Source/Programs/UnrealBuildAccelerator/Detours/Private/Windows/UbaDetoursFunctionsKernelBase.inl b/Engine/Source/Programs/UnrealBuildAccelerator/Detours/Private/Windows/UbaDetoursFunctionsKernelBase.inl
--- a/Engine/Source/Programs/UnrealBuildAccelerator/Detours/Private/Windows/UbaDetoursFunctionsKernelBase.inl
+++ b/Engine/Source/Programs/UnrealBuildAccelerator/Detours/Private/Windows/UbaDetoursFunctionsKernelBase.inl
@@ -845,15 +845,19 @@ BOOL Detoured_DeviceIoControl(HANDLE hDevice, DWORD dwIoControlCode, LPVOID lpIn
 			if (dwIoControlCode == FSCTL_GET_REPARSE_POINT)
 			{
 				//UBA_ASSERTF(false, TC("FSCTL_GET_REPARSE_POINT"));
-				return STATUS_NOT_A_REPARSE_POINT;
+				return ERROR_NOT_A_REPARSE_POINT;
 			}
 			else if (dwIoControlCode == FSCTL_SET_SPARSE)
 			{
 				UBA_ASSERTF(false, TC("FSCTL_SET_SPARSE"));
 			}
+			else if (dwIoControlCode == FILE_DEVICE_MASS_STORAGE)
+			{
+				return ERROR_INVALID_FUNCTION;
+			}
 			else
 			{
-				UBA_ASSERTF(false, TC("DeviceIoControl on detoured handle with code %u"), dwIoControlCode);
+				UBA_ASSERTF(false, TC("DeviceIoControl on detoured handle with code %u (%s)"), dwIoControlCode, HandleToName(hDevice));
 			}
 		}
 		trueHandle = h.trueHandle;
@@ -3580,24 +3584,6 @@ BOOL Detoured_CreateProcessA(LPCSTR lpApplicationName, LPSTR lpCommandLine, LPSE
 	return Detoured_CreateProcessW(lpApplicationNameW, lpCommandLineW, lpProcessAttributes, lpThreadAttributes, bInheritHandles, dwCreationFlags, lpEnvironment, lpCurrentDirectoryW, &lpStartupInfoW, lpProcessInformation);
 }
 
-void Detoured_ExitProcess(UINT uExitCode)
-{
-	// Can't log this one
-	DETOURED_CALL(ExitProcess);
-	//DEBUG_LOG_TRUE(L"ExitProcess", L"(%u)", uExitCode);
-
-	{
-		CloseCaches();
-		SendExitMessage((DWORD)uExitCode, GetTime());
-		PostDeinit();
-	}
-
-#if UBA_DEBUG_LOG_ENABLED
-	FlushDebugLog();
-#endif
-	True_ExitProcess(uExitCode);
-}
-
 BOOL Detoured_TerminateProcess(HANDLE hProcess, UINT uExitCode)
 {
 	DETOURED_CALL(TerminateProcess);
@@ -4545,6 +4531,11 @@ SIZE_T Detoured_GetLargePageMinimum()
 //	return True_CreateFileMappingNumaW(hFile, lpFileMappingAttributes, flProtect, dwMaximumSizeHigh, dwMaximumSizeLow, lpName, nndPreferred);
 //}
 
+void Detoured_ExitProcess(UINT uExitCode)
+{
+	return True_ExitProcess(uExitCode);
+}
+
 LPSTR Detoured_GetCommandLineA()
 {
 	DETOURED_CALL(GetCommandLineA);
diff --git a/Engine/Source/Programs/UnrealBuildAccelerator/Detours/Private/Windows/UbaDetoursFunctionsNtDll.inl b/Engine/Source/Programs/UnrealBuildAccelerator/Detours/Private/Windows/UbaDetoursFunctionsNtDll.inl
--- a/Engine/Source/Programs/UnrealBuildAccelerator/Detours/Private/Windows/UbaDetoursFunctionsNtDll.inl
+++ b/Engine/Source/Programs/UnrealBuildAccelerator/Detours/Private/Windows/UbaDetoursFunctionsNtDll.inl
@@ -77,7 +77,7 @@ NTSTATUS Detoured_NtQueryVolumeInformationFile(HANDLE FileHandle, PIO_STATUS_BLO
 		auto& dh = asDetouredHandle(FileHandle);
 		if (dh.fileObject->fileInfo->memoryFile)
 		{
-			if (FsInformationClass == 4) // FileFsDeviceInformation
+			if (FsInformationClass == FileFsDeviceInformation)
 			{
 				auto& info = *(FILE_FS_DEVICE_INFORMATION*)FsInformation;
 				info.DeviceType = FILE_DEVICE_FILE_SYSTEM;
@@ -88,7 +88,7 @@ NTSTATUS Detoured_NtQueryVolumeInformationFile(HANDLE FileHandle, PIO_STATUS_BLO
 		TrueHandle = dh.trueHandle;
 		if (TrueHandle == INVALID_HANDLE_VALUE)
 		{
-			if (FsInformationClass == 1) // FileFsVolumeInformation
+			if (FsInformationClass == FileFsVolumeInformation)
 			{
 				// TODO This code path is here to handle nodejs queries..
 				
@@ -109,7 +109,7 @@ NTSTATUS Detoured_NtQueryVolumeInformationFile(HANDLE FileHandle, PIO_STATUS_BLO
 	}
 	else if (isListDirectoryHandle(FileHandle))
 	{
-		if (FsInformationClass == 4) // FileFsDeviceInformation
+		if (FsInformationClass == FileFsDeviceInformation)
 		{
 			auto& info = *(FILE_FS_DEVICE_INFORMATION*)FsInformation;
 			info.DeviceType = FILE_DEVICE_FILE_SYSTEM;
@@ -119,7 +119,7 @@ NTSTATUS Detoured_NtQueryVolumeInformationFile(HANDLE FileHandle, PIO_STATUS_BLO
 		UBA_ASSERTF(false, L"NtQueryVolumeInformationFile called in ListDirectoryHandle using class %u which is not implemented (%ls)", FsInformationClass, HandleToName(FileHandle));
 	}
 	auto res = True_NtQueryVolumeInformationFile(TrueHandle, IoStatusBlock, FsInformation, Length, FsInformationClass);
-	DEBUG_LOG_TRUE(L"NtQueryVolumeInformationFile", L"%llu (%ls) -> %ls", uintptr_t(FileHandle), HandleToName(FileHandle), ToString(res).data);
+	DEBUG_LOG_TRUE(L"NtQueryVolumeInformationFile", L"%llu (class %u) (%ls) -> %ls", uintptr_t(FileHandle), FsInformationClass, HandleToName(FileHandle), ToString(res).data);
 	return res;
 }
 
@@ -129,14 +129,14 @@ NTSTATUS Detoured_NtQueryInformationFile(HANDLE FileHandle, PIO_STATUS_BLOCK IoS
 	if (isListDirectoryHandle(FileHandle))
 	{
 		auto& listHandle = asListDirectoryHandle(FileHandle);
-		if (FileInformationClass == 51) // FileIsRemoteDeviceInformation
+		if (FileInformationClass == FileIsRemoteDeviceInformation)
 		{
 			auto& info = *(FILE_IS_REMOTE_DEVICE_INFORMATION*)FileInformation;
 			info.IsRemote = FALSE;
 			DEBUG_LOG_DETOURED(L"NtQueryInformationFile", L"(FileIsRemoteDeviceInformation) %llu (%ls) -> Success", uintptr_t(FileHandle), HandleToName(FileHandle));
 			return STATUS_SUCCESS;
 		}
-		else if (FileInformationClass == 59) // FileIdInformation
+		else if (FileInformationClass == FileIdInformation)
 		{
 			auto& info = *(FILE_ID_INFORMATION*)FileInformation;
 			if (listHandle.dir.tableOffset != InvalidTableOffset)
@@ -159,7 +159,7 @@ NTSTATUS Detoured_NtQueryInformationFile(HANDLE FileHandle, PIO_STATUS_BLOCK IoS
 			return STATUS_SUCCESS;
 		}
 		/*
-		else if (FileInformationClass == 9) // FileNameInformation
+		else if (FileInformationClass == FileNameInformation)
 		{
 			auto& info = *(FILE_NAME_INFORMATION*)FileInformation;
 			u32 nameLen = u32(wcslen(listHandle.name));
@@ -189,7 +189,7 @@ NTSTATUS Detoured_NtQueryInformationFile(HANDLE FileHandle, PIO_STATUS_BLOCK IoS
 
 		if (TrueHandle == INVALID_HANDLE_VALUE)
 		{
-			if (FileInformationClass == 18) // FileAllInformation 
+			if (FileInformationClass == FileAllInformation)
 			{
 				UBA_ASSERT(dh.dirTableOffset != ~0u);
 				DirectoryTable::EntryInformation entryInfo;
@@ -209,9 +209,10 @@ NTSTATUS Detoured_NtQueryInformationFile(HANDLE FileHandle, PIO_STATUS_BLOCK IoS
 				info.StandardInformation.DeletePending = false;
 				info.StandardInformation.Directory = false;
 				info.InternalInformation.IndexNumber.QuadPart = entryInfo.fileIndex;
+				DEBUG_LOG_DETOURED(L"NtQueryInformationFile", L"(FileAllInformation) %llu (%s) -> %ls", uintptr_t(FileHandle), HandleToName(FileHandle), ToString(STATUS_SUCCESS).data);
 				return STATUS_SUCCESS;
 			}
-			if (FileInformationClass == 34) // FileNetworkOpenInformation
+			if (FileInformationClass == FileNetworkOpenInformation)
 			{
 				UBA_ASSERT(dh.dirTableOffset != ~0u);
 				DirectoryTable::EntryInformation entryInfo;
@@ -230,6 +231,7 @@ NTSTATUS Detoured_NtQueryInformationFile(HANDLE FileHandle, PIO_STATUS_BLOCK IoS
 					fileSize = entryInfo.size;
 				info.EndOfFile.QuadPart = fileSize;
 				info.FileAttributes = entryInfo.attributes;
+				DEBUG_LOG_DETOURED(L"NtQueryInformationFile", L"(FileNetworkOpenInformation) %llu (%s) -> %ls", uintptr_t(FileHandle), HandleToName(FileHandle), ToString(STATUS_SUCCESS).data);
 				return STATUS_SUCCESS;
 			}
 
@@ -490,7 +492,7 @@ NTSTATUS Detoured_NtSetInformationFile(HANDLE FileHandle, PIO_STATUS_BLOCK IoSta
 	{
 		DetouredHandle& h = asDetouredHandle(FileHandle);
 
-		if (FileInformationClass == 10) // FileRenameInformation 
+		if (FileInformationClass == FileRenameInformation)
 		{
 			// We can end up in here through MoveFileEx
 			auto& info = *(FILE_RENAME_INFORMATION*)FileInformation;
@@ -619,6 +621,22 @@ NTSTATUS NTAPI Detoured_NtCreateSection(PHANDLE SectionHandle, ACCESS_MASK Desir
 	return True_NtCreateSection(SectionHandle, DesiredAccess, ObjectAttributes, MaximumSize, SectionPageProtection, AllocationAttributes, FileHandle);
 }
 
+VOID NTAPI Detoured_RtlExitUserProcess(NTSTATUS ExitStatus)
+{
+	// Can't log this one
+	DETOURED_CALL(ExitProcess);
+	//DEBUG_LOG_TRUE(L"RtlExitUserProcess", L"(%u)", uExitCode);
+
+	CloseCaches();
+	SendExitMessage((DWORD)ExitStatus, GetTime());
+	PostDeinit();
+	#if UBA_DEBUG_LOG_ENABLED
+	FlushDebugLog();
+	#endif
+
+	True_RtlExitUserProcess(ExitStatus);
+}
+
 bool g_checkRtlHeap = true;
 
 SIZE_T Detoured_RtlSizeHeap(HANDLE HeapPtr, ULONG Flags, PVOID Ptr)
@@ -1808,6 +1826,24 @@ NTSTATUS NTAPI Detoured_NtReadFile(HANDLE FileHandle, HANDLE Event, PIO_APC_ROUT
 	return True_NtReadFile(FileHandle, Event, ApcRoutine, ApcContext, IoStatusBlock, Buffer, Length, ByteOffset, Key);
 }
 
+NTSTATUS NTAPI Detoured_NtTerminateProcess(HANDLE ProcessHandle, NTSTATUS ExitStatus)
+{
+	/*
+	DEBUG_LOG_TRUE(L"NtTerminateProcess", L"");
+	
+	{
+		CloseCaches();
+		SendExitMessage((DWORD)0, GetTime());
+		PostDeinit();
+	}
+
+#if UBA_DEBUG_LOG_ENABLED
+	FlushDebugLog();
+#endif
+	*/
+	return True_NtTerminateProcess(ProcessHandle, ExitStatus);
+}
+
 NTSTATUS NTAPI Detoured_NtAlpcCreatePort(PHANDLE PortHandle, POBJECT_ATTRIBUTES ObjectAttributes, PALPC_PORT_ATTRIBUTES PortAttributes)
 {
 	//DEBUG_LOG_TRUE(L"NtAlpcCreatePort", L"");
@@ -1914,5 +1950,4 @@ NTSTATUS Detoured_RtlDosPathNameToNtPathName_U_WithStatus(PCWSTR dos_path, PUNIC
 	return res;
 }
 
-
 #endif
diff --git a/Engine/Source/Programs/UnrealBuildAccelerator/Detours/Private/Windows/UbaDetoursFunctionsWin.h b/Engine/Source/Programs/UnrealBuildAccelerator/Detours/Private/Windows/UbaDetoursFunctionsWin.h
--- a/Engine/Source/Programs/UnrealBuildAccelerator/Detours/Private/Windows/UbaDetoursFunctionsWin.h
+++ b/Engine/Source/Programs/UnrealBuildAccelerator/Detours/Private/Windows/UbaDetoursFunctionsWin.h
@@ -122,7 +122,6 @@
 	DETOURED_FUNCTION(WriteConsoleA) \
 	DETOURED_FUNCTION(WriteConsoleW) \
 	DETOURED_FUNCTION(ReadConsoleW) \
-	DETOURED_FUNCTION(ExitProcess) \
 	DETOURED_FUNCTION(VirtualAlloc) \
 	DETOURED_FUNCTION(GetQueuedCompletionStatusEx) \
 	DETOURED_FUNCTION(GetSecurityInfo) \
@@ -155,6 +154,7 @@
 	DETOURED_FUNCTION(NtSetInformationFile) \
 	DETOURED_FUNCTION(NtSetInformationObject) \
 	DETOURED_FUNCTION(NtCreateSection) \
+	DETOURED_FUNCTION(RtlExitUserProcess) \
 	DETOURED_FUNCTION(RtlSizeHeap) \
 	DETOURED_FUNCTION(RtlFreeHeap) \
 	DETOURED_FUNCTION(RtlAnsiStringToUnicodeString) \
@@ -237,6 +237,7 @@
 #if defined(DETOURED_INCLUDE_DEBUG)
 
 #define DETOURED_FUNCTIONS_KERNELBASE_DEBUG \
+	DETOURED_FUNCTION(ExitProcess) \
 	DETOURED_FUNCTION(GetCommandLineA) \
 	DETOURED_FUNCTION(CommandLineToArgvW) \
 	DETOURED_FUNCTION(FreeLibrary) \
@@ -358,6 +359,7 @@
 	DETOURED_FUNCTION(NtCreateIoCompletion) \
 	DETOURED_FUNCTION(NtFlushBuffersFileEx) \
 	DETOURED_FUNCTION(NtReadFile) \
+	DETOURED_FUNCTION(NtTerminateProcess) \
 	DETOURED_FUNCTION(NtAlpcCreatePort) \
 	DETOURED_FUNCTION(NtAlpcConnectPort) \
 	DETOURED_FUNCTION(NtAlpcCreatePortSection) \
@@ -426,10 +428,12 @@ extern "C" {
 	NTSTATUS NTAPI NtCopyFileChunk(HANDLE Source, HANDLE Dest, HANDLE Event, PIO_STATUS_BLOCK IoStatusBlock, ULONG Length, PULONG SourceOffset, PULONG DestOffset, PULONG SourceKey, PULONG DestKey, ULONG Flags);
 	NTSTATUS NTAPI NtFlushBuffersFileEx(HANDLE FileHandle, ULONG Flags, PVOID Parameters, ULONG ParametersSize, PIO_STATUS_BLOCK IoStatusBlock);
 	NTSTATUS NTAPI NtReadFile(HANDLE FileHandle, HANDLE Event, PIO_APC_ROUTINE ApcRoutine, PVOID ApcContext, PIO_STATUS_BLOCK IoStatusBlock, PVOID Buffer, ULONG Length, PLARGE_INTEGER ByteOffset, PULONG Key);
+	NTSTATUS NTAPI NtTerminateProcess(HANDLE ProcessHandle, NTSTATUS ExitStatus);
 	NTSTATUS NTAPI NtSetInformationFile(HANDLE FileHandle, PIO_STATUS_BLOCK IoStatusBlock, PVOID FileInformation, ULONG Length, FILE_INFORMATION_CLASS FileInformationClass);
 	NTSTATUS NTAPI NtSetInformationObject(HANDLE ObjectHandle, OBJECT_INFORMATION_CLASS ObjectInformationClass, PVOID ObjectInformation, ULONG Length);
 	NTSTATUS NTAPI NtCreateSection(PHANDLE SectionHandle, ACCESS_MASK DesiredAccess, POBJECT_ATTRIBUTES ObjectAttributes, PLARGE_INTEGER MaximumSize, ULONG SectionPageProtection, ULONG AllocationAttributes, HANDLE FileHandle);
 	NTSTATUS NTAPI NtCreateIoCompletion(PHANDLE IoCompletionHandle, ACCESS_MASK DesiredAccess, POBJECT_ATTRIBUTES ObjectAttributes, ULONG Count);
+	NTSTATUS NTAPI NtQuerySecurityObject(HANDLE Handle, SECURITY_INFORMATION SecurityInformation, PSECURITY_DESCRIPTOR SecurityDescriptor, ULONG Length, PULONG LengthNeeded);
 	NTSTATUS NTAPI NtAlpcCreatePort(PHANDLE PortHandle, POBJECT_ATTRIBUTES ObjectAttributes, PALPC_PORT_ATTRIBUTES PortAttributes);
 	NTSTATUS NTAPI NtAlpcConnectPort(PHANDLE PortHandle, PUNICODE_STRING PortName, POBJECT_ATTRIBUTES ObjectAttributes, PALPC_PORT_ATTRIBUTES PortAttributes, DWORD ConnectionFlags, PSID RequiredServerSid, PPORT_MESSAGE ConnectionMessage, PSIZE_T ConnectMessageSize, PALPC_MESSAGE_ATTRIBUTES OutMessageAttributes, PALPC_MESSAGE_ATTRIBUTES InMessageAttributes, PLARGE_INTEGER Timeout);
 	NTSTATUS NTAPI NtAlpcCreatePortSection(HANDLE PortHandle, ULONG Flags, HANDLE SectionHandle, SIZE_T SectionSize, PHANDLE AlpcSectionHandle, PSIZE_T ActualSectionSize);
@@ -451,7 +455,8 @@ extern "C" {
 	BOOLEAN NTAPI RtlFreeHeap(PVOID HeapHandle, ULONG Flags, PVOID HeapBase);
 	BOOLEAN NTAPI RtlValidateHeap(HANDLE HeapPtr, ULONG Flags, PVOID Block);
 	NTSTATUS NTAPI RtlDosPathNameToNtPathName_U_WithStatus(PCWSTR dos_path, PUNICODE_STRING ntpath, PWSTR* file_part, VOID* reserved);
-
+	VOID NTAPI RtlExitUserProcess(NTSTATUS ExitStatus);
+	VOID NTSYSAPI RtlMapGenericMask(PACCESS_MASK AccessMask, const GENERIC_MAPPING *GenericMapping);
 	PVOID NTAPI RtlAllocateHeap( PVOID HeapHandle, ULONG Flags, SIZE_T Size);
 	PVOID NTAPI RtlReAllocateHeap(PVOID HeapHandle, ULONG Flags, PVOID BaseAddress, SIZE_T Size);
 	SIZE_T NTAPI RtlSizeHeap(HANDLE HeapPtr, ULONG Flags, PVOID Ptr);
