[UBA]
* Fixed bug in SearchPathForFile where it could crash if PATH environment variable contained a path that was empty (ended with a semicolon as an example)

diff --git a/Engine/Source/Programs/UnrealBuildAccelerator/Common/Private/UbaFile.cpp b/Engine/Source/Programs/UnrealBuildAccelerator/Common/Private/UbaFile.cpp
--- a/Engine/Source/Programs/UnrealBuildAccelerator/Common/Private/UbaFile.cpp
+++ b/Engine/Source/Programs/UnrealBuildAccelerator/Common/Private/UbaFile.cpp
@@ -560,10 +560,12 @@ namespace uba
 
 			*it = 0;
 
-			fullPath.Clear().Append(lastStart);
-			if (TestFileExists(TC(" using PATH environment variable")))
-				return true;
-
+			if (*lastStart)
+			{
+				fullPath.Clear().Append(lastStart);
+				if (TestFileExists(TC(" using PATH environment variable")))
+					return true;
+			}
 			lastStart = ++it;
 		}
 		return false;
