[UBA]
* Removed debug log in scheduler that was accidentally left behind

diff --git a/Engine/Source/Programs/UnrealBuildAccelerator/Common/Private/UbaScheduler.cpp b/Engine/Source/Programs/UnrealBuildAccelerator/Common/Private/UbaScheduler.cpp
--- a/Engine/Source/Programs/UnrealBuildAccelerator/Common/Private/UbaScheduler.cpp
+++ b/Engine/Source/Programs/UnrealBuildAccelerator/Common/Private/UbaScheduler.cpp
@@ -478,7 +478,6 @@ namespace uba
 		#if UBA_DEBUG
 		if (m_cancelled)
 		{
-			m_session.GetLogger().Error(TC("STATUS"));
 			u32 index = 0;
 			for (auto& entry : m_processEntries)
 			{
