[UBA]
* Nicer implementation of SearchPathForFile when traversing patg environment variable

diff --git a/Engine/Source/Programs/UnrealBuildAccelerator/Common/Private/UbaFile.cpp b/Engine/Source/Programs/UnrealBuildAccelerator/Common/Private/UbaFile.cpp
--- a/Engine/Source/Programs/UnrealBuildAccelerator/Common/Private/UbaFile.cpp
+++ b/Engine/Source/Programs/UnrealBuildAccelerator/Common/Private/UbaFile.cpp
@@ -557,12 +557,9 @@ namespace uba
 				++it;
 				continue;
 			}
-
-			*it = 0;
-
-			if (*lastStart)
+			if (u32 varLen = u32(it - lastStart))
 			{
-				fullPath.Clear().Append(lastStart);
+				fullPath.Clear().Append(lastStart, varLen);
 				if (TestFileExists(TC(" using PATH environment variable")))
 					return true;
 			}
