[UBA]
* Fixed detoured DeviceIOControl to handle requests that link.exe does

diff --git a/Engine/Source/Programs/UnrealBuildAccelerator/Detours/Private/Windows/UbaDetoursFunctionsKernelBase.inl b/Engine/Source/Programs/UnrealBuildAccelerator/Detours/Private/Windows/UbaDetoursFunctionsKernelBase.inl
--- a/Engine/Source/Programs/UnrealBuildAccelerator/Detours/Private/Windows/UbaDetoursFunctionsKernelBase.inl
+++ b/Engine/Source/Programs/UnrealBuildAccelerator/Detours/Private/Windows/UbaDetoursFunctionsKernelBase.inl
@@ -851,17 +851,30 @@ BOOL Detoured_DeviceIoControl(HANDLE hDevice, DWORD dwIoControlCode, LPVOID lpIn
 			{
 				UBA_ASSERTF(false, TC("FSCTL_SET_SPARSE"));
 			}
-			else if (dwIoControlCode == FILE_DEVICE_MASS_STORAGE)
+			else if (dwIoControlCode == IOCTL_STORAGE_GET_DEVICE_NUMBER)
 			{
-				return ERROR_INVALID_FUNCTION;
+				// This is used by link.exe to figure out if it can use some faster paths to track unique files
+				const wchar_t* originalName = h.fileObject->fileInfo->originalName;
+				UBA_ASSERT(originalName && originalName[0]);
+				DWORD deviceNumber = originalName[0] % 32u;
+
+				auto& number = *(STORAGE_DEVICE_NUMBER*)lpOutBuffer;
+				number.DeviceType = FILE_DEVICE_DISK;
+				number.DeviceNumber = deviceNumber;
+				number.PartitionNumber = 0;
+				if (lpBytesReturned)
+					*lpBytesReturned = sizeof(STORAGE_DEVICE_NUMBER);
+				DEBUG_LOG_DETOURED(L"DeviceIoControl", L"(IOCTL_STORAGE_GET_DEVICE_NUMBER) %llu (%ls) -> Success", uintptr_t(hDevice), HandleToName(hDevice));
+				return TRUE;
 			}
 			else
 			{
-				UBA_ASSERTF(false, TC("DeviceIoControl on detoured handle with code %u (%s)"), dwIoControlCode, HandleToName(hDevice));
+				UBA_ASSERTF(false, TC("DeviceIoControl on detoured handle with code %u not handled (%s)"), dwIoControlCode, HandleToName(hDevice));
 			}
 		}
 		trueHandle = h.trueHandle;
 	}
+	DEBUG_LOG_TRUE(L"DeviceIoControl", L"%llu (%ls) -> Success", uintptr_t(hDevice), HandleToName(hDevice));
 	return True_DeviceIoControl(trueHandle, dwIoControlCode, lpInBuffer, nInBufferSize, lpOutBuffer, nOutBufferSize, lpBytesReturned, lpOverlapped);
 }
 
