[UBA]
* Moved ReadFile implementation to NtReadFile and added checks so NtReadFile is always called from ReadFile

diff --git a/Engine/Source/Programs/UnrealBuildAccelerator/Detours/Private/Windows/UbaDetoursFunctionsKernelBase.inl b/Engine/Source/Programs/UnrealBuildAccelerator/Detours/Private/Windows/UbaDetoursFunctionsKernelBase.inl
--- a/Engine/Source/Programs/UnrealBuildAccelerator/Detours/Private/Windows/UbaDetoursFunctionsKernelBase.inl
+++ b/Engine/Source/Programs/UnrealBuildAccelerator/Detours/Private/Windows/UbaDetoursFunctionsKernelBase.inl
@@ -391,71 +391,6 @@ BOOL Detoured_UnlockFileEx(HANDLE hFile, DWORD dwReserved, DWORD nNumberOfBytesT
 	return True_UnlockFileEx(trueHandle, dwReserved, nNumberOfBytesToUnlockLow, nNumberOfBytesToUnlockHigh, lpOverlapped);
 }
 
-BOOL Detoured_ReadFile(HANDLE hFile, LPVOID lpBuffer, DWORD nNumberOfBytesToRead, LPDWORD lpNumberOfBytesRead, LPOVERLAPPED lpOverlapped)
-{
-	DETOURED_CALL(ReadFile);
-	HANDLE trueHandle = hFile;
-	UBA_ASSERT(!isListDirectoryHandle(hFile));
-	if (isDetouredHandle(hFile))
-	{
-		auto& dh = asDetouredHandle(hFile);
-
-		if (dh.type == HandleType_StdIn) // HACK HACK
-		{
-			UBA_ASSERTF(false, L"Trying to read input from stdin while application is running in a way console can not be accessed");
-			memcpy(lpBuffer, "Y\r\n", 3);
-			*lpNumberOfBytesRead = 3;
-			return TRUE;
-		}
-
-		auto& fo = *dh.fileObject;
-		fo.wasUsed = true;
-
-		FileInfo& fi = *fo.fileInfo;
-
-		if (fi.isFileMap || fi.memoryFile)
-		{
-			// TODO: Handle lpOverlapped - If a read happen and there is 0 left it should return 0 with SetLastError(ERROR_HANDLE_EOF)
-			if (!EnsureMapped(dh))
-			{
-				DEBUG_LOG_DETOURED(L"ReadFile", L"%llu %u (%ls) -> FAILED TO MAP", uintptr_t(hFile), nNumberOfBytesToRead, HandleToName(hFile));
-				return FALSE;
-			}
-			UBA_ASSERTF(fi.fileMapMem || !fi.memoryFile->isThrowAway, TC("Trying to read throw-away file %s"), HandleToName(hFile));
-
-			u8* mem = fi.fileMapMem ? fi.fileMapMem : fi.memoryFile->baseAddress;
-			u64 size = fi.fileMapMem ? fi.fileMapMemSize : fi.memoryFile->writtenSize;
-			u64 readPos = dh.pos;
-			if (lpOverlapped)
-			{
-				UBA_ASSERT(!lpOverlapped->hEvent);
-				readPos = ToLargeInteger(lpOverlapped->OffsetHigh, lpOverlapped->Offset).QuadPart;
-			}
-			UBA_ASSERTF(readPos <= size, L"ReadFile - Read position is higher than size of file (pointer: %llu, size: %llu) (%ls)", readPos, size, HandleToName(hFile));
-			u64 leftToRead = size - readPos;
-			if (nNumberOfBytesToRead > leftToRead)
-				nNumberOfBytesToRead = (DWORD)leftToRead;
-			if (nNumberOfBytesToRead)
-				memcpy(lpBuffer, mem + readPos, nNumberOfBytesToRead);
-			if (!lpOverlapped)
-				dh.pos += nNumberOfBytesToRead;
-			if (lpNumberOfBytesRead)
-				*lpNumberOfBytesRead = nNumberOfBytesToRead;
-			SetLastError(ERROR_SUCCESS);
-			DEBUG_LOG_DETOURED(L"ReadFile", L"%llu %u (%ls) -> Success", uintptr_t(hFile), nNumberOfBytesToRead, HandleToName(hFile));
-			return TRUE;
-		}
-		UBA_ASSERT(dh.trueHandle != INVALID_HANDLE_VALUE);
-		trueHandle = dh.trueHandle;
-	}
-
-	TimerScope ts(g_kernelStats.readFile);
-	BOOL res = True_ReadFile(trueHandle, lpBuffer, nNumberOfBytesToRead, lpNumberOfBytesRead, lpOverlapped);
-	DEBUG_LOG_TRUE(L"ReadFile", L"%llu %u/%u (%ls) -> %ls", uintptr_t(hFile), lpNumberOfBytesRead ? *lpNumberOfBytesRead : ~0u, nNumberOfBytesToRead, HandleToName(hFile), ToString(res));
-	return res;
-}
-
-
 BOOL Detoured_WriteConsoleA(HANDLE hConsoleOutput, const VOID* lpBuffer, DWORD nNumberOfCharsToWrite, LPDWORD lpNumberOfCharsWritten, LPVOID lpReserved)
 {
 	DETOURED_CALL(WriteConsoleA);
@@ -4062,14 +3997,23 @@ HANDLE Detoured_ReOpenFile(HANDLE hOriginalFile, DWORD dwDesiredAccess, DWORD dw
 	return True_ReOpenFile(hOriginalFile, dwDesiredAccess, dwShareMode, dwFlagsAndAttributes);
 }
 
+BOOL Detoured_ReadFile(HANDLE hFile, LPVOID lpBuffer, DWORD nNumberOfBytesToRead, LPDWORD lpNumberOfBytesRead, LPOVERLAPPED lpOverlapped)
+{
+	DETOURED_CALL(ReadFile);
+	t_ntReadFileCalled = false;
+	BOOL res = True_ReadFile(hFile, lpBuffer, nNumberOfBytesToRead, lpNumberOfBytesRead, lpOverlapped);
+	UBA_ASSERT(t_ntReadFileCalled);
+	return res;
+}
+
 BOOL Detoured_ReadFileEx(HANDLE hFile, LPVOID lpBuffer, DWORD nNumberOfBytesToRead, LPOVERLAPPED lpOverlapped, LPOVERLAPPED_COMPLETION_ROUTINE lpCompletionRoutine)
 {
 	DETOURED_CALL(ReadFileEx);
 	DEBUG_LOG_TRUE(L"ReadFileEx", L"%llu (%ls)", uintptr_t(hFile), HandleToName(hFile));
-	UBA_ASSERT(!isDetouredHandle(hFile));
-	UBA_ASSERT(!isListDirectoryHandle(hFile));
-	TimerScope ts(g_kernelStats.readFile);
-	return True_ReadFileEx(hFile, lpBuffer, nNumberOfBytesToRead, lpOverlapped, lpCompletionRoutine);
+	t_ntReadFileCalled = false;
+	BOOL res = True_ReadFileEx(hFile, lpBuffer, nNumberOfBytesToRead, lpOverlapped, lpCompletionRoutine);
+	UBA_ASSERT(t_ntReadFileCalled);
+	return res;
 }
 
 BOOL Detoured_ReadFileScatter(HANDLE hFile, FILE_SEGMENT_ELEMENT* aSegmentArray, DWORD nNumberOfBytesToRead, LPDWORD lpReserved, LPOVERLAPPED lpOverlapped)
diff --git a/Engine/Source/Programs/UnrealBuildAccelerator/Detours/Private/Windows/UbaDetoursFunctionsNtDll.inl b/Engine/Source/Programs/UnrealBuildAccelerator/Detours/Private/Windows/UbaDetoursFunctionsNtDll.inl
--- a/Engine/Source/Programs/UnrealBuildAccelerator/Detours/Private/Windows/UbaDetoursFunctionsNtDll.inl
+++ b/Engine/Source/Programs/UnrealBuildAccelerator/Detours/Private/Windows/UbaDetoursFunctionsNtDll.inl
@@ -1500,6 +1500,95 @@ NTSTATUS NTAPI Detoured_NtOpenFile(PHANDLE FileHandle, ACCESS_MASK DesiredAccess
 	return Shared_NtCreateFile(false, FileHandle, DesiredAccess, ObjectAttributes, IoStatusBlock, (PLARGE_INTEGER)NULL, 0L, ShareAccess, FILE_OPEN, OpenOptions, (PVOID)NULL, 0L);
 }
 
+NTSTATUS NTAPI Detoured_NtReadFile(HANDLE FileHandle, HANDLE Event, PIO_APC_ROUTINE ApcRoutine, PVOID ApcContext, PIO_STATUS_BLOCK IoStatusBlock, PVOID Buffer, ULONG Length, PLARGE_INTEGER ByteOffset, PULONG Key)
+{
+	DETOURED_CALL(NtReadFile);
+
+	#if UBA_DEBUG
+	t_ntReadFileCalled = true;
+	#endif
+
+	auto trueRun = [&](HANDLE fileHandle)
+		{
+			TimerScope ts(g_kernelStats.readFile);
+			NTSTATUS status = True_NtReadFile(fileHandle, Event, ApcRoutine, ApcContext, IoStatusBlock, Buffer, Length, ByteOffset, Key);
+			DEBUG_LOG_TRUE(L"NtReadFile", L"%llu (%ls) -> %s", uintptr_t(fileHandle), HandleToName(fileHandle), ToString(status).data);
+			return status;
+		};
+
+	UBA_ASSERT(!isListDirectoryHandle(FileHandle));
+	if (!isDetouredHandle(FileHandle))
+		return trueRun(FileHandle);
+
+	auto& dh = asDetouredHandle(FileHandle);
+
+	if (dh.type == HandleType_StdIn)
+	{
+		UBA_ASSERTF(false, L"Trying to read input from stdin while application is running in a way console can not be accessed");
+		memcpy(Buffer, "Y\r\n", 3);
+		IoStatusBlock->Status = STATUS_SUCCESS;
+		IoStatusBlock->Information = 3;
+		return STATUS_SUCCESS;
+	}
+
+	UBA_ASSERT(dh.type != HandleType_StdIn);
+	auto& fo = *dh.fileObject;
+	fo.wasUsed = true;
+
+	FileInfo& fi = *fo.fileInfo;
+	if (!fi.isFileMap && !fi.memoryFile)
+		return trueRun(dh.trueHandle);
+
+	// TODO: Handle lpOverlapped - If a read happen and there is 0 left it should return 0 with SetLastError(ERROR_HANDLE_EOF)
+	if (!EnsureMapped(dh))
+	{
+		DEBUG_LOG_DETOURED(L"FileHandle", L"%llu %u (%ls) -> FAILED TO MAP", uintptr_t(FileHandle), Length, HandleToName(FileHandle));
+		return FALSE;
+	}
+	UBA_ASSERTF(fi.fileMapMem || !fi.memoryFile->isThrowAway, TC("Trying to read throw-away file %s"), HandleToName(FileHandle));
+
+	u8* mem = fi.fileMapMem ? fi.fileMapMem : fi.memoryFile->baseAddress;
+	u64 size = fi.fileMapMem ? fi.fileMapMemSize : fi.memoryFile->writtenSize;
+	u64 readPos = dh.pos;
+
+	NTSTATUS status = STATUS_SUCCESS;
+	u64 bytesRead = 0;
+
+	if (ByteOffset != nullptr)
+		readPos = ByteOffset->QuadPart;
+
+	if (readPos < 0)
+		status = STATUS_INVALID_PARAMETER;
+	else if (readPos >= size)
+	{
+		status = STATUS_END_OF_FILE;
+		bytesRead = 0;
+	}
+	else
+	{
+		u64 availableBytes = size - readPos;
+		bytesRead = Min((u64)Length, availableBytes);
+		memcpy(Buffer, (u8*)mem + readPos, bytesRead);
+		if (ByteOffset == nullptr)
+			dh.pos = readPos + bytesRead;
+	}
+
+	if (IoStatusBlock != nullptr)
+	{
+		IoStatusBlock->Status = status;
+		IoStatusBlock->Information = bytesRead;
+	}
+
+	if (Event != nullptr)
+		SetEvent(Event);
+
+	if (ApcRoutine != nullptr)
+		ApcRoutine(ApcContext, IoStatusBlock, 0);
+
+	DEBUG_LOG_DETOURED(L"NtReadFile", L"(MemoryFile) %llu (%ls) %llu length=%u -> %s", uintptr_t(FileHandle), HandleToName(FileHandle), readPos, Length, ToString(status).data);
+	return status;
+}
+
 NTSTATUS NTAPI Detoured_NtFsControlFile(HANDLE FileHandle, HANDLE Event, PIO_APC_ROUTINE ApcRoutine, PVOID ApcContext, PIO_STATUS_BLOCK IoStatusBlock, ULONG FsControlCode, PVOID InputBuffer, ULONG InputBufferLength, PVOID OutputBuffer, ULONG OutputBufferLength)
 {
 	DETOURED_CALL(NtFsControlFile);
@@ -1819,28 +1908,32 @@ NTSTATUS NTAPI Detoured_NtFlushBuffersFileEx(HANDLE FileHandle, ULONG Flags, PVO
 	return True_NtFlushBuffersFileEx(FileHandle, Flags, Parameters, ParametersSize, IoStatusBlock);
 }
 
-NTSTATUS NTAPI Detoured_NtReadFile(HANDLE FileHandle, HANDLE Event, PIO_APC_ROUTINE ApcRoutine, PVOID ApcContext, PIO_STATUS_BLOCK IoStatusBlock, PVOID Buffer, ULONG Length, PLARGE_INTEGER ByteOffset, PULONG Key)
+#if UBA_DEBUG
+thread_local bool t_ntReadFileCalled;
+#endif
+
+NTSTATUS NTAPI Detoured_NtWriteFile(HANDLE FileHandle, HANDLE Event, PIO_APC_ROUTINE ApcRoutine, PVOID ApcContext, PIO_STATUS_BLOCK IoStatusBlock, PVOID Buffer, ULONG Length, PLARGE_INTEGER ByteOffset, PULONG Key)
 {
 	DETOURED_CALL(NtReadFile);
-	UBA_ASSERT(!isDetouredHandle(FileHandle));
-	return True_NtReadFile(FileHandle, Event, ApcRoutine, ApcContext, IoStatusBlock, Buffer, Length, ByteOffset, Key);
+	UBA_ASSERT(!isListDirectoryHandle(FileHandle));
+	if (!isDetouredHandle(FileHandle))
+		return True_NtWriteFile(FileHandle, Event, ApcRoutine, ApcContext, IoStatusBlock, Buffer, Length, ByteOffset, Key);
+	auto& dh = asDetouredHandle(FileHandle);
+	UBA_ASSERT(dh.type == HandleType_File);
+	auto& fo = *dh.fileObject;
+	fo.wasUsed = true;
+	FileInfo& fi = *fo.fileInfo;
+	UBA_ASSERT(!fi.isFileMap && !fi.memoryFile);
+	HANDLE trueHandle = dh.trueHandle;
+	UBA_ASSERT(trueHandle != INVALID_HANDLE_VALUE);
+	NTSTATUS res = True_NtWriteFile(trueHandle, Event, ApcRoutine, ApcContext, IoStatusBlock, Buffer, Length, ByteOffset, Key);
+	DEBUG_LOG_TRUE(L"NtWriteFile", L"%llu (%ls) -> %s", uintptr_t(FileHandle), HandleToName(FileHandle), ToString(res).data);
+	return res;
 }
 
 NTSTATUS NTAPI Detoured_NtTerminateProcess(HANDLE ProcessHandle, NTSTATUS ExitStatus)
 {
-	/*
-	DEBUG_LOG_TRUE(L"NtTerminateProcess", L"");
-	
-	{
-		CloseCaches();
-		SendExitMessage((DWORD)0, GetTime());
-		PostDeinit();
-	}
-
-#if UBA_DEBUG_LOG_ENABLED
-	FlushDebugLog();
-#endif
-	*/
+	// NtTerminateProcess runs after DllMain DLL_PROCESS_DETACH which means we can't use this to send exit messaage etc
 	return True_NtTerminateProcess(ProcessHandle, ExitStatus);
 }
 
diff --git a/Engine/Source/Programs/UnrealBuildAccelerator/Detours/Private/Windows/UbaDetoursFunctionsWin.h b/Engine/Source/Programs/UnrealBuildAccelerator/Detours/Private/Windows/UbaDetoursFunctionsWin.h
--- a/Engine/Source/Programs/UnrealBuildAccelerator/Detours/Private/Windows/UbaDetoursFunctionsWin.h
+++ b/Engine/Source/Programs/UnrealBuildAccelerator/Detours/Private/Windows/UbaDetoursFunctionsWin.h
@@ -51,7 +51,6 @@
 	DETOURED_FUNCTION(LockFileEx) \
 	DETOURED_FUNCTION(UnlockFile) \
 	DETOURED_FUNCTION(UnlockFileEx) \
-	DETOURED_FUNCTION(ReadFile) \
 	DETOURED_FUNCTION(WriteFile) \
 	DETOURED_FUNCTION(WriteFileEx) \
 	DETOURED_FUNCTION(DeviceIoControl) \
@@ -143,6 +142,7 @@
 	DETOURED_FUNCTION(NtClose) \
 	DETOURED_FUNCTION(NtCreateFile) \
 	DETOURED_FUNCTION(NtOpenFile) \
+	DETOURED_FUNCTION(NtReadFile) \
 	DETOURED_FUNCTION(NtFsControlFile) \
 	DETOURED_FUNCTION(NtCopyFileChunk) \
 	DETOURED_FUNCTION(NtQueryVolumeInformationFile) \
@@ -285,6 +285,7 @@
 	DETOURED_FUNCTION(CreateFileTransactedW) \
 	DETOURED_FUNCTION(OpenFile) \
 	DETOURED_FUNCTION(ReOpenFile) \
+	DETOURED_FUNCTION(ReadFile) \
 	DETOURED_FUNCTION(ReadFileEx) \
 	DETOURED_FUNCTION(ReadFileScatter) \
 	DETOURED_FUNCTION(SetFileValidData) \
@@ -358,7 +359,7 @@
 	DETOURED_FUNCTION(RtlDosPathNameToNtPathName_U_WithStatus) \
 	DETOURED_FUNCTION(NtCreateIoCompletion) \
 	DETOURED_FUNCTION(NtFlushBuffersFileEx) \
-	DETOURED_FUNCTION(NtReadFile) \
+	DETOURED_FUNCTION(NtWriteFile) \
 	DETOURED_FUNCTION(NtTerminateProcess) \
 	DETOURED_FUNCTION(NtAlpcCreatePort) \
 	DETOURED_FUNCTION(NtAlpcConnectPort) \
@@ -428,6 +429,7 @@ extern "C" {
 	NTSTATUS NTAPI NtCopyFileChunk(HANDLE Source, HANDLE Dest, HANDLE Event, PIO_STATUS_BLOCK IoStatusBlock, ULONG Length, PULONG SourceOffset, PULONG DestOffset, PULONG SourceKey, PULONG DestKey, ULONG Flags);
 	NTSTATUS NTAPI NtFlushBuffersFileEx(HANDLE FileHandle, ULONG Flags, PVOID Parameters, ULONG ParametersSize, PIO_STATUS_BLOCK IoStatusBlock);
 	NTSTATUS NTAPI NtReadFile(HANDLE FileHandle, HANDLE Event, PIO_APC_ROUTINE ApcRoutine, PVOID ApcContext, PIO_STATUS_BLOCK IoStatusBlock, PVOID Buffer, ULONG Length, PLARGE_INTEGER ByteOffset, PULONG Key);
+	NTSTATUS NTAPI NtWriteFile(HANDLE FileHandle, HANDLE Event, PIO_APC_ROUTINE ApcRoutine, PVOID ApcContext, PIO_STATUS_BLOCK IoStatusBlock, PVOID Buffer, ULONG Length, PLARGE_INTEGER ByteOffset, PULONG Key);
 	NTSTATUS NTAPI NtTerminateProcess(HANDLE ProcessHandle, NTSTATUS ExitStatus);
 	NTSTATUS NTAPI NtSetInformationFile(HANDLE FileHandle, PIO_STATUS_BLOCK IoStatusBlock, PVOID FileInformation, ULONG Length, FILE_INFORMATION_CLASS FileInformationClass);
 	NTSTATUS NTAPI NtSetInformationObject(HANDLE ObjectHandle, OBJECT_INFORMATION_CLASS ObjectInformationClass, PVOID ObjectInformation, ULONG Length);
