[UBA]
* Added detoured NtQuerySecurityObject

diff --git a/Engine/Source/Programs/UnrealBuildAccelerator/Detours/Private/Windows/UbaDetoursFunctionsKernelBase.inl b/Engine/Source/Programs/UnrealBuildAccelerator/Detours/Private/Windows/UbaDetoursFunctionsKernelBase.inl
--- a/Engine/Source/Programs/UnrealBuildAccelerator/Detours/Private/Windows/UbaDetoursFunctionsKernelBase.inl
+++ b/Engine/Source/Programs/UnrealBuildAccelerator/Detours/Private/Windows/UbaDetoursFunctionsKernelBase.inl
@@ -4349,10 +4349,11 @@ BOOL Detoured_PeekNamedPipe(HANDLE hNamedPipe, LPVOID lpBuffer, DWORD nBufferSiz
 
 BOOL Detoured_GetKernelObjectSecurity(HANDLE Handle, SECURITY_INFORMATION RequestedInformation, PSECURITY_DESCRIPTOR pSecurityDescriptor, DWORD nLength, LPDWORD lpnLengthNeeded)
 {
-	HANDLE trueHandle = Handle;
-	if (isDetouredHandle(Handle))
-		trueHandle = asDetouredHandle(Handle).trueHandle;
-	return True_GetKernelObjectSecurity(trueHandle, RequestedInformation, pSecurityDescriptor, nLength, lpnLengthNeeded);
+	DEBUG_LOG_TRUE(L"GetKernelObjectSecurity", L"%llu", uintptr_t(Handle));
+	t_ntQuerySecurityObjectCalled = false;
+	BOOL res = True_GetKernelObjectSecurity(Handle, RequestedInformation, pSecurityDescriptor, nLength, lpnLengthNeeded);
+	UBA_ASSERT(t_ntQuerySecurityObjectCalled);
+	return res;
 }
 
 BOOL Detoured_ImpersonateNamedPipeClient(HANDLE hNamedPipe)
diff --git a/Engine/Source/Programs/UnrealBuildAccelerator/Detours/Private/Windows/UbaDetoursFunctionsNtDll.inl b/Engine/Source/Programs/UnrealBuildAccelerator/Detours/Private/Windows/UbaDetoursFunctionsNtDll.inl
--- a/Engine/Source/Programs/UnrealBuildAccelerator/Detours/Private/Windows/UbaDetoursFunctionsNtDll.inl
+++ b/Engine/Source/Programs/UnrealBuildAccelerator/Detours/Private/Windows/UbaDetoursFunctionsNtDll.inl
@@ -1512,28 +1512,23 @@ NTSTATUS NTAPI Detoured_NtReadFile(HANDLE FileHandle, HANDLE Event, PIO_APC_ROUT
 	t_ntReadFileCalled = true;
 	#endif
 
-	auto trueRun = [&](HANDLE fileHandle)
+	auto trueRun = [&](HANDLE fileHandle, const tchar* prefix)
 		{
+			(void)prefix;
 			TimerScope ts(g_kernelStats.readFile);
 			NTSTATUS status = True_NtReadFile(fileHandle, Event, ApcRoutine, ApcContext, IoStatusBlock, Buffer, Length, ByteOffset, Key);
-			DEBUG_LOG_TRUE(L"NtReadFile", L"%llu (%ls) -> %s", uintptr_t(fileHandle), HandleToName(fileHandle), ToString(status).data);
+			DEBUG_LOG_TRUE(L"NtReadFile", L"%s%llu (%ls) -> %s", prefix, uintptr_t(fileHandle), HandleToName(fileHandle), ToString(status).data);
 			return status;
 		};
 
 	UBA_ASSERT(!isListDirectoryHandle(FileHandle));
 	if (!isDetouredHandle(FileHandle))
-		return trueRun(FileHandle);
+		return trueRun(FileHandle, TC(""));
 
 	auto& dh = asDetouredHandle(FileHandle);
 
 	if (dh.type == HandleType_StdIn)
-	{
-		UBA_ASSERTF(false, L"Trying to read input from stdin while application is running in a way console can not be accessed");
-		memcpy(Buffer, "Y\r\n", 3);
-		IoStatusBlock->Status = STATUS_SUCCESS;
-		IoStatusBlock->Information = 3;
-		return STATUS_SUCCESS;
-	}
+		return trueRun(True_GetStdHandle(STD_INPUT_HANDLE), TC("(STDIN)"));
 
 	UBA_ASSERT(dh.type != HandleType_StdIn);
 	auto& fo = *dh.fileObject;
@@ -1541,7 +1536,7 @@ NTSTATUS NTAPI Detoured_NtReadFile(HANDLE FileHandle, HANDLE Event, PIO_APC_ROUT
 
 	FileInfo& fi = *fo.fileInfo;
 	if (!fi.isFileMap && !fi.memoryFile)
-		return trueRun(dh.trueHandle);
+		return trueRun(dh.trueHandle, TC(""));
 
 	// TODO: Handle lpOverlapped - If a read happen and there is 0 left it should return 0 with SetLastError(ERROR_HANDLE_EOF)
 	if (!EnsureMapped(dh))
@@ -1897,6 +1892,114 @@ NTSTATUS Detoured_NtQueryInformationProcess(HANDLE ProcessHandle, PROCESSINFOCLA
 	return res;
 }
 
+#if UBA_DEBUG
+thread_local bool t_ntQuerySecurityObjectCalled;
+#endif
+
+NTSTATUS Detoured_NtQuerySecurityObject(HANDLE Handle, SECURITY_INFORMATION SecurityInformation, PSECURITY_DESCRIPTOR SecurityDescriptor, ULONG Length, PULONG LengthNeeded)
+{
+	DETOURED_CALL(NtQuerySecurityObject);
+	
+	#if UBA_DEBUG
+	t_ntQuerySecurityObjectCalled = true;
+	#endif
+
+	auto callTrue = [&](HANDLE handle)
+		{
+			NTSTATUS res = True_NtQuerySecurityObject(handle, SecurityInformation, SecurityDescriptor, Length, LengthNeeded);
+			DEBUG_LOG_TRUE(L"NtQuerySecurityObject", L"%llu -> %ls", uintptr_t(Handle), ToString(res).data);
+			return res;
+		};
+
+	if (!isDetouredHandle(Handle))
+		return callTrue(Handle);
+
+	HANDLE trueHandle = asDetouredHandle(Handle).trueHandle;
+	if (trueHandle != INVALID_HANDLE_VALUE)
+		return callTrue(trueHandle);
+
+	BYTE sys[SECURITY_MAX_SID_SIZE], adm[SECURITY_MAX_SID_SIZE], app[SECURITY_MAX_SID_SIZE], res[SECURITY_MAX_SID_SIZE];
+	DWORD s=sizeof(sys);
+	CreateWellKnownSid(WinLocalSystemSid,0,sys,&s);
+	s=sizeof(adm);
+	CreateWellKnownSid(WinBuiltinAdministratorsSid,0,adm,&s);
+
+	SID_IDENTIFIER_AUTHORITY auth = SECURITY_APP_PACKAGE_AUTHORITY;
+	InitializeSid((PSID)app,&auth,2);
+	*GetSidSubAuthority((PSID)app,0)=2;
+	*GetSidSubAuthority((PSID)app,1)=1;
+	InitializeSid((PSID)res,&auth,2);
+	*GetSidSubAuthority((PSID)res,0)=2;
+	*GetSidSubAuthority((PSID)res,1)=2;
+
+	// get current user SID
+	BYTE userSid[SECURITY_MAX_SID_SIZE];
+	//DWORD userSidLen = sizeof(userSid);
+	HANDLE token;
+	OpenProcessToken(GetCurrentProcess(), TOKEN_QUERY, &token);
+	DWORD len;
+	GetTokenInformation(token, TokenUser, nullptr, 0, &len);
+	BYTE* tokenBuf = (BYTE*)_alloca(len);
+	GetTokenInformation(token, TokenUser, tokenBuf, len, &len);
+	PSID srcUser = ((TOKEN_USER*)tokenBuf)->User.Sid;
+	CopySid(sizeof(userSid), userSid, srcUser);
+	CloseHandle(token);
+
+	ULONG aclSz = sizeof(ACL) +
+		sizeof(ACCESS_ALLOWED_ACE)-sizeof(DWORD)+GetLengthSid(app) +
+		sizeof(ACCESS_ALLOWED_ACE)-sizeof(DWORD)+GetLengthSid(res) +
+		sizeof(ACCESS_ALLOWED_ACE)-sizeof(DWORD)+GetLengthSid(userSid) +
+		sizeof(ACCESS_ALLOWED_ACE)-sizeof(DWORD)+GetLengthSid(adm) +
+		sizeof(ACCESS_ALLOWED_ACE)-sizeof(DWORD)+GetLengthSid(adm) +
+		sizeof(ACCESS_ALLOWED_ACE)-sizeof(DWORD)+GetLengthSid(sys) +
+		sizeof(ACCESS_ALLOWED_ACE)-sizeof(DWORD)+GetLengthSid(sys);
+
+	ULONG sdSz = sizeof(SECURITY_DESCRIPTOR_RELATIVE) +
+		GetLengthSid(userSid) + GetLengthSid(sys) + aclSz;
+
+	if(LengthNeeded)*LengthNeeded=sdSz;
+	if(Length<sdSz) return STATUS_BUFFER_TOO_SMALL;
+
+	BYTE* p=(BYTE*)SecurityDescriptor;
+	auto sd=(SECURITY_DESCRIPTOR_RELATIVE*)p;
+	p+=sizeof(*sd);
+
+	// owner = current user
+	BYTE* owner=p;
+	CopySid(GetLengthSid(userSid),owner,userSid);
+	p+=GetLengthSid(userSid);
+
+	// group = SYSTEM (same as before)
+	BYTE* group=p;
+	CopySid(GetLengthSid(sys),group,sys);
+	p+=GetLengthSid(sys);
+
+	ACL* acl=(ACL*)p;
+	InitializeAcl(acl,aclSz,ACL_REVISION);
+
+	AddAccessAllowedAceEx(acl,ACL_REVISION,OBJECT_INHERIT_ACE|CONTAINER_INHERIT_ACE,FILE_GENERIC_READ|FILE_GENERIC_EXECUTE,app);
+	AddAccessAllowedAceEx(acl,ACL_REVISION,OBJECT_INHERIT_ACE|CONTAINER_INHERIT_ACE,FILE_GENERIC_READ|FILE_GENERIC_EXECUTE,res);
+
+	// current user full access
+	AddAccessAllowedAceEx(acl,ACL_REVISION,0,FILE_ALL_ACCESS,userSid);
+
+	AddAccessAllowedAceEx(acl,ACL_REVISION,INHERITED_ACE,FILE_ALL_ACCESS,adm);
+	AddAccessAllowedAceEx(acl,ACL_REVISION,INHERITED_ACE|OBJECT_INHERIT_ACE|CONTAINER_INHERIT_ACE|INHERIT_ONLY_ACE,FILE_ALL_ACCESS,adm);
+
+	AddAccessAllowedAceEx(acl,ACL_REVISION,INHERITED_ACE,FILE_ALL_ACCESS,sys);
+	AddAccessAllowedAceEx(acl,ACL_REVISION,INHERITED_ACE|OBJECT_INHERIT_ACE|CONTAINER_INHERIT_ACE|INHERIT_ONLY_ACE,FILE_ALL_ACCESS,sys);
+
+	sd->Revision=SECURITY_DESCRIPTOR_REVISION;
+	sd->Control=SE_SELF_RELATIVE|SE_DACL_PRESENT;
+	sd->Owner=(ULONG)(owner-(BYTE*)sd);
+	sd->Group=(ULONG)(group-(BYTE*)sd);
+	sd->Dacl =(ULONG)((BYTE*)acl-(BYTE*)sd);
+
+	DEBUG_LOG_DETOURED(L"NtQuerySecurityObject", L"%llu -> %ls", uintptr_t(Handle), ToString(STATUS_SUCCESS).data);
+	return STATUS_SUCCESS;
+}
+
+
 #if defined(DETOURED_INCLUDE_DEBUG)
 
 NTSTATUS NTAPI Detoured_NtCreateIoCompletion(PHANDLE IoCompletionHandle, ACCESS_MASK DesiredAccess, POBJECT_ATTRIBUTES ObjectAttributes, ULONG Count)
diff --git a/Engine/Source/Programs/UnrealBuildAccelerator/Detours/Private/Windows/UbaDetoursFunctionsWin.h b/Engine/Source/Programs/UnrealBuildAccelerator/Detours/Private/Windows/UbaDetoursFunctionsWin.h
--- a/Engine/Source/Programs/UnrealBuildAccelerator/Detours/Private/Windows/UbaDetoursFunctionsWin.h
+++ b/Engine/Source/Programs/UnrealBuildAccelerator/Detours/Private/Windows/UbaDetoursFunctionsWin.h
@@ -151,6 +151,7 @@
 	DETOURED_FUNCTION(NtQueryFullAttributesFile) \
 	DETOURED_FUNCTION(NtQueryObject) \
 	DETOURED_FUNCTION(NtQueryInformationProcess) \
+	DETOURED_FUNCTION(NtQuerySecurityObject) \
 	DETOURED_FUNCTION(NtSetInformationFile) \
 	DETOURED_FUNCTION(NtSetInformationObject) \
 	DETOURED_FUNCTION(NtCreateSection) \
