[UBA]
* Lots of fixes in detoured CreateProcess handling various command line types

diff --git a/Engine/Source/Programs/UnrealBuildAccelerator/Detours/Private/Windows/UbaDetoursFunctionsKernelBase.inl b/Engine/Source/Programs/UnrealBuildAccelerator/Detours/Private/Windows/UbaDetoursFunctionsKernelBase.inl
--- a/Engine/Source/Programs/UnrealBuildAccelerator/Detours/Private/Windows/UbaDetoursFunctionsKernelBase.inl
+++ b/Engine/Source/Programs/UnrealBuildAccelerator/Detours/Private/Windows/UbaDetoursFunctionsKernelBase.inl
@@ -3251,8 +3251,9 @@ BOOL Detoured_CreateProcessW(LPCWSTR lpApplicationName, LPWSTR lpCommandLine, LP
 		isChild = false;
 	}
 
+	tchar* tempAllocated = nullptr;
 	tchar* finalCommandLine = lpCommandLine;
-	auto g = MakeGuard([&]() { if (finalCommandLine != lpCommandLine) delete[] finalCommandLine; });
+	auto g = MakeGuard([&]() { if (tempAllocated) delete[] tempAllocated; });
 
 	StringView commandLineWithoutApplication;
 	{
@@ -3267,7 +3268,7 @@ BOOL Detoured_CreateProcessW(LPCWSTR lpApplicationName, LPWSTR lpCommandLine, LP
 			// so if there is "link.exe"/LIB as a command it will not produce proper wmain(argv)
 			if (*endOfApplication != ' ')
 			{
-				finalCommandLine = new tchar[originalCmd.count + 1];
+				finalCommandLine = tempAllocated = new tchar[originalCmd.count + 1];
 				u64 firstPart = endOfApplication - originalCmd.data;
 				memcpy(finalCommandLine, originalCmd.data, firstPart*sizeof(tchar));
 				finalCommandLine[firstPart] = ' ';
@@ -3296,6 +3297,8 @@ BOOL Detoured_CreateProcessW(LPCWSTR lpApplicationName, LPWSTR lpCommandLine, LP
 		StringBuffer<512> temp;
 		if (originalCmd.data[0] == '"')
 			temp.Append(originalCmd.data+1, commandLineWithoutApplication.data - originalCmd.data - 2);
+		else if (commandLineWithoutApplication.count == 0)
+			temp.Append(originalCmd);
 		else
 			temp.Append(originalCmd.data, commandLineWithoutApplication.data - originalCmd.data - 1);
 		if (temp[1] == ':')
@@ -3303,6 +3306,8 @@ BOOL Detoured_CreateProcessW(LPCWSTR lpApplicationName, LPWSTR lpCommandLine, LP
 		else
 			application.Append(temp.data);
 		UBA_ASSERTF(application.count, TC("Invalid application from command line (name: %s)"), originalCmd.data);
+		if (application.count < 4 || application[application.count - 4] != '.')
+			application.Append(TCV(".exe")); // TODO: Maybe also handle other extensions?
 	}
 	DevirtualizePath(application);
 
@@ -3328,6 +3333,34 @@ BOOL Detoured_CreateProcessW(LPCWSTR lpApplicationName, LPWSTR lpCommandLine, LP
 		return True_CreateProcessW(cmdExe.data, (LPWSTR)TC("cmd.exe /c exit 0"), lpProcessAttributes, lpThreadAttributes, bInheritHandles, dwCreationFlags, lpEnvironment, workingDir.data, lpStartupInfo, lpProcessInformation);
 	}
 
+	// Need to turn into cmd.exe manually since we don't it to tamper with the path (remote execution will end up using the true path instead of virtualized)
+	if (application.EndsWith(TCV(".bat")))
+	{
+		//UBA_ASSERT(false);
+		StringView extra(TCV("cmd.exe /C "));
+		UBA_ASSERT(!application.Contains(' '));
+		u64 len = finalCommandLine ? TStrlen(finalCommandLine) : 0;
+		tchar* newCommandLine = new tchar[application.count + extra.count + len + 1];
+		memcpy(newCommandLine, extra.data, application.count*sizeof(tchar));
+		memcpy(newCommandLine + extra.count, application.data, application.count*sizeof(tchar));
+		memcpy(newCommandLine + extra.count + application.count, finalCommandLine ? finalCommandLine : TC(""), (len + 1)*sizeof(tchar));
+		delete[] tempAllocated;
+		finalCommandLine = tempAllocated = newCommandLine;
+		application.Clear().Append(TCV("C:\\Windows\\system32\\cmd.exe"));
+	}
+
+	bool is32Bit = application.EndsWith(TC("gperf.exe")); // Hard coded for 32-bit application gperf.exe... this could be extended but will only work for apps that does not procuce output files since we don't track non-detoured applications
+	if (is32Bit)
+	{
+		const wchar_t* realApplication = application.data;
+		u64 pathLen = application.count;
+		StringBuffer<512> tempBuf;
+		if (g_runningRemote)
+			Rpc_GetFullFileName(realApplication, pathLen, tempBuf, false);
+		SuppressDetourScope _;
+		return True_CreateProcessW(realApplication, finalCommandLine, lpProcessAttributes, lpThreadAttributes, bInheritHandles, dwCreationFlags, lpEnvironment, lpCurrentDirectory, lpStartupInfo, lpProcessInformation);
+	}
+
 	TString currentDir;
 	u32 processId = 0;
 	char dll[1024];
@@ -3417,25 +3450,69 @@ BOOL Detoured_CreateProcessW(LPCWSTR lpApplicationName, LPWSTR lpCommandLine, LP
 		}
 	}
 
-	//if (dwCreationFlags & EXTENDED_STARTUPINFO_PRESENT) // 0x00080000
-	//{
-	//	auto& siex = *(STARTUPINFOEXW*)lpStartupInfo;
-	//	(void)siex;
-	//}
+	bool allowCreationFlagChange = true;
+	if (dwCreationFlags & EXTENDED_STARTUPINFO_PRESENT)
+	{
+		struct PROC_THREAD_ATTRIBUTE { DWORD_PTR Attribute; SIZE_T Size; ULONG_PTR Value; };
+		struct PROC_THREAD_ATTRIBUTE_LIST_INTERNAL { DWORD Unknown1; DWORD AttributeCount; DWORD Unknown2; DWORD Unknown3; DWORD Unknown4; DWORD Unknown5; PROC_THREAD_ATTRIBUTE Attributes[1]; };
+		auto& siex = *(STARTUPINFOEXW*)lpStartupInfo;
+		if (auto attrList = (PROC_THREAD_ATTRIBUTE_LIST_INTERNAL*)siex.lpAttributeList)
+		{
+			for (DWORD i = 0; i < attrList->AttributeCount; i++)
+			{
+				DWORD_PTR attr = attrList->Attributes[i].Attribute;
+				SIZE_T size = attrList->Attributes[i].Size;
+				ULONG_PTR value = attrList->Attributes[i].Value;
+				DWORD attrType = attr & 0x0000FFFF;
+				if (attrType == ProcThreadAttributeHandleList)
+				{
+					HANDLE* handles = (HANDLE*)value;
+					DWORD handleCount = (DWORD)(size / sizeof(HANDLE));
+					for (DWORD h = 0; h < handleCount; h++)
+					{
+						if (isDetouredHandle(handles[h]))
+						{
+							DetouredHandle& dh = asDetouredHandle(handles[h]);
+							UBA_ASSERT(dh.trueHandle != INVALID_HANDLE_VALUE);
+							handles[h] = dh.trueHandle;
+						}
+						
+					}
+					allowCreationFlagChange = false;
+				}
+				else if (attrType == ProcThreadAttributeParentProcess)
+				{
+					UBA_ASSERT(!isDetouredHandle(*(HANDLE*)value));
+				}
+				else if (attrType == 1) // ProcThreadAttributeExtendedFlags
+				{
+					//struct PS_ATTRIBUTE { ULONG  Attribute; SIZE_T Size; union { ULONG_PTR Value; PVOID ValuePtr; }; PSIZE_T ReturnLength; };
+				}
+				else
+				{
+					// Don't know what 1 is, Ignore
+					UBA_ASSERTF(false, TC("Check if thread attribute %u needs to be handled"), attrType);
+				}
+			}
+		}
+	}
 
 	// CREATE_NO_WINDOW is an optimization that prevents additional overhead per spawned child process.
 	// .. disabled for now because there are hierarchies of processes that does not like this. Haven't figured out exactly why
 	// .. can be reproed by running some of the very complex cmake/python scripts when compiling clang (lto.exports etc)
 	bool createNoWindow = false; //!g_setStdHandleCalled; // If console mode has been called there are probably some sort of stdin/out rerouting stuff.. so let's not set CREATE_NO_WINDOW
 
-	if (g_rules->AllowDetach())
-		dwCreationFlags |= DETACHED_PROCESS;
-	else if (createNoWindow)
-		dwCreationFlags |= CREATE_NO_WINDOW;
-	else
+	if (allowCreationFlagChange && (dwCreationFlags & (DETACHED_PROCESS | CREATE_NO_WINDOW)) == 0)
 	{
-		lpStartupInfo->dwFlags |= STARTF_USESHOWWINDOW;
-		lpStartupInfo->wShowWindow = SW_HIDE;
+		if (g_rules->AllowDetach())
+			dwCreationFlags |= DETACHED_PROCESS;
+		else if (createNoWindow)
+			dwCreationFlags |= CREATE_NO_WINDOW;
+		else
+		{
+			lpStartupInfo->dwFlags |= STARTF_USESHOWWINDOW;
+			lpStartupInfo->wShowWindow = SW_HIDE;
+		}
 	}
 
 	dwCreationFlags |= CREATE_SUSPENDED;
@@ -3468,6 +3545,16 @@ BOOL Detoured_CreateProcessW(LPCWSTR lpApplicationName, LPWSTR lpCommandLine, LP
 	}
 	--t_disallowDetour;
 
+	#if UBA_DEBUG
+	USHORT processMachine = 0;
+	USHORT nativeMachine = 0;
+	if (IsWow64Process2(lpProcessInformation->hProcess, &processMachine, &nativeMachine))
+	{
+		bool is64Bit = processMachine == IMAGE_FILE_MACHINE_UNKNOWN;
+		UBA_ASSERTF(is64Bit, TC("Can't detour 32-bit application %s"), application.data);
+	}
+	#endif
+
 	if (isChild)
 	{
 		RPC_MESSAGE(StartProcess, createProcess)
