[UBA]
* Added support for more paths (classes) in detour nt layer. NtQueryVolumeInformationFile, NtQueryInformationFile, NtQueryDirectoryFile, NtQueryObject

diff --git a/Engine/Source/Programs/UnrealBuildAccelerator/Detours/Private/Windows/UbaDetoursFunctionsNtDll.inl b/Engine/Source/Programs/UnrealBuildAccelerator/Detours/Private/Windows/UbaDetoursFunctionsNtDll.inl
--- a/Engine/Source/Programs/UnrealBuildAccelerator/Detours/Private/Windows/UbaDetoursFunctionsNtDll.inl
+++ b/Engine/Source/Programs/UnrealBuildAccelerator/Detours/Private/Windows/UbaDetoursFunctionsNtDll.inl
@@ -116,6 +116,23 @@ NTSTATUS Detoured_NtQueryVolumeInformationFile(HANDLE FileHandle, PIO_STATUS_BLO
 			info.Characteristics = 0;
 			return STATUS_SUCCESS;
 		}
+		if (FsInformationClass == FileFsVolumeInformation)
+		{
+			// TODO This code path is here to handle nodejs queries..
+			auto& listHandle = asListDirectoryHandle(FileHandle);
+			u32 entryOffset = listHandle.dir.tableOffset | 0x80000000;
+
+			auto& info = *(FILE_FS_VOLUME_INFORMATION*)FsInformation;
+			DirectoryTable::EntryInformation entryInfo;
+			g_directoryTable.GetEntryInformation(entryInfo, entryOffset);
+			UBA_ASSERT(entryInfo.attributes != 0);
+			info.VolumeCreationTime.QuadPart = 0;
+			info.VolumeSerialNumber = entryInfo.volumeSerial;
+			info.VolumeLabelLength = 0;
+			info.SupportsObjects = false;
+			info.VolumeLabel[0] = 0;
+			return STATUS_SUCCESS;
+		}
 		UBA_ASSERTF(false, L"NtQueryVolumeInformationFile called in ListDirectoryHandle using class %u which is not implemented (%ls)", FsInformationClass, HandleToName(FileHandle));
 	}
 	auto res = True_NtQueryVolumeInformationFile(TrueHandle, IoStatusBlock, FsInformation, Length, FsInformationClass);
@@ -189,6 +206,21 @@ NTSTATUS Detoured_NtQueryInformationFile(HANDLE FileHandle, PIO_STATUS_BLOCK IoS
 
 		if (TrueHandle == INVALID_HANDLE_VALUE)
 		{
+			if (FileInformationClass == FileBasicInformation)
+			{
+				auto& info = *(FILE_BASIC_INFORMATION*)FileInformation;
+				UBA_ASSERT(dh.dirTableOffset != ~0u);
+				DirectoryTable::EntryInformation entryInfo;
+				g_directoryTable.GetEntryInformation(entryInfo, dh.dirTableOffset);
+				UBA_ASSERT(entryInfo.attributes != 0);
+				info.CreationTime.QuadPart = entryInfo.lastWrite;
+				info.LastAccessTime.QuadPart = entryInfo.lastWrite;
+				info.LastWriteTime.QuadPart = entryInfo.lastWrite;
+				info.ChangeTime.QuadPart = entryInfo.lastWrite;
+				info.FileAttributes = entryInfo.attributes;
+				DEBUG_LOG_DETOURED(L"NtQueryInformationFile", L"(FileBasicInformation) %llu (%s) -> %ls", uintptr_t(FileHandle), HandleToName(FileHandle), ToString(STATUS_SUCCESS).data);
+				return STATUS_SUCCESS;
+			}
 			if (FileInformationClass == FileAllInformation)
 			{
 				UBA_ASSERT(dh.dirTableOffset != ~0u);
@@ -234,22 +266,41 @@ NTSTATUS Detoured_NtQueryInformationFile(HANDLE FileHandle, PIO_STATUS_BLOCK IoS
 				DEBUG_LOG_DETOURED(L"NtQueryInformationFile", L"(FileNetworkOpenInformation) %llu (%s) -> %ls", uintptr_t(FileHandle), HandleToName(FileHandle), ToString(STATUS_SUCCESS).data);
 				return STATUS_SUCCESS;
 			}
-
-			//if (FileInformationClass == 9) // FileNameInformation
-			//{
-			//	const wchar_t* name = dh.fileObject->fileInfo->originalName;
-			//	auto& info = *(FILE_NAME_INFORMATION*)FileInformation;
-			//	info.FileName[0] = '\\';
-			//	info.FileNameLength = 2;
-			//	//return STATUS_SUCCESS;
-			//	u32 nameLen = u32(wcslen(name));
-			//	//UBA_ASSERT(info.FileNameLength/2 > nameLen);
-			//	//memcpy(info.FileName, name, nameLen*2+2);
-			//	//info.FileNameLength = nameLen;
-			//	memcpy(info.FileName, L"\\\\", 4);
-			//	info.FileNameLength = 2;
-			//	return STATUS_SUCCESS;
-			//}
+			if (FileInformationClass == FileNormalizedNameInformation)
+			{
+				const tchar* src = TC("");//dh.fileObject->fileInfo->originalName;
+				if (Length < sizeof(ULONG))
+					return STATUS_INFO_LENGTH_MISMATCH;
+				auto& info = *(FILE_NAME_INFORMATION*)FileInformation;
+				u64 destLen = Length - FIELD_OFFSET(FILE_NAME_INFORMATION, FileName);
+				u64 srcLen = TStrlen(src)*sizeof(tchar);
+				u64 toWrite = Min(srcLen, destLen);
+				info.FileNameLength = (ULONG)srcLen;
+				memcpy(info.FileName, src, toWrite);
+				NTSTATUS res = destLen >= srcLen ? STATUS_SUCCESS : STATUS_BUFFER_OVERFLOW;
+				IoStatusBlock->Information = FIELD_OFFSET(FILE_NAME_INFORMATION, FileName) + toWrite;
+				IoStatusBlock->Status = res;
+				DEBUG_LOG_DETOURED(L"NtQueryInformationFile", L"(FileNormalizedNameInformation) %llu (%s) -> %ls", uintptr_t(FileHandle), HandleToName(FileHandle), ToString(res).data);
+				return res;
+			}
+			if (FileInformationClass == FileNameInformation)
+			{
+				const wchar_t* fileName = dh.fileObject->fileInfo->originalName;
+				auto& info = *(FILE_NAME_INFORMATION*)FileInformation;
+				ULONG header = FIELD_OFFSET(FILE_NAME_INFORMATION, FileName);
+				ULONG nameBytes = (ULONG)(wcslen(fileName) * sizeof(WCHAR));
+				ULONG required = header + nameBytes;
+				if (Length < header)
+					return STATUS_INFO_LENGTH_MISMATCH;
+				ULONG copy = Min(nameBytes, Length - header);
+				info.FileNameLength = nameBytes;
+				if (copy)
+					memcpy(info.FileName, fileName, copy);
+				IoStatusBlock->Status = STATUS_SUCCESS;
+				IoStatusBlock->Information = Min(Length, required);
+				DEBUG_LOG_DETOURED(L"NtQueryInformationFile", L"(FileNameInformation) %llu (%s) -> %ls", uintptr_t(FileHandle), HandleToName(FileHandle), ToString(STATUS_SUCCESS).data);
+				return STATUS_SUCCESS;
+			}
 			UBA_ASSERTF(false, L"NtQueryInformationFile (%u) failed using detoured handle %ls (%ls)", FileInformationClass, dh.fileObject->fileInfo->name, dh.fileObject->fileInfo->originalName);
 		}
 	}
@@ -281,6 +332,31 @@ NTSTATUS NTAPI Detoured_NtQueryDirectoryFile(HANDLE FileHandle, HANDLE Event, PI
 		u8* it = (u8*)FileInformation;
 		u8* bufferEnd = it + Length;
 
+		u32 structSize = 0;
+		u32 fileNameOffset = 0;
+		if (FileInformationClass == FileDirectoryInformation)
+		{
+			structSize = sizeof(FILE_DIRECTORY_INFORMATION);
+			fileNameOffset = FIELD_OFFSET(FILE_DIRECTORY_INFORMATION, FileName);
+		}
+		else if (FileInformationClass == FileIdBothDirectoryInformation)
+		{
+			structSize = sizeof(FILE_ID_BOTH_DIR_INFORMATION);
+			fileNameOffset = FIELD_OFFSET(FILE_ID_BOTH_DIR_INFORMATION, FileName);
+		}
+		else
+		{
+			UBA_ASSERT(false);
+			IoStatusBlock->Status = STATUS_OBJECT_NAME_NOT_FOUND;
+			return STATUS_OBJECT_NAME_NOT_FOUND;
+		}
+
+		if (Length < structSize)
+			return STATUS_BUFFER_OVERFLOW;
+
+		// Fix for buggy program (rustc.exe). even if there are zero entries they still seem to check this variable
+		memset(FileInformation, 0, structSize);
+
 		while (true)
 		{
 			if (listHandle.it == listHandle.fileTableOffsets.size())
@@ -299,25 +375,7 @@ NTSTATUS NTAPI Detoured_NtQueryDirectoryFile(HANDLE FileHandle, HANDLE Event, PI
 
 			u32 fileNameBytes = u32(wcslen(fileName) * 2);
 
-			wchar_t* fileNamePos = nullptr;
-			u32 structSize = 0;
-			if (FileInformationClass == FileDirectoryInformation)
-			{
-				structSize = sizeof(FILE_DIRECTORY_INFORMATION);
-				fileNamePos = ((FILE_DIRECTORY_INFORMATION*)it)->FileName;
-			}
-			else if (FileInformationClass == 2)//FileFullDirectoryInformation)
-			{
-				structSize = sizeof(FILE_FULL_DIR_INFORMATION);
-				fileNamePos = ((FILE_FULL_DIR_INFORMATION*)it)->FileName;
-			}
-			else
-			{
-				UBA_ASSERT(false);
-				IoStatusBlock->Status = STATUS_OBJECT_NAME_NOT_FOUND;
-				return STATUS_OBJECT_NAME_NOT_FOUND;
-			}
-
+			wchar_t* fileNamePos = (wchar_t*)((u8*)it + fileNameOffset);
 			u8* writeEnd = (u8*)fileNamePos + fileNameBytes;
 			if (writeEnd > bufferEnd)
 			{
@@ -1837,20 +1895,16 @@ NTSTATUS Detoured_NtQueryObject(HANDLE Handle, OBJECT_INFORMATION_CLASS ObjectIn
 				g_directoryTable.GetFinalPath(buffer, fixedPath.data);
 				VirtualizePath(buffer);
 
-				if (g_runningRemote || buffer[0] != fixedPath[0])
-				{
-					// It was remote or virtualized.. let's make up device drive.
-					// TODO: Maybe we need to replicate device drive names to remotes
-					buffer.Prepend(AsView(L"\\Device\\HarddiskVolume100"), 2);
-				}
-				else
-				{
-					wchar_t drive[3] = { buffer[0], ':', 0 };
-					wchar_t device[256];
-					DWORD deviceLen = QueryDosDeviceW(drive, device, sizeof_array(device));
-					UBA_ASSERT(deviceLen);
-					buffer.Prepend(StringView(device, deviceLen), 2);
-				}
+				// Remote machines will just use the native exe dir drive information
+				// Since there are no better options. This will make it so we don't have to implement detouring for all kinds of queries about volume information etc.
+				if (g_runningRemote)
+					buffer[0] = g_exeDir[0];
+
+				wchar_t drive[3] = { buffer[0], ':', 0 };
+				wchar_t device[256];
+				DWORD deviceLen = QueryDosDeviceW(drive, device, sizeof_array(device));
+				UBA_ASSERT(deviceLen);
+				buffer.Prepend(StringView(device, deviceLen), 2);
 
 				u64 bufferSize = (buffer.count+1)*sizeof(tchar);
 				u64 totalSize = sizeof(UNICODE_STRING) + bufferSize;
