[UBA]
* Removed detouring of fputs. I think this was implemented back when host process also spawn detoured child processes of detoured processes..

diff --git a/Engine/Source/Programs/UnrealBuildAccelerator/Detours/Private/Windows/UbaDetoursFunctionsUcrtBase.inl b/Engine/Source/Programs/UnrealBuildAccelerator/Detours/Private/Windows/UbaDetoursFunctionsUcrtBase.inl
--- a/Engine/Source/Programs/UnrealBuildAccelerator/Detours/Private/Windows/UbaDetoursFunctionsUcrtBase.inl
+++ b/Engine/Source/Programs/UnrealBuildAccelerator/Detours/Private/Windows/UbaDetoursFunctionsUcrtBase.inl
@@ -260,24 +260,6 @@ int Detoured__write(int fd, const void* buffer, unsigned int count)
 	return res;
 }
 
-int Detoured_fputs(const char* str, FILE* stream)
-{
-	// This code is not working properly on wine helpers.
-	// When using iwyu on remote wine helper we get a fno that matches stderr but it is not a stderr.. it is supposed to write to a file
-	// Unfortunately I can't remember the initial reason this code was added so I can't change it and know if I broke something or not
-	if (!g_runningRemote || !g_isRunningWine)
-	{
-		int fno = _fileno(stream); (void)fno;
-		int errfno = _fileno(stderr);
-		if (fno == errfno || fno == _fileno(stdout))
-		{
-			Shared_WriteConsole(str, u32(strlen(str)), fno == errfno);
-			return 1;
-		}
-	}
-	return True_fputs(str, stream);
-}
-
 int Detoured__wspawnl(int mode, const wchar_t* cmdname, const wchar_t* arg0, const wchar_t* arg1, ...)
 {
 	DETOURED_CALL(_wspawnl);
@@ -424,4 +406,27 @@ char* Detoured__getcwd(char* buffer, int maxlen)
 	return res;
 }
 
+int Detoured_fputs(const char* str, FILE* stream)
+{
+	// This code is not working properly on wine helpers.
+	// When using iwyu on remote wine helper we get a fno that matches stderr but it is not a stderr.. it is supposed to write to a file
+	// Unfortunately I can't remember the initial reason this code was added so I can't change it and know if I broke something or not
+	// - Decided to remove this function because I think it was written before the refactoring where detoured processes spawned their own children
+	// and pipes became a huge mess
+	/*
+	if (!g_runningRemote || !g_isRunningWine)
+	{
+	int fno = _fileno(stream); (void)fno;
+	int errfno = _fileno(stderr);
+	if (fno == errfno || fno == _fileno(stdout))
+	{
+	//Shared_WriteConsole(str, u32(strlen(str)), fno == errfno);
+	//return 1;
+	}
+	}
+	*/
+	return True_fputs(str, stream);
+}
+
+
 #endif // defined(DETOURED_INCLUDE_DEBUG)
