[UBA]
* Fixed so error log is suppressed when detoured child process is killed through job handle

diff --git a/Engine/Source/Programs/UnrealBuildAccelerator/Common/Private/UbaProcess.cpp b/Engine/Source/Programs/UnrealBuildAccelerator/Common/Private/UbaProcess.cpp
--- a/Engine/Source/Programs/UnrealBuildAccelerator/Common/Private/UbaProcess.cpp
+++ b/Engine/Source/Programs/UnrealBuildAccelerator/Common/Private/UbaProcess.cpp
@@ -530,6 +530,9 @@ namespace uba
 				return true;
 			}
 
+			if (m_parentReportedExit)
+				return false;
+
 			StringBuffer<> err;
 			err.Appendf(TC("ERROR: Process %u (%s) not active but did not get exit message. Received %u messages. Signal code: %i. Exit value or signal: %i. Execution time: %s."), m_nativeProcessId, m_startInfo.GetDescription(), m_messageCount, signalInfo.si_code, signalInfo.si_status, TimeToText(GetTime() - m_startTime).str);
 			LogLine(false, err.data, LogEntryType_Error);
