[UBA]
* Improved some test logging
* Added code to loop tests

diff --git a/Engine/Source/Programs/UnrealBuildAccelerator/Test/Private/UbaTestAll.h b/Engine/Source/Programs/UnrealBuildAccelerator/Test/Private/UbaTestAll.h
--- a/Engine/Source/Programs/UnrealBuildAccelerator/Test/Private/UbaTestAll.h
+++ b/Engine/Source/Programs/UnrealBuildAccelerator/Test/Private/UbaTestAll.h
@@ -170,9 +170,13 @@ namespace uba
 		if (argc > 1)
 			filter = argv[1];
 
+		u32 loopCount = 1;
 		//UBA_TEST(TestStress) // This can not be submitted.. it depends on CoordinatorHorde and credentials
 		//UBA_TEST(TestStdOutRemote) // This can not be submitted.. depends on a running UbaAgent
-		UBA_TESTS
+		for (u32 i=0; i!=loopCount; ++i)
+		{
+			UBA_TESTS
+		}
 
 		logger.Info(TC("Tests finished successfully!"));
 
diff --git a/Engine/Source/Programs/UnrealBuildAccelerator/Test/Private/UbaTestSession.cpp b/Engine/Source/Programs/UnrealBuildAccelerator/Test/Private/UbaTestSession.cpp
--- a/Engine/Source/Programs/UnrealBuildAccelerator/Test/Private/UbaTestSession.cpp
+++ b/Engine/Source/Programs/UnrealBuildAccelerator/Test/Private/UbaTestSession.cpp
@@ -178,7 +178,7 @@ namespace uba
 		if (serverShouldListen)
 		{
 			if (!server.StartListen(tcpBackend, port))
-				return false;
+				return logger.Error(TC("Failed to start listen"));
 			if (!client.Connect(tcpBackend, TC("127.0.0.1"), port))
 				return logger.Error(TC("Failed to connect"));
 			if (!client.Connect(tcpBackend, TC("127.0.0.1"), port))
@@ -190,14 +190,14 @@ namespace uba
 				return logger.Error(TC("Failed to listen"));
 			if (!server.AddClient(tcpBackend, TC("127.0.0.1"), port))
 				return logger.Error(TC("Failed to connect"));
-			while (server.HasConnectInProgress())
+			while (!client.IsConnected())
 				Sleep(1);
 		}
 		auto disconnectGuard = MakeGuard([&]() { tcpBackend.StopListen(); client.Disconnect(); server.RemoveDisconnectedConnections(); });
 
 		Config config;
 		if (!client.FetchConfig(config))
-			return false;
+			return logger.Error(TC("Failed to fetch config"));
 
 		StringBuffer<MaxPath> rootDir;
 		rootDir.Append(testRootDir).Append(TCV("UbaClient")).AppendValue(port);
@@ -206,7 +206,7 @@ namespace uba
 		if (!deleteAll)
 			toDelete.Append(PathSeparator).Append(TCV("sessions"));
 		if (!DeleteAllFiles(logger, toDelete.data))
-			return false;
+			return logger.Error(TC("Failed to delete all files"));
 
 		StorageClientCreateInfo storageClientInfo(client, rootDir.data);
 		storageClientInfo.Apply(config);
diff --git a/Engine/Source/Programs/UnrealBuildAccelerator/TestApp/Private/Windows/UbaTestAppWin.cpp b/Engine/Source/Programs/UnrealBuildAccelerator/TestApp/Private/Windows/UbaTestAppWin.cpp
--- a/Engine/Source/Programs/UnrealBuildAccelerator/TestApp/Private/Windows/UbaTestAppWin.cpp
+++ b/Engine/Source/Programs/UnrealBuildAccelerator/TestApp/Private/Windows/UbaTestAppWin.cpp
@@ -403,18 +403,20 @@ int wmain(int argc, wchar_t* argv[])
 
 		wchar_t buf[32];
 		int read = ReadTestFile(buf, 32*2, L"SpecialFile1");
-		if (read == -1)
-			return -1;
+		if (read != 2)
+			return LogError(L"Expected to read two bytes but got %i", read);
 		buf[read*2] = 0;
 
 		int readIndex = _wtoi(buf);
 		if (callIndex != readIndex)
-			return LogError(L"callIndex different from readIndex");
+			return LogError(L"ReadTestFile callIndex (%i) different from readIndex (%i)", callIndex, readIndex);
 		return 0;
 	}
 	else if (wcsncmp(argv[1], L"-write=", 7) == 0)
 	{
 		int callIndex = _wtoi(argv[1] + 7);
+		if (callIndex > 9)
+			return LogError(L"CallIndex too large (%i)", callIndex);
 		wchar_t buf[32];
 		buf[0] = wchar_t('1' + callIndex);
 		return WriteTestFile(buf, 2, L"SpecialFile1");
@@ -422,16 +424,18 @@ int wmain(int argc, wchar_t* argv[])
 	else if (wcsncmp(argv[1], L"-readwrite=", 11) == 0)
 	{
 		int callIndex = _wtoi(argv[1] + 11);
+		if (callIndex > 9)
+			return LogError(L"CallIndex too large (%i)", callIndex);
 
 		wchar_t buf[32];
 		int read = ReadTestFile(buf, 32*2, L"SpecialFile1");
-		if (read == -1)
-			return -1;
+		if (read != 2)
+			return LogError(L"Expected to read two bytes but got %i", read);
 		buf[read*2] = 0;
 
 		int readIndex = _wtoi(buf);
 		if (callIndex != readIndex)
-			return LogError(L"callIndex different from readIndex");
+			return LogError(L"ReadWriteTestFile callIndex (%i) different from readIndex (%i)", callIndex, readIndex);
 
 		buf[0] = wchar_t('1' + readIndex);
 		return WriteTestFile(buf, 2, L"SpecialFile1");
