[UBA]
* Added error if test code fails to delete test file

diff --git a/Engine/Source/Programs/UnrealBuildAccelerator/Test/Private/UbaTest.cpp b/Engine/Source/Programs/UnrealBuildAccelerator/Test/Private/UbaTest.cpp
--- a/Engine/Source/Programs/UnrealBuildAccelerator/Test/Private/UbaTest.cpp
+++ b/Engine/Source/Programs/UnrealBuildAccelerator/Test/Private/UbaTest.cpp
@@ -34,7 +34,9 @@ namespace uba
 	{
 		StringBuffer<> testFileName;
 		testFileName.Append(testRootDir).EnsureEndsWithSlash().Append(fileName).FixPathSeparators();
-		return DeleteFileW(testFileName.data);
+		if (!DeleteFileW(testFileName.data))
+			return logger.Error(TC("Failed to delete test file %s"), fileName.data);
+		return true;
 	}
 
 	bool FileExists(Logger& logger, StringView dir, StringView fileName)
