[UBA]
* Added assert description for when child process id is out of range

diff --git a/Engine/Source/Programs/UnrealBuildAccelerator/Common/Private/UbaProcess.cpp b/Engine/Source/Programs/UnrealBuildAccelerator/Common/Private/UbaProcess.cpp
--- a/Engine/Source/Programs/UnrealBuildAccelerator/Common/Private/UbaProcess.cpp
+++ b/Engine/Source/Programs/UnrealBuildAccelerator/Common/Private/UbaProcess.cpp
@@ -1413,7 +1413,7 @@ namespace uba
 	bool ProcessImpl::HandleStartProcess(BinaryReader& reader, BinaryWriter& writer)
 	{
 		u32 processId = reader.ReadU32();
-		UBA_ASSERT(processId > 0 && processId <= m_childProcesses.size());
+		UBA_ASSERTF(processId > 0 && processId <= m_childProcesses.size(), TC("Got process id %u, expected number lower than %u"), processId, u32(m_childProcesses.size()));
 		auto& process = *(ProcessImpl*)m_childProcesses[processId - 1].m_process;
 		bool result = reader.ReadBool();
 		u32 lastError = reader.ReadU32();
