[UBA]
* Added detouring of CoCreateInstance to fix clang fallback search for vs toolchain (hard coded path right now which is not great)
* Added hard coded response for GetVolumeInformationByHandleW lpFileSystemFlags when no true device handle
* Fixed bug in SetFileAttributes when remote and path is in temp dir
* Moved SetFilePointer and SetFilePointerEx to only exist in debug and implemented NtQueryInformationFile and NtSetInformationFile instead
* Fixed so CreateProcess update custom environment if it contains TEMP/TMP dirs since they always need to be properly set when running remotely

diff --git a/Engine/Source/Programs/UnrealBuildAccelerator/Core/Public/UbaWinTypes.h b/Engine/Source/Programs/UnrealBuildAccelerator/Core/Public/UbaWinTypes.h
--- a/Engine/Source/Programs/UnrealBuildAccelerator/Core/Public/UbaWinTypes.h
+++ b/Engine/Source/Programs/UnrealBuildAccelerator/Core/Public/UbaWinTypes.h
@@ -59,6 +59,20 @@ struct FILE_BASIC_INFORMATION {
 	ULONG FileAttributes;
 };
 
+struct FILE_FS_DEVICE_INFORMATION
+{
+	DEVICE_TYPE DeviceType;
+	ULONG Characteristics;
+};
+
+struct FILE_FS_ATTRIBUTE_INFORMATION
+{
+	ULONG FileSystemAttributes;
+	LONG  MaximumComponentNameLength;
+	ULONG FileSystemNameLength;
+	WCHAR FileSystemName[1];
+};
+
 struct FILE_NETWORK_OPEN_INFORMATION
 {
 	LARGE_INTEGER CreationTime;
@@ -70,11 +84,33 @@ struct FILE_NETWORK_OPEN_INFORMATION
 	ULONG FileAttributes;
 };
 
+struct FILE_POSITION_INFORMATION
+{
+	LARGE_INTEGER CurrentByteOffset;
+};
+
+struct FILE_END_OF_FILE_INFORMATION
+{
+	LARGE_INTEGER EndOfFile;
+};
+
+struct FILE_STANDARD_INFORMATION {
+	LARGE_INTEGER AllocationSize;
+	LARGE_INTEGER EndOfFile;
+	ULONG         NumberOfLinks;
+	BOOLEAN       DeletePending;
+	BOOLEAN       Directory;
+};
+
+
 constexpr FILE_INFORMATION_CLASS FileBasicInformation = (FILE_INFORMATION_CLASS)4;
+constexpr FILE_INFORMATION_CLASS FileStandardInformation = (FILE_INFORMATION_CLASS)5;
 constexpr FILE_INFORMATION_CLASS FileNameInformation = (FILE_INFORMATION_CLASS)9;
 constexpr FILE_INFORMATION_CLASS FileRenameInformation = (FILE_INFORMATION_CLASS)10;
 constexpr FILE_INFORMATION_CLASS FileDispositionInformation = (FILE_INFORMATION_CLASS)13;
+constexpr FILE_INFORMATION_CLASS FilePositionInformation = (FILE_INFORMATION_CLASS)14;
 constexpr FILE_INFORMATION_CLASS FileAllInformation = (FILE_INFORMATION_CLASS)18;
+constexpr FILE_INFORMATION_CLASS FileEndOfFileInformation = (FILE_INFORMATION_CLASS)20;
 constexpr FILE_INFORMATION_CLASS FileNetworkOpenInformation = (FILE_INFORMATION_CLASS)34;
 constexpr FILE_INFORMATION_CLASS FileNormalizedNameInformation = (FILE_INFORMATION_CLASS)35;
 constexpr FILE_INFORMATION_CLASS FileIdBothDirectoryInformation = (FILE_INFORMATION_CLASS)37;
diff --git a/Engine/Source/Programs/UnrealBuildAccelerator/Detours/Private/Windows/UbaDetoursFunctionsCombase.inl b/Engine/Source/Programs/UnrealBuildAccelerator/Detours/Private/Windows/UbaDetoursFunctionsCombase.inl
new file mode 100644
--- /dev/null
+++ b/Engine/Source/Programs/UnrealBuildAccelerator/Detours/Private/Windows/UbaDetoursFunctionsCombase.inl
@@ -0,0 +1,35 @@
+// Copyright Epic Games, Inc. All Rights Reserved.
+
+HRESULT Detoured_CoCreateInstance(REFCLSID rclsid, LPUNKNOWN pUnkOuter, DWORD dwClsContext, REFIID riid, LPVOID *ppv)
+{
+	DEBUG_LOG_DETOURED(L"CoCreateInstance", L"%s", GuidToString((Guid&)rclsid).data);
+
+	// Special handling for com object used to figure out visual studio installation. This is used as a fallback path for lld-link.exe among others
+	// {177F0C4A-1CD3-4DE7-A32C-71DBBB9FA36D}
+	static const CLSID CLSID_SetupConfiguration = { 0x177F0C4A, 0x1CD3, 0x4DE7, { 0xA3, 0x2C, 0x71, 0xDB, 0xBB, 0x9F, 0xA3, 0x6D } };
+
+	if (IsEqualCLSID(rclsid, CLSID_SetupConfiguration))
+	{
+		// TODO: Should probably not be hard coded and instead be transferred from host
+		const wchar_t* dllPath = L"c:\\ProgramData\\Microsoft\\VisualStudio\\Setup\\x64\\Microsoft.VisualStudio.Setup.Configuration.Native.dll";
+		HMODULE hMod = LoadLibraryW(dllPath);
+		if (!hMod)
+			return REGDB_E_CLASSNOTREG;
+
+		typedef HRESULT (WINAPI* FnDllGetClassObject)(REFCLSID, REFIID, LPVOID*);
+		auto DllGetClassObject = (FnDllGetClassObject)GetProcAddress(hMod, "DllGetClassObject");
+		if (!DllGetClassObject)
+			return REGDB_E_CLASSNOTREG;
+
+		IClassFactory* pCF = nullptr;
+		HRESULT hr = DllGetClassObject(rclsid, IID_IClassFactory, (void**)&pCF);
+		if (FAILED(hr)) return hr;
+
+		hr = pCF->CreateInstance(pUnkOuter, riid, ppv);
+		pCF->Release();
+		return hr;
+	}
+
+
+	return True_CoCreateInstance(rclsid, pUnkOuter, dwClsContext, riid, ppv);
+}
\ No newline at end of file
diff --git a/Engine/Source/Programs/UnrealBuildAccelerator/Detours/Private/Windows/UbaDetoursFunctionsKernelBase.inl b/Engine/Source/Programs/UnrealBuildAccelerator/Detours/Private/Windows/UbaDetoursFunctionsKernelBase.inl
--- a/Engine/Source/Programs/UnrealBuildAccelerator/Detours/Private/Windows/UbaDetoursFunctionsKernelBase.inl
+++ b/Engine/Source/Programs/UnrealBuildAccelerator/Detours/Private/Windows/UbaDetoursFunctionsKernelBase.inl
@@ -414,11 +414,9 @@ BOOL Detoured_WriteConsoleW(HANDLE hConsoleOutput, const VOID* lpBuffer, DWORD n
 BOOL Detoured_ReadConsoleW(HANDLE hConsoleInput, LPVOID lpBuffer, DWORD nNumberOfCharsToRead, LPDWORD lpNumberOfCharsRead, PCONSOLE_READCONSOLE_CONTROL pInputControl)
 {
 	DETOURED_CALL(ReadConsoleW);
-	#if UBA_DEBUG
-	Rpc_WriteLogf(L"WARNING Got call to ReadConsoleW.. this is not handled by Uba yet"); // seems like wine's cmd.exe likes calling this
-	#endif
-	return 0;
-	//return True_ReadConsoleW(hConsoleInput, lpBuffer, nNumberOfCharsToRead, lpNumberOfCharsRead, pInputControl);
+	//DEBUG_LOG_DETOURED(L"ReadConsoleW", L"(%llu)", u64(hConsoleInput)); // Too much spam
+	UBA_ASSERT(!isDetouredHandle(hConsoleInput));
+	return True_ReadConsoleW(hConsoleInput, lpBuffer, nNumberOfCharsToRead, lpNumberOfCharsRead, pInputControl);
 }
 
 UINT Detoured_GetDriveTypeW(LPCWSTR lpRootPathName)
@@ -494,7 +492,27 @@ BOOL Detoured_GetVolumeInformationByHandleW(HANDLE hFile, LPWSTR lpVolumeNameBuf
 	{
 		UBA_ASSERT(!lpVolumeNameBuffer);
 		UBA_ASSERT(!lpMaximumComponentLength);
-		UBA_ASSERT(!lpFileSystemFlags);
+		if (lpFileSystemFlags) // esbuild.exe query this
+		{
+			// Standard nfts flags. It might be that we want to comment some of these outs because UBA (especially remotely) don't want to support these
+			*lpFileSystemFlags = 0
+			| FILE_CASE_SENSITIVE_SEARCH
+			| FILE_CASE_PRESERVED_NAMES
+			| FILE_UNICODE_ON_DISK
+			| FILE_PERSISTENT_ACLS
+			| FILE_FILE_COMPRESSION
+			| FILE_VOLUME_QUOTAS
+			| FILE_SUPPORTS_SPARSE_FILES
+			| FILE_SUPPORTS_REPARSE_POINTS
+			| FILE_SUPPORTS_OBJECT_IDS
+			| FILE_SUPPORTS_ENCRYPTION
+			| FILE_NAMED_STREAMS
+			| FILE_SUPPORTS_HARD_LINKS
+			| FILE_SUPPORTS_EXTENDED_ATTRIBUTES
+			| FILE_SUPPORTS_OPEN_BY_FILE_ID
+			| FILE_SUPPORTS_USN_JOURNAL
+			;
+		}
 		DirectoryTable::EntryInformation entryInfo;
 		g_directoryTable.GetEntryInformation(entryInfo, entryOffset);
 		if (lpVolumeSerialNumber)
@@ -917,99 +935,6 @@ DWORD Detoured_GetFileSizeEx(HANDLE hFile, PLARGE_INTEGER lpFileSize)
 	return res;
 }
 
-DWORD Detoured_SetFilePointer(HANDLE hFile, LONG lDistanceToMove, PLONG lpDistanceToMoveHigh, DWORD dwMoveMethod)
-{
-	DETOURED_CALL(SetFilePointer);
-
-	HANDLE trueHandle = hFile;
-	if (isDetouredHandle(hFile))
-	{
-		auto& dh = asDetouredHandle(hFile);
-		FileInfo& fi = *dh.fileObject->fileInfo;
-		if (fi.memoryFile || fi.isFileMap)
-		{
-			LARGE_INTEGER liDistanceToMove;
-			liDistanceToMove.LowPart = 0;
-			liDistanceToMove.HighPart = lpDistanceToMoveHigh ? *lpDistanceToMoveHigh : 0;
-			liDistanceToMove.QuadPart += lDistanceToMove;
-
-			if (dwMoveMethod == FILE_BEGIN)
-				dh.pos = liDistanceToMove.QuadPart;
-			else if (dwMoveMethod == FILE_CURRENT)
-				dh.pos += liDistanceToMove.QuadPart;
-			else if (dwMoveMethod == FILE_END)
-			{
-				u64 size = 0;
-				if (fi.memoryFile)
-					size = fi.memoryFile->writtenSize;
-				else
-				{
-					//UBA_ASSERT(fi.fileMapMemEnd);
-					size = fi.size;// fi.fileMapMemSize;
-				}
-				dh.pos = Max(0ll, (LONGLONG)size + liDistanceToMove.QuadPart);
-			}
-			DEBUG_LOG_DETOURED(L"SetFilePointer", L"%llu %lli %u (%ls) -> %u", uintptr_t(hFile), liDistanceToMove.QuadPart, dwMoveMethod, HandleToName(hFile), DWORD(dh.pos));
-			SetLastError(ERROR_SUCCESS);
-			return DWORD(dh.pos);
-		}
-
-		trueHandle = dh.trueHandle;
-		UBA_ASSERTF(trueHandle != INVALID_HANDLE_VALUE, TC("SetFilePointer got detoured handle %llu that was not memory file or filemap (%ls)"), uintptr_t(hFile), HandleToName(hFile));
-	}
-
-	DEBUG_LOG_TRUE(L"SetFilePointer", L"%llu (%ls)", uintptr_t(hFile), HandleToName(hFile));
-	return True_SetFilePointer(trueHandle, lDistanceToMove, lpDistanceToMoveHigh, dwMoveMethod);
-}
-
-DWORD Detoured_SetFilePointerEx(HANDLE hFile, LARGE_INTEGER liDistanceToMove, PLARGE_INTEGER lpNewFilePointer, DWORD dwMoveMethod)
-{
-	DETOURED_CALL(SetFilePointerEx);
-	HANDLE trueHandle = hFile;
-	if (isDetouredHandle(hFile))
-	{
-		DetouredHandle& dh = asDetouredHandle(hFile);
-		if (dh.type >= HandleType_StdErr)
-		{
-			if (lpNewFilePointer)
-				*lpNewFilePointer = ToLargeInteger(0);
-			// TODO: What should we do with this?
-			SetLastError(ERROR_SUCCESS);
-			DEBUG_LOG_DETOURED(L"SetFilePointerEx", L"%llu %lli %u (%ls) -> Success", uintptr_t(hFile), liDistanceToMove.QuadPart, dwMoveMethod, HandleToName(hFile));
-			return TRUE;
-		}
-
-		FileInfo& fi = *dh.fileObject->fileInfo;
-		if (fi.memoryFile || fi.isFileMap)
-		{
-			if (dwMoveMethod == FILE_BEGIN)
-				dh.pos = liDistanceToMove.QuadPart;
-			else if (dwMoveMethod == FILE_CURRENT)
-				dh.pos += liDistanceToMove.QuadPart;
-			else if (dwMoveMethod == FILE_END)
-			{
-				u64 size = 0;
-				if (fi.memoryFile)
-					size = fi.memoryFile->writtenSize;
-				else if (fi.fileMapMem)
-					size = fi.fileMapMemSize;
-				else
-					size = fi.size;
-				dh.pos = Max(0ll, (LONGLONG)size + liDistanceToMove.QuadPart);
-			}
-			if (lpNewFilePointer)
-				*lpNewFilePointer = ToLargeInteger(dh.pos);
-			SetLastError(ERROR_SUCCESS);
-			DEBUG_LOG_DETOURED(L"SetFilePointerEx", L"%llu %lli %u (%ls) -> Success", uintptr_t(hFile), liDistanceToMove.QuadPart, dwMoveMethod, HandleToName(hFile));
-			return TRUE;
-		}
-		trueHandle = dh.trueHandle;
-		UBA_ASSERTF(trueHandle != INVALID_HANDLE_VALUE, TC("SetFilePointer got detoured handle %llu that was not memory file or filemap (%ls)"), uintptr_t(hFile), HandleToName(hFile));
-	}
-	DEBUG_LOG_TRUE(L"SetFilePointerEx", L"%llu %lli %u (%ls)", uintptr_t(hFile), liDistanceToMove.QuadPart, dwMoveMethod, HandleToName(hFile));
-	return True_SetFilePointerEx(trueHandle, liDistanceToMove, lpNewFilePointer, dwMoveMethod); // This ends up in NtSetInformationFile
-}
-
 BOOL Detoured_SetEndOfFile(HANDLE hFile)
 {
 	DETOURED_CALL(SetEndOfFile);
@@ -1206,7 +1131,23 @@ BOOL Detoured_SetFileAttributesW(LPCWSTR lpFileName, DWORD dwFileAttributes)
 		return true;
 	}
 
-	DEBUG_LOG_TRUE(L"SetFileAttributesW", L"(%ls) %u", lpFileName, dwFileAttributes);
+	if (g_runningRemote && !fixedFile.StartsWith(g_systemTemp))
+	{
+		StringKey fileNameKey = ToStringKeyLower(fixedFile);
+		u32 dirTableOffset = Rpc_GetEntryOffset(fileNameKey, fixedFile, true);
+		if (dirTableOffset != ~u32(0))
+		{
+			DirectoryTable::EntryInformation info;
+			g_directoryTable.GetEntryInformation(info, dirTableOffset);
+			UBA_ASSERTF(info.attributes == dwFileAttributes, TC("Not implemented - Remote process wants to SetFileAttributes on file %s"), fixedFile.data);
+			DEBUG_LOG_TRUE(L"SetFileAttributesW", L"(IGNORED_REMOTE) (%ls) %u", lpFileName, dwFileAttributes);
+			return true;
+		}
+		UBA_ASSERTF(false, L"SetFileAttributesW from remote is not implemented");
+	}
+
+	DEBUG_LOG_TRUE(L"SetFileAttributesW", L"(SUPPRESS)(%ls) %u", lpFileName, dwFileAttributes);
+	SuppressDetourScope _;
 	TimerScope ts(g_kernelStats.setFileInfo);
 	return True_SetFileAttributesW(lpFileName, dwFileAttributes);
 }
@@ -2173,6 +2114,11 @@ BOOL Detoured_GetFileInformationByHandleEx(HANDLE hFile, FILE_INFO_BY_HANDLE_CLA
 				UBA_ASSERTF(!mf, L"GetFileInformationByHandleEx called for memory file using class %u which is not implemented (%s)", fileInformationClass, HandleToName(hFile));
 			}
 		}
+		else if (fileInformationClass == FileIdBothDirectoryRestartInfo || fileInformationClass == FileIdBothDirectoryInfo)
+		{
+			SetLastError(ERROR_ACCESS_DENIED);
+			return FALSE;
+		}
 
 #if UBA_DEBUG_VALIDATE
 		originalName = dh.fileObject->fileInfo->originalName;
@@ -2180,6 +2126,21 @@ BOOL Detoured_GetFileInformationByHandleEx(HANDLE hFile, FILE_INFO_BY_HANDLE_CLA
 	}
 	else if (isListDirectoryHandle(hFile))
 	{
+		if (fileInformationClass == FileFullDirectoryRestartInfo  || fileInformationClass == FileFullDirectoryInfo)
+		{
+			DEBUG_LOG_DETOURED(L"GetFileInformationByHandleEx", L"(FileFullDirectoryInfo) %llu (%ls)", uintptr_t(hFile), HandleToName(hFile));
+			IO_STATUS_BLOCK io;
+			NTSTATUS status = Detoured_NtQueryDirectoryFile(hFile, NULL, NULL, NULL, &io, lpFileInformation, dwBufferSize, FileDirectoryInformation, TRUE, NULL, (fileInformationClass == FileFullDirectoryRestartInfo));
+			return SetLastErrorFromNtStatus(status);
+		}
+		if (fileInformationClass == FileIdBothDirectoryRestartInfo || fileInformationClass == FileIdBothDirectoryInfo)
+		{
+			DEBUG_LOG_DETOURED(L"GetFileInformationByHandleEx", L"(FileIdBothDirectoryInfo) %llu (%ls)", uintptr_t(hFile), HandleToName(hFile));
+			IO_STATUS_BLOCK io;
+			NTSTATUS status = Detoured_NtQueryDirectoryFile(hFile, NULL, NULL, NULL, &io, lpFileInformation, dwBufferSize, FileIdBothDirectoryInformation, TRUE, NULL, (fileInformationClass == FileIdBothDirectoryRestartInfo));
+			return SetLastErrorFromNtStatus(status);
+		}
+
 		auto& listHandle = asListDirectoryHandle(hFile);
 		if (listHandle.dir.tableOffset != InvalidTableOffset)
 			entryOffset = listHandle.dir.tableOffset | 0x80000000;
@@ -2248,17 +2209,8 @@ BOOL Detoured_GetFileInformationByHandleEx(HANDLE hFile, FILE_INFO_BY_HANDLE_CLA
 		else if (fileInformationClass == FileRemoteProtocolInfo)
 		{
 			SetLastError(ERROR_INVALID_PARAMETER);
+			DEBUG_LOG_DETOURED(L"GetFileInformationByHandleEx", L"(DIRTABLE) (FileRemoteProtocolInfo) %llu (%s) -> %s", uintptr_t(hFile), HandleToName(hFile), ToString(FALSE));
 			return FALSE;
-			/*
-			auto& data = *(FILE_REMOTE_PROTOCOL_INFO*)lpFileInformation;
-
-			SuppressCreateFileDetourScope s;
-			HANDLE h = CreateFile(HandleToName(hFile), FILE_READ_ATTRIBUTES|FILE_LIST_DIRECTORY|FILE_WRITE_EA, FILE_SHARE_READ, NULL, OPEN_EXISTING, FILE_FLAG_BACKUP_SEMANTICS, NULL);
-			BOOL res = GetFileInformationByHandleEx(h, FileRemoteProtocolInfo, lpFileInformation, dwBufferSize);
-			auto err = GetLastError();
-			CloseHandle(h);
-			return FALSE;
-			*/
 		}
 		else if (fileInformationClass == FileAttributeTagInfo)
 		{
@@ -3497,6 +3449,38 @@ BOOL Detoured_CreateProcessW(LPCWSTR lpApplicationName, LPWSTR lpCommandLine, LP
 		}
 	}
 
+	// Sometimes programs decides to change environment variables such as TEMP/TMP/SystemRoot etc which will break when running remotely
+	Vector<tchar> tempEnvironment;
+	if (lpEnvironment)
+	{
+		auto addChars = [&](const tchar* text, u64 len)
+			{
+				u64 destPos = tempEnvironment.size();
+				tempEnvironment.resize(tempEnvironment.size() + len);
+				tchar* dest = tempEnvironment.data() + destPos;
+				memcpy(dest, text, len*sizeof(tchar));
+			};
+		UBA_ASSERT(dwCreationFlags & CREATE_UNICODE_ENVIRONMENT);
+		for(StringView env = ToView((tchar*)lpEnvironment); env.count; env = ToView(env.data + env.count + 1))
+		{
+			u32 equalPos;
+			env.Contains('=', &equalPos);
+			StringView name = env.SubStr(0, equalPos);
+			if (name.Equals(TCV("TEMP")) || name.Equals(TCV("TMP")))
+			{
+				addChars(name.data, name.count+1);
+				addChars(g_systemTemp.data, g_systemTemp.count+1);
+			}
+			else
+			{
+				addChars(env.data, env.count+1);
+			}
+		}
+
+		tempEnvironment.push_back(0);
+		lpEnvironment = tempEnvironment.data();
+	}
+
 	// CREATE_NO_WINDOW is an optimization that prevents additional overhead per spawned child process.
 	// .. disabled for now because there are hierarchies of processes that does not like this. Haven't figured out exactly why
 	// .. can be reproed by running some of the very complex cmake/python scripts when compiling clang (lto.exports etc)
@@ -3582,7 +3566,7 @@ BOOL Detoured_CreateProcessW(LPCWSTR lpApplicationName, LPWSTR lpCommandLine, LP
 	detouredHandle->trueHandle = trueHandle;
 	lpProcessInformation->hProcess = makeDetouredHandle(detouredHandle);
 
-	DEBUG_LOG_DETOURED(L"CreateProcessW", L"%llu (0x%llx)", u64(lpProcessInformation->hProcess), u64(trueHandle));
+	DEBUG_LOG_DETOURED(L"CreateProcessW", L"(RESULT) %s %llu (0x%llx)", application.data, u64(lpProcessInformation->hProcess), u64(trueHandle));
 	return TRUE;
 }
 
@@ -3960,6 +3944,20 @@ BOOL Detoured_GetVolumePathNameA(LPCSTR lpszFileName, LPSTR lpszVolumePathName,
 	return True_GetVolumePathNameA(lpszFileName, lpszVolumePathName, cchBufferLength);
 }
 
+DWORD Detoured_SetFilePointer(HANDLE hFile, LONG lDistanceToMove, PLONG lpDistanceToMoveHigh, DWORD dwMoveMethod)
+{
+	DETOURED_CALL(SetFilePointer);
+	DEBUG_LOG_TRUE(L"SetFilePointer", L"%llu %lli %u (%ls)", uintptr_t(hFile), liDistanceToMove.QuadPart, dwMoveMethod, HandleToName(hFile));
+	return True_SetFilePointer(hFile, lDistanceToMove, lpDistanceToMoveHigh, dwMoveMethod);
+}
+
+DWORD Detoured_SetFilePointerEx(HANDLE hFile, LARGE_INTEGER liDistanceToMove, PLARGE_INTEGER lpNewFilePointer, DWORD dwMoveMethod)
+{
+	DETOURED_CALL(SetFilePointerEx);
+	DEBUG_LOG_TRUE(L"SetFilePointerEx", L"%llu %lli %u (%ls)", uintptr_t(hFile), liDistanceToMove.QuadPart, dwMoveMethod, HandleToName(hFile));
+	return True_SetFilePointerEx(hFile, liDistanceToMove, lpNewFilePointer, dwMoveMethod); // This ends up in NtSetInformationFile
+}
+
 DWORD Detoured_GetFileAttributesA(LPCSTR lpFileName)
 {
 	// Is verified that both windows and wine are calling GetFileAttributesW
@@ -4621,13 +4619,13 @@ LSTATUS Detoured_RegOpenKeyExW(HKEY hKey, LPCWSTR lpSubKey, DWORD ulOptions, REG
 	//UBA_NOT_IMPLEMENTED();
 	SuppressCreateFileDetourScope cfs;
 	auto res = True_RegOpenKeyExW(hKey, lpSubKey, ulOptions, samDesired, phkResult);
-	DEBUG_LOG_TRUE(L"RegOpenKeyExW", L"(%ls) -> %ls", lpSubKey, ToString(res == ERROR_SUCCESS));
+	DEBUG_LOG_TRUE(L"RegOpenKeyExW", L"(%ls) -> %ls (%llu)", lpSubKey, ToString(res == ERROR_SUCCESS), phkResult ? u64(*phkResult) : 0ull);
 	return res;
 }
 
 LSTATUS Detoured_RegCreateKeyExW(HKEY hKey, LPCWSTR lpSubKey, DWORD Reserved, LPWSTR lpClass, DWORD dwOptions, REGSAM samDesired, const LPSECURITY_ATTRIBUTES lpSecurityAttributes, PHKEY phkResult, LPDWORD lpdwDisposition)
 {
-	DETOURED_CALL(RegOpenKeyExW);
+	DETOURED_CALL(RegCreateKeyExW);
 	//UBA_NOT_IMPLEMENTED();
 	SuppressCreateFileDetourScope cfs;
 	auto res = True_RegCreateKeyExW(hKey, lpSubKey, Reserved, lpClass, dwOptions, samDesired, lpSecurityAttributes, phkResult, lpdwDisposition);
@@ -4736,6 +4734,7 @@ BOOL Detoured_SetConsoleCursorPosition(HANDLE hConsoleOutput, COORD dwCursorPosi
 {
 	DETOURED_CALL(SetConsoleCursorPosition);
 	DEBUG_LOG_DETOURED(L"SetConsoleCursorPosition", L"");
+	UBA_ASSERT(!isDetouredHandle(hConsoleOutput));
 	return True_SetConsoleCursorPosition(hConsoleOutput, dwCursorPosition);
 }
 
diff --git a/Engine/Source/Programs/UnrealBuildAccelerator/Detours/Private/Windows/UbaDetoursFunctionsNtDll.inl b/Engine/Source/Programs/UnrealBuildAccelerator/Detours/Private/Windows/UbaDetoursFunctionsNtDll.inl
--- a/Engine/Source/Programs/UnrealBuildAccelerator/Detours/Private/Windows/UbaDetoursFunctionsNtDll.inl
+++ b/Engine/Source/Programs/UnrealBuildAccelerator/Detours/Private/Windows/UbaDetoursFunctionsNtDll.inl
@@ -54,20 +54,6 @@ StringBuffer<32> ToString(NTSTATUS s)
 }
 #endif
 
-struct FILE_FS_DEVICE_INFORMATION
-{
-	DEVICE_TYPE DeviceType;
-	ULONG Characteristics;
-};
-
-struct FILE_FS_ATTRIBUTE_INFORMATION
-{
-	ULONG FileSystemAttributes;
-	LONG  MaximumComponentNameLength;
-	ULONG FileSystemNameLength;
-	WCHAR FileSystemName[1];
-};
-
 NTSTATUS Detoured_NtQueryVolumeInformationFile(HANDLE FileHandle, PIO_STATUS_BLOCK IoStatusBlock, PVOID FsInformation, ULONG Length, FS_INFORMATION_CLASS FsInformationClass)
 {
 	DETOURED_CALL(NtQueryVolumeInformationFile);
@@ -104,6 +90,13 @@ NTSTATUS Detoured_NtQueryVolumeInformationFile(HANDLE FileHandle, PIO_STATUS_BLO
 				info.VolumeLabel[0] = 0;
 				return STATUS_SUCCESS;
 			}
+			if (FsInformationClass == FileFsDeviceInformation)
+			{
+				auto& info = *(FILE_FS_DEVICE_INFORMATION*)FsInformation;
+				info.DeviceType = FILE_DEVICE_FILE_SYSTEM;
+				info.Characteristics = 0;
+				return STATUS_SUCCESS;
+			}
 			UBA_ASSERTF(false, L"NtQueryVolumeInformationFile using class %u not handled %ls (%ls)", FsInformationClass, dh.fileObject->fileInfo->name, dh.fileObject->fileInfo->originalName);
 		}
 	}
@@ -206,6 +199,45 @@ NTSTATUS Detoured_NtQueryInformationFile(HANDLE FileHandle, PIO_STATUS_BLOCK IoS
 
 		if (TrueHandle == INVALID_HANDLE_VALUE)
 		{
+			if (FileInformationClass == FileStandardInformation)
+			{
+				UBA_ASSERT(dh.type == HandleType_File);
+				FileInfo& fi = *dh.fileObject->fileInfo;
+				UBA_ASSERT(fi.memoryFile || fi.isFileMap);
+				auto& info = *(FILE_STANDARD_INFORMATION*)FileInformation;
+				u64 size = 0;
+				u64 allocationSize = 0;
+				if (fi.memoryFile)
+				{
+					size = fi.memoryFile->writtenSize;
+					allocationSize = fi.memoryFile->committedSize;
+				}
+				else if (fi.fileMapMem)
+				{
+					size = fi.fileMapMemSize;
+					allocationSize = size;
+				}
+				else
+				{
+					size = fi.size;
+					allocationSize = size;
+				}
+				info.DeletePending = FALSE;
+				info.NumberOfLinks = 0;
+				info.Directory = FALSE;
+				info.AllocationSize.QuadPart = allocationSize;
+				info.EndOfFile.QuadPart = size;
+				return STATUS_SUCCESS;
+			}
+			if (FileInformationClass == FilePositionInformation)
+			{
+				UBA_ASSERT(dh.type == HandleType_File);
+				FileInfo& fi = *dh.fileObject->fileInfo;(void)fi;
+				UBA_ASSERT(fi.memoryFile || fi.isFileMap);
+				auto& info = *(FILE_POSITION_INFORMATION*)FileInformation;
+				info.CurrentByteOffset.QuadPart = dh.pos;
+				return STATUS_SUCCESS;
+			}
 			if (FileInformationClass == FileBasicInformation)
 			{
 				auto& info = *(FILE_BASIC_INFORMATION*)FileInformation;
@@ -550,7 +582,44 @@ NTSTATUS Detoured_NtSetInformationFile(HANDLE FileHandle, PIO_STATUS_BLOCK IoSta
 	{
 		DetouredHandle& h = asDetouredHandle(FileHandle);
 
-		if (FileInformationClass == FileRenameInformation)
+		if (FileInformationClass == FilePositionInformation)
+		{
+			UBA_ASSERT(h.type == HandleType_File);
+			auto& info = *(FILE_POSITION_INFORMATION*)FileInformation;
+			if (h.type >= HandleType_StdErr)
+			{
+				UBA_ASSERT(false); // TODO: What should we do with this?
+				return STATUS_SUCCESS;
+			}
+
+			FileInfo& fi = *h.fileObject->fileInfo;
+			if (fi.memoryFile || fi.isFileMap)
+			{
+				h.pos = info.CurrentByteOffset.QuadPart;
+				return STATUS_SUCCESS;
+			}
+			trueHandle = h.trueHandle;
+			UBA_ASSERT(trueHandle != INVALID_HANDLE_VALUE);
+		}
+		if (FileInformationClass == FileEndOfFileInformation)
+		{
+			UBA_ASSERT(h.type == HandleType_File);
+			FileInfo& fi = *h.fileObject->fileInfo;
+			if (MemoryFile* mf = fi.memoryFile)
+			{
+				DEBUG_LOG_DETOURED(L"SetEndOfFile (MEMORY)", L"%llu NewSize: %llu (%ls) -> Success", uintptr_t(FileHandle), mf->writtenSize, HandleToName(FileHandle));
+				mf->writtenSize = h.pos;
+				mf->isReported = false;
+				mf->EnsureCommitted(h, mf->writtenSize);
+
+				auto& info = *(FILE_END_OF_FILE_INFORMATION*)FileInformation;
+				info.EndOfFile.QuadPart = h.pos;
+				return STATUS_SUCCESS;
+			}
+			trueHandle = h.trueHandle;
+			UBA_ASSERT(trueHandle != INVALID_HANDLE_VALUE);
+		}
+		else if (FileInformationClass == FileRenameInformation)
 		{
 			// We can end up in here through MoveFileEx
 			auto& info = *(FILE_RENAME_INFORMATION*)FileInformation;
@@ -760,6 +829,20 @@ VOID Detoured_RtlFreeAnsiString(PANSI_STRING AnsiString)
 	return True_RtlFreeAnsiString(AnsiString);
 }
 
+BOOLEAN Detoured_RtlIsPackageSid(PSID Sid)
+{
+	DETOURED_CALL(RtlIsPackageSid);
+	DEBUG_LOG_TRUE(L"RtlIsPackageSid", L"");
+	return FALSE;//True_RtlIsPackageSid(Sid);
+}
+
+BOOLEAN Detoured_RtlIsCapabilitySid(PSID sid)
+{
+	DETOURED_CALL(RtlIsCapabilitySid);
+	DEBUG_LOG_TRUE(L"RtlIsCapabilitySid", L"");
+	return FALSE;
+}
+
 NTSTATUS NTAPI Local_NtCreateFile(bool IsCreateFunc, PHANDLE hFileHandle, ACCESS_MASK DesiredAccess, POBJECT_ATTRIBUTES ObjectAttributes, PIO_STATUS_BLOCK IoStatusBlock, PLARGE_INTEGER AllocationSize, ULONG FileAttributes,
 	ULONG ShareAccess, ULONG CreateDisposition, ULONG CreateOptions, PVOID EaBuffer, ULONG EaLength)
 {
diff --git a/Engine/Source/Programs/UnrealBuildAccelerator/Detours/Private/Windows/UbaDetoursFunctionsWin.cpp b/Engine/Source/Programs/UnrealBuildAccelerator/Detours/Private/Windows/UbaDetoursFunctionsWin.cpp
--- a/Engine/Source/Programs/UnrealBuildAccelerator/Detours/Private/Windows/UbaDetoursFunctionsWin.cpp
+++ b/Engine/Source/Programs/UnrealBuildAccelerator/Detours/Private/Windows/UbaDetoursFunctionsWin.cpp
@@ -541,6 +541,7 @@ const wchar_t* ToString(BOOL b) { return b ? L"Success" : L"Error"; }
 #include "UbaDetoursFunctionsDbgHelp.inl"
 #include "UbaDetoursFunctionsShell32.inl"
 #include "UbaDetoursFunctionsRpcrt4.inl"
+#include "UbaDetoursFunctionsCombase.inl"
 
 extern u32 g_consoleStringIndex;
 
@@ -925,10 +926,69 @@ void PreInit(const DetoursPayload& payload)
 				FatalError(1351, L"Failed to preload %s (%u)", *it, GetLastError());
 }
 
+#if 0
+static bool PatchImportByName(HMODULE module, const char* importedDll, const char* funcName, void* replacement)
+{
+	auto base = (unsigned char*)module;
+	auto dos  = (IMAGE_DOS_HEADER*)base;
+	if (dos->e_magic != IMAGE_DOS_SIGNATURE)
+		return false;
+
+	auto nt = (IMAGE_NT_HEADERS*)(base + dos->e_lfanew);
+	if (nt->Signature != IMAGE_NT_SIGNATURE)
+		return false;
+
+	auto& dir = nt->OptionalHeader.DataDirectory[IMAGE_DIRECTORY_ENTRY_IMPORT];
+	if (!dir.VirtualAddress || !dir.Size)
+		return false;
+
+	auto imp = (IMAGE_IMPORT_DESCRIPTOR*)(base + dir.VirtualAddress);
+
+	for (; imp->Name; ++imp)
+	{
+		const char* dllName = (const char*)(base + imp->Name);
+		if (_stricmp(dllName, importedDll) != 0)
+			continue;
+
+		auto thunk     = (IMAGE_THUNK_DATA*)(base + imp->FirstThunk);
+		auto origThunk = (IMAGE_THUNK_DATA*)(base + imp->OriginalFirstThunk);
+
+		for (; origThunk->u1.AddressOfData; ++origThunk, ++thunk)
+		{
+			if (IMAGE_SNAP_BY_ORDINAL(origThunk->u1.Ordinal))
+				continue;
+
+			auto ibn = (IMAGE_IMPORT_BY_NAME*)(base + origThunk->u1.AddressOfData);
+			if (strcmp((const char*)ibn->Name, funcName) != 0)
+				continue;
+
+			DWORD oldProt;
+			if (!VirtualProtect(&thunk->u1.Function, sizeof(void*), PAGE_READWRITE, &oldProt))
+				return false;
+
+			thunk->u1.Function = (ULONG_PTR)replacement;
+
+			DWORD dummy;
+			VirtualProtect(&thunk->u1.Function, sizeof(void*), oldProt, &dummy);
+			FlushInstructionCache(GetCurrentProcess(), nullptr, 0);
+			return true;
+		}
+	}
+
+	return false;
+}
+#endif
+
 void Init(const DetoursPayload& payload, u64 startTime)
 {
 	AddExceptionHandler();
 	InitMemory();
+	
+	#if 0
+	HMODULE exe = GetModuleHandleW(nullptr);
+	PatchImportByName(exe, "ntdll.dll", "RtlIsPackageSid", (void*)Detoured_RtlIsPackageSid);
+	PatchImportByName(exe, "ntdll.dll", "RtlIsCapabilitySid", (void*)Detoured_RtlIsPackageSid);
+	#endif
 
 	DetourAttachFunctions(g_runningRemote);
 
diff --git a/Engine/Source/Programs/UnrealBuildAccelerator/Detours/Private/Windows/UbaDetoursFunctionsWin.h b/Engine/Source/Programs/UnrealBuildAccelerator/Detours/Private/Windows/UbaDetoursFunctionsWin.h
--- a/Engine/Source/Programs/UnrealBuildAccelerator/Detours/Private/Windows/UbaDetoursFunctionsWin.h
+++ b/Engine/Source/Programs/UnrealBuildAccelerator/Detours/Private/Windows/UbaDetoursFunctionsWin.h
@@ -30,6 +30,7 @@
 	DETOURED_FUNCTIONS_SHLWAPI \
 	DETOURED_FUNCTIONS_UCRTBASE \
 	DETOURED_FUNCTIONS_RPCRT4 \
+	DETOURED_FUNCTIONS_COMBASE \
 
 #if defined(_M_ARM64)
 #define DETOURED_FUNCTION_X64(x)
@@ -57,8 +58,6 @@
 	DETOURED_FUNCTION(FlushFileBuffers) \
 	DETOURED_FUNCTION(GetFileSize) \
 	DETOURED_FUNCTION(GetFileSizeEx) \
-	DETOURED_FUNCTION(SetFilePointer) \
-	DETOURED_FUNCTION(SetFilePointerEx) \
 	DETOURED_FUNCTION(SetEndOfFile) \
 	DETOURED_FUNCTION(SetFileTime) \
 	DETOURED_FUNCTION(GetFileTime) \
@@ -229,6 +228,9 @@
 	DETOURED_FUNCTION(_msize_base) \
 	DETOURED_FUNCTION(_recalloc_base) \
 
+#define DETOURED_FUNCTIONS_COMBASE \
+	DETOURED_FUNCTION(CoCreateInstance)
+
 #else
 #define DETOURED_FUNCTIONS_MEMORY
 #define DETOURED_FUNCTIONS_MEMORY_NON_WINE
@@ -274,6 +276,8 @@
 	DETOURED_FUNCTION(GetDiskFreeSpaceExA) \
 	DETOURED_FUNCTION(GetLongPathNameA) \
 	DETOURED_FUNCTION(GetVolumePathNameA) \
+	DETOURED_FUNCTION(SetFilePointer) \
+	DETOURED_FUNCTION(SetFilePointerEx) \
 	DETOURED_FUNCTION(GetFileAttributesA) \
 	DETOURED_FUNCTION(GetFileAttributesExA) \
 	DETOURED_FUNCTION_X64(LoadLibraryW) \
@@ -482,29 +486,23 @@ struct FILE_RENAME_INFORMATION {
     WCHAR FileName[1];
 };
 
-struct FILE_IS_REMOTE_DEVICE_INFORMATION {
+struct FILE_IS_REMOTE_DEVICE_INFORMATION
+{
 	BOOLEAN IsRemote;
-	};
+};
 
-	struct FILE_ID_INFORMATION {
+struct FILE_ID_INFORMATION {
 	ULONGLONG   VolumeSerialNumber;
 	FILE_ID_128 FileId;
-	};
+};
 
-	struct FILE_NAME_INFORMATION {
+struct FILE_NAME_INFORMATION {
 	ULONG FileNameLength;
 	WCHAR FileName[1];
-	};
-
-	struct FILE_STANDARD_INFORMATION {
-	LARGE_INTEGER AllocationSize;
-	LARGE_INTEGER EndOfFile;
-	ULONG         NumberOfLinks;
-	BOOLEAN       DeletePending;
-	BOOLEAN       Directory;
 };
 
-struct FILE_INTERNAL_INFORMATION {
+struct FILE_INTERNAL_INFORMATION
+{
 	LARGE_INTEGER IndexNumber;
 };
 
