[UBA]
* Moved SetEndOfFile to debug and let it fall back to NtSetInformationFile instead

diff --git a/Engine/Source/Programs/UnrealBuildAccelerator/Detours/Private/Windows/UbaDetoursFunctionsKernelBase.inl b/Engine/Source/Programs/UnrealBuildAccelerator/Detours/Private/Windows/UbaDetoursFunctionsKernelBase.inl
--- a/Engine/Source/Programs/UnrealBuildAccelerator/Detours/Private/Windows/UbaDetoursFunctionsKernelBase.inl
+++ b/Engine/Source/Programs/UnrealBuildAccelerator/Detours/Private/Windows/UbaDetoursFunctionsKernelBase.inl
@@ -935,32 +935,6 @@ DWORD Detoured_GetFileSizeEx(HANDLE hFile, PLARGE_INTEGER lpFileSize)
 	return res;
 }
 
-BOOL Detoured_SetEndOfFile(HANDLE hFile)
-{
-	DETOURED_CALL(SetEndOfFile);
-	HANDLE trueHandle = hFile;
-	if (isDetouredHandle(hFile))
-	{
-		auto& dh = asDetouredHandle(hFile);
-		FileInfo& fi = *dh.fileObject->fileInfo;
-		if (MemoryFile* mf = fi.memoryFile)
-		{
-			DEBUG_LOG_DETOURED(L"SetEndOfFile (MEMORY)", L"%llu NewSize: %llu (%ls) -> Success", uintptr_t(hFile), mf->writtenSize, HandleToName(hFile));
-			mf->writtenSize = dh.pos;
-			mf->isReported = false;
-			mf->EnsureCommitted(dh, mf->writtenSize);
-			SetLastError(ERROR_SUCCESS);
-			return true;
-		}
-		trueHandle = dh.trueHandle;
-		UBA_ASSERT(trueHandle != INVALID_HANDLE_VALUE);
-	}
-
-	BOOL res = True_SetEndOfFile(trueHandle);
-	DEBUG_LOG_TRUE(L"SetEndOfFile", L"%llu (%ls) -> %ls", uintptr_t(hFile), HandleToName(hFile), ToString(res));
-	return res;
-}
-
 BOOL Detoured_SetFileTime(HANDLE hFile, const FILETIME* lpCreationTime, const FILETIME* lpLastAccessTime, const FILETIME* lpLastWriteTime)
 {
 	DETOURED_CALL(SetFileTime);
@@ -3947,7 +3921,7 @@ BOOL Detoured_GetVolumePathNameA(LPCSTR lpszFileName, LPSTR lpszVolumePathName,
 DWORD Detoured_SetFilePointer(HANDLE hFile, LONG lDistanceToMove, PLONG lpDistanceToMoveHigh, DWORD dwMoveMethod)
 {
 	DETOURED_CALL(SetFilePointer);
-	DEBUG_LOG_TRUE(L"SetFilePointer", L"%llu %lli %u (%ls)", uintptr_t(hFile), liDistanceToMove.QuadPart, dwMoveMethod, HandleToName(hFile));
+	DEBUG_LOG_TRUE(L"SetFilePointer", L"%llu %lli %u (%ls)", uintptr_t(hFile), ToLargeInteger(lpDistanceToMoveHigh ? *lpDistanceToMoveHigh : 0ll, lDistanceToMove).QuadPart, dwMoveMethod, HandleToName(hFile));
 	return True_SetFilePointer(hFile, lDistanceToMove, lpDistanceToMoveHigh, dwMoveMethod);
 }
 
@@ -3958,6 +3932,13 @@ DWORD Detoured_SetFilePointerEx(HANDLE hFile, LARGE_INTEGER liDistanceToMove, PL
 	return True_SetFilePointerEx(hFile, liDistanceToMove, lpNewFilePointer, dwMoveMethod); // This ends up in NtSetInformationFile
 }
 
+BOOL Detoured_SetEndOfFile(HANDLE hFile)
+{
+	DETOURED_CALL(SetEndOfFile);
+	DEBUG_LOG_TRUE(L"SetEndOfFile", L"%llu (%ls)", uintptr_t(hFile), HandleToName(hFile));
+	return True_SetEndOfFile(hFile);
+}
+
 DWORD Detoured_GetFileAttributesA(LPCSTR lpFileName)
 {
 	// Is verified that both windows and wine are calling GetFileAttributesW
diff --git a/Engine/Source/Programs/UnrealBuildAccelerator/Detours/Private/Windows/UbaDetoursFunctionsNtDll.inl b/Engine/Source/Programs/UnrealBuildAccelerator/Detours/Private/Windows/UbaDetoursFunctionsNtDll.inl
--- a/Engine/Source/Programs/UnrealBuildAccelerator/Detours/Private/Windows/UbaDetoursFunctionsNtDll.inl
+++ b/Engine/Source/Programs/UnrealBuildAccelerator/Detours/Private/Windows/UbaDetoursFunctionsNtDll.inl
@@ -231,6 +231,10 @@ NTSTATUS Detoured_NtQueryInformationFile(HANDLE FileHandle, PIO_STATUS_BLOCK IoS
 			}
 			if (FileInformationClass == FilePositionInformation)
 			{
+				if (dh.type >= HandleType_StdErr)
+				{
+					return STATUS_SUCCESS;
+				}
 				UBA_ASSERT(dh.type == HandleType_File);
 				FileInfo& fi = *dh.fileObject->fileInfo;(void)fi;
 				UBA_ASSERT(fi.memoryFile || fi.isFileMap);
@@ -584,9 +588,14 @@ NTSTATUS Detoured_NtSetInformationFile(HANDLE FileHandle, PIO_STATUS_BLOCK IoSta
 
 		if (FileInformationClass == FilePositionInformation)
 		{
-			UBA_ASSERT(h.type == HandleType_File);
 			auto& info = *(FILE_POSITION_INFORMATION*)FileInformation;
 			if (h.type >= HandleType_StdErr)
+			{
+				info.CurrentByteOffset.QuadPart = 0;
+				return STATUS_SUCCESS;
+			}
+			UBA_ASSERT(h.type == HandleType_File);
+			if (h.type >= HandleType_StdErr)
 			{
 				UBA_ASSERT(false); // TODO: What should we do with this?
 				return STATUS_SUCCESS;
