[UBA]
* Changed how rename works. it was broken before in certain scenarios where same file was renamed multiple times.

diff --git a/Engine/Source/Programs/UnrealBuildAccelerator/Common/Private/UbaSession.cpp b/Engine/Source/Programs/UnrealBuildAccelerator/Common/Private/UbaSession.cpp
--- a/Engine/Source/Programs/UnrealBuildAccelerator/Common/Private/UbaSession.cpp
+++ b/Engine/Source/Programs/UnrealBuildAccelerator/Common/Private/UbaSession.cpp
@@ -3294,6 +3294,10 @@ namespace uba
 
 		if (msg.deleteOnClose)
 		{
+			auto& writtenFiles = msg.process.m_shared.writtenFiles;
+			auto findIt2 = writtenFiles.find(activeFile.nameKey);
+			if (findIt2 != writtenFiles.end())
+				writtenFiles.erase(findIt2);
 			RemoveWrittenFile(msg.process, activeFile.nameKey);
 		}
 		else
@@ -3304,17 +3308,30 @@ namespace uba
 			if (!msg.newName.IsEmpty())
 			{
 				UBA_ASSERT(!msg.deleteOnClose);
-				RemoveWrittenFile(msg.process, key);
 				name = msg.newName;
 				key = msg.newNameKey;
 				if (!m_runningRemote)
 					msgName = msg.newName;
 			}
 			UBA_ASSERT(key != StringKeyZero);
+			auto& writtenFiles = msg.process.m_shared.writtenFiles;
+
 			SCOPED_FUTEX(msg.process.m_shared.writtenFilesLock, writtenLock);
-			auto insres = msg.process.m_shared.writtenFiles.try_emplace(key);
+			auto insres = writtenFiles.try_emplace(key);
 			WrittenFile& writtenFile = insres.first->second;
 
+			if (!msg.newName.IsEmpty()) // Transfer potential memory map from old file
+			{
+				auto oldFindIt = writtenFiles.find(activeFile.nameKey);
+				if (oldFindIt != writtenFiles.end())
+				{
+					UBA_ASSERT(insres.second || !oldFindIt->second.memoryHandle.IsValid());
+					writtenFile.memoryHandle = oldFindIt->second.memoryHandle;
+					writtenFile.memoryWritten = oldFindIt->second.memoryWritten;
+					writtenFiles.erase(oldFindIt);
+				}
+			}
+
 			if (m_allowOutputFiles && writtenFile.owner != nullptr && writtenFile.owner != &msg.process)
 			{
 				// This can happen when library has /GL (whole program optimization) but target has not.. then link.exe will restart
diff --git a/Engine/Source/Programs/UnrealBuildAccelerator/Common/Private/UbaSessionClient.cpp b/Engine/Source/Programs/UnrealBuildAccelerator/Common/Private/UbaSessionClient.cpp
--- a/Engine/Source/Programs/UnrealBuildAccelerator/Common/Private/UbaSessionClient.cpp
+++ b/Engine/Source/Programs/UnrealBuildAccelerator/Common/Private/UbaSessionClient.cpp
@@ -949,15 +949,22 @@ namespace uba
 			{
 				UBA_ASSERT(msg.toKey != StringKeyZero);
 				auto insres = writtenFiles.try_emplace(msg.toKey);
-				UBA_ASSERTF(insres.second, TC("Moving written file %s to other written file %s. (%s)"), fromName, toName, process.m_startInfo.description);
-				insres.first->second = findIt->second;
-				insres.first->second.key = msg.toKey;
-				insres.first->second.name = toName;
-				insres.first->second.owner = &process;
+				WrittenFile& toFile = insres.first->second;
+				
+				if (!insres.second)
+				{
+					UBA_ASSERTF(!toFile.memoryHandle.IsValid(), TC("This should be cleaned up"));
+				}
+
+				toFile = findIt->second; // Copy all state first and then fixup
+				toFile.key = msg.toKey;
+				toFile.name = toName;
+				toFile.owner = &process;
 				writtenFiles.erase(findIt);
 			}
-			else
+			else if (!msg.fromName.StartsWith(m_tempPath) || !msg.toName.StartsWith(m_tempPath))
 			{
+				UBA_ASSERTF(false, TC("TODO: Moving file %s to %s not owned by client. Need to tell server"), fromName, toName);
 				// TODO: Need to tell server 
 			}
 		}
@@ -969,7 +976,6 @@ namespace uba
 			if (findIt != m_outputFiles.end())
 			{
 				auto insres = m_outputFiles.try_emplace(toName);
-				UBA_ASSERTF(insres.second, TC("Failed to add move destination file %s as output file because it is already added. (Moved from %s)"), toName, fromName);
 				insres.first->second = findIt->second;
 				m_outputFiles.erase(findIt);
 				sendMove = false;
diff --git a/Engine/Source/Programs/UnrealBuildAccelerator/Detours/Private/Windows/UbaDetoursFunctionsKernelBase.inl b/Engine/Source/Programs/UnrealBuildAccelerator/Detours/Private/Windows/UbaDetoursFunctionsKernelBase.inl
--- a/Engine/Source/Programs/UnrealBuildAccelerator/Detours/Private/Windows/UbaDetoursFunctionsKernelBase.inl
+++ b/Engine/Source/Programs/UnrealBuildAccelerator/Detours/Private/Windows/UbaDetoursFunctionsKernelBase.inl
@@ -2390,14 +2390,19 @@ BOOL Detoured_SetFileInformationByHandle(HANDLE hFile, FILE_INFO_BY_HANDLE_CLASS
 
 	if (FileInformationClass == FileRenameInfo)
 	{
-		DEBUG_LOG_TRUE(L"SetFileInformationByHandle", L"%llu (FileRenameInfo)", uintptr_t(hFile));
+		// We can't call True_SetFileInformationByHandle on remote builds because it resolves relative paths causing them to be relative to remote working dir
 		auto& info = *(FILE_RENAME_INFO*)lpFileInformation;
-		StringBuffer<> newName;
-		newName.Append(info.FileName, info.FileNameLength / 2);
-		t_renameFileNewName = newName.data;
-		bool res = True_SetFileInformationByHandle(hFile, FileInformationClass, lpFileInformation, dwBufferSize);
-		t_renameFileNewName = nullptr;
-		return res;
+		u8 data[2048];
+		UBA_ASSERT(sizeof(FILE_RENAME_INFORMATION) + info.FileNameLength <= sizeof(data));
+		auto& ntInfo = *(FILE_RENAME_INFORMATION*)data;
+		ntInfo.ReplaceIfExists = info.ReplaceIfExists ? TRUE : FALSE;
+		ntInfo.RootDirectory   = info.RootDirectory;
+		ntInfo.FileNameLength  = info.FileNameLength;
+		memcpy(ntInfo.FileName, info.FileName, info.FileNameLength);
+		DEBUG_LOG_TRUE(L"SetFileInformationByHandle", L"%llu (FileRenameInfo)", uintptr_t(hFile));
+		IO_STATUS_BLOCK io;
+		NTSTATUS res = Detoured_NtSetInformationFile(hFile, &io, &ntInfo, dwBufferSize, FileRenameInformation);
+		return SetLastErrorFromNtStatus(res);
 	}
 	else if (FileInformationClass == FileDispositionInfo)
 	{
diff --git a/Engine/Source/Programs/UnrealBuildAccelerator/Detours/Private/Windows/UbaDetoursFunctionsNtDll.inl b/Engine/Source/Programs/UnrealBuildAccelerator/Detours/Private/Windows/UbaDetoursFunctionsNtDll.inl
--- a/Engine/Source/Programs/UnrealBuildAccelerator/Detours/Private/Windows/UbaDetoursFunctionsNtDll.inl
+++ b/Engine/Source/Programs/UnrealBuildAccelerator/Detours/Private/Windows/UbaDetoursFunctionsNtDll.inl
@@ -630,22 +630,16 @@ NTSTATUS Detoured_NtSetInformationFile(HANDLE FileHandle, PIO_STATUS_BLOCK IoSta
 		}
 		else if (FileInformationClass == FileRenameInformation)
 		{
-			// We can end up in here through MoveFileEx
+			// We can end up in here through MoveFileEx and SetFileInformationByHandle
 			auto& info = *(FILE_RENAME_INFORMATION*)FileInformation;
-			const wchar_t* newNamePtr = t_renameFileNewName;
-			StringBuffer<> newNameTemp;
-			if (!newNamePtr)
-			{
-				newNameTemp.Append(info.FileName, info.FileNameLength / 2);
-				newNamePtr = newNameTemp.data;
-			}
-			if (StartsWith(newNamePtr, L"\\??\\"))
-				newNamePtr += 4;
+			StringView nameView(info.FileName, info.FileNameLength / 2);
+			if (nameView.StartsWith(L"\\??\\"))
+				nameView = nameView.Skip(4);
 			StringBuffer<> newName;
-			FixPath(newName, newNamePtr);
+			FixPath(newName, nameView.data, nameView.count);
 			DevirtualizePath(newName);
 			FileObject& fo = *h.fileObject;
-			fo.newName = newName.data;
+			fo.newName = newName.ToString();
 
 			StringKey newFileNameKey = ToStringKeyLower(newName);
 
@@ -661,15 +655,6 @@ NTSTATUS Detoured_NtSetInformationFile(HANDLE FileHandle, PIO_STATUS_BLOCK IoSta
 
 			UBA_ASSERTF(newFileInfo, L"File info already exists for the rename we are doing from %ls to %ls", fo.fileInfo->originalName, newName.data); // TODO: Implement when we find a test case
 #endif
-
-			if (!fo.closeId)
-			{
-				wchar_t temp[1024];
-				u64 size;
-				StringBuffer<> fixedPath;
-				FixPath(fixedPath, newName.data);
-				Rpc_CreateFileW(fixedPath, newFileNameKey, AccessFlag_Write, temp, sizeof_array(temp), size, fo.closeId, true);
-			}
 			DEBUG_LOG_DETOURED(L"NtSetInformationFile", L"File is set to be renamed on close (from %ls to %ls)", HandleToName(FileHandle), fo.newName.c_str());
 
 			if (auto memoryFile = fo.fileInfo->memoryFile)
@@ -945,6 +930,8 @@ NTSTATUS NTAPI Shared_NtCreateFile(bool IsCreateFunc, PHANDLE hFileHandle, ACCES
 		{
 			if (!CanDetour(createFileName))
 				suppressCreateFileDetour = true;
+			else if (StartsWith(createFileName, TC("\\\\?\\pipe\\")) && !g_rules->IsThrowAway(ToView(createFileName), g_runningRemote))
+				suppressCreateFileDetour = true;
 			else if (!FixPath(fileName, createFileName))
 			{
 				DEBUG_LOG(L"FixPath failed for string '%ls'", createFileName);
@@ -954,7 +941,7 @@ NTSTATUS NTAPI Shared_NtCreateFile(bool IsCreateFunc, PHANDLE hFileHandle, ACCES
 			// TODO: Instead of using t_createFileFileName.. should we just resolve the path from ObjectAttributes->RootDirectory?
 			ObjectAttributes->RootDirectory = nullptr;
 		}
-		else if (bufChars >= 8 && memcmp(buf, L"\\??\\", 8) == 0)
+		else if (bufChars >= 4 && memcmp(buf, L"\\??\\", 8) == 0)
 		{
 			if (!CanDetour(buf))
 				suppressCreateFileDetour = true;
@@ -968,21 +955,29 @@ NTSTATUS NTAPI Shared_NtCreateFile(bool IsCreateFunc, PHANDLE hFileHandle, ACCES
 				if (isDetouredHandle(ObjectAttributes->RootDirectory))
 				{
 					auto& dh = asDetouredHandle(ObjectAttributes->RootDirectory);
-					fileName.Append(dh.fileObject->fileInfo->originalName).EnsureEndsWithSlash().Append(buf, bufChars);
+					fileName.Append(dh.fileObject->fileInfo->originalName);
+					if (bufChars)
+						fileName.EnsureEndsWithSlash().Append(buf, bufChars);
 					rootDir = dh.trueHandle;
 					ObjectAttributes->RootDirectory = nullptr;
 				}
 				else if (isListDirectoryHandle(ObjectAttributes->RootDirectory))
 				{
 					auto& lh = asListDirectoryHandle(ObjectAttributes->RootDirectory);
-					fileName.Append(lh.originalName).EnsureEndsWithSlash().Append(buf, bufChars);
+					fileName.Append(lh.originalName);
+					if (bufChars)
+						fileName.EnsureEndsWithSlash().Append(buf, bufChars);
 					ObjectAttributes->RootDirectory = nullptr;
 				}
 				else
 				{
-					// TODO: Revisit!
-					//UBA_ASSERT(false);
-					suppressCreateFileDetour = true;
+					tchar rootDirPath[1024];
+					u32 rootDirLen = True_GetFinalPathNameByHandleW(rootDir, rootDirPath, 1024, 0);
+					UBA_ASSERT(rootDirLen != 0 && rootDirLen < sizeof(rootDirPath));
+					UBA_ASSERT(StartsWith(rootDirPath, TC("\\\\?\\")));
+					fileName.Append(rootDirPath + 4, rootDirLen - 4);
+					if (bufChars)
+						fileName.EnsureEndsWithSlash().Append(buf, bufChars);
 				}
 			}
 			else
@@ -1338,6 +1333,8 @@ NTSTATUS NTAPI Shared_NtCreateFile(bool IsCreateFunc, PHANDLE hFileHandle, ACCES
 
 	auto CreateFileHandle = [&](HANDLE th = INVALID_HANDLE_VALUE)
 		{
+			IoStatusBlock->Status = STATUS_SUCCESS;
+			IoStatusBlock->Information = 1;
 			auto fo = new FileObject();
 			fo->desiredAccess = DesiredAccess;
 			fo->closeId = closeId;
@@ -1912,16 +1909,58 @@ NTSTATUS NTAPI Detoured_NtClose(HANDLE handle)
 				newInfo.name = newInfo.originalName;
 				newInfo.fileNameKey = fileNameKey;
 				UBA_ASSERT(!fo->deleteOnClose);
+				u32 refCount = fi.refCount;
 				fi = {};
+				fi.refCount = refCount;
 				fi.deleted = true;
 				fo->ownsFileInfo = false;
 				fo->newName.clear();
 			}
 		}
 	}
-	else if (fo->deleteOnClose && dh.trueHandle == INVALID_HANDLE_VALUE) // We have used an optimized handle that actually never opens the file so we need to delete it manually
+	else if (fo->deleteOnClose) // We have used an optimized handle that actually never opens the file so we need to delete it manually
 	{
-		DeleteFileW(fi.originalName);
+		if (dh.trueHandle == INVALID_HANDLE_VALUE)
+			DeleteFileW(fi.originalName);
+		else if (!fo->closeId && !g_exitMessageSent)
+		{
+			// If file was DeleteOnClose but never got a closeId we need to create one so we can delete the file on session side.
+			u64 size;
+			Rpc_CreateFileW(ToView(fi.originalName), fi.fileNameKey, AccessFlag_Write, temp, sizeof_array(temp), size, fo->closeId, true);
+		}
+	}
+
+	if (!fo->newName.empty())
+	{
+		if (!fo->closeId)
+		{
+			// If file was DeleteOnClose but never got a closeId we need to create one so we can delete the file on session side.
+			u64 size;
+			Rpc_CreateFileW(ToView(fi.originalName), fi.fileNameKey, AccessFlag_Write, temp, sizeof_array(temp), size, fo->closeId, true);
+		}
+
+		UBA_ASSERT(!fo->deleteOnClose);
+		// It might be that same process will open it again, so we will need to update the mapping table
+		StringBuffer<> fixedNewName;
+		FixPath(fixedNewName, fo->newName.c_str());
+		fixedNewName.MakeLower();
+		StringKey fileNameKey = ToStringKey(fixedNewName);
+		SCOPED_WRITE_LOCK(g_mappedFileTable.m_lookupLock, _);
+		auto insres = g_mappedFileTable.m_lookup.try_emplace(fileNameKey);
+		FileInfo& newInfo = insres.first->second;
+		newInfo = fi;
+		newInfo.originalName = g_memoryBlock.Strdup(fo->newName).data;
+		UBA_ASSERT(Equals(newInfo.name, fi.name));
+		if (!g_runningRemote)
+			newInfo.name = newInfo.originalName;
+		newInfo.fileNameKey = fileNameKey;
+		UBA_ASSERT(!fo->deleteOnClose);
+		u32 refCount = fi.refCount;
+		fi = {};
+		fi.refCount = refCount;
+		fi.deleted = true;
+		fo->ownsFileInfo = false;
+		//fo->newName.clear(); // Don't remove newName here since we might want to enter the update close
 	}
 
 	if (fo->closeId)
