[UBA]
* Reverted change in BinaryReader read string.

diff --git a/Engine/Source/Programs/UnrealBuildAccelerator/Core/Public/UbaBinaryReaderWriter.h b/Engine/Source/Programs/UnrealBuildAccelerator/Core/Public/UbaBinaryReaderWriter.h
--- a/Engine/Source/Programs/UnrealBuildAccelerator/Core/Public/UbaBinaryReaderWriter.h
+++ b/Engine/Source/Programs/UnrealBuildAccelerator/Core/Public/UbaBinaryReaderWriter.h
@@ -349,7 +349,7 @@ namespace uba
 	u64 BinaryReader::ReadString(tchar* str, u64 strCapacity)
 	{
 		u64 charLen = Read7BitEncoded();
-		if (charLen < strCapacity)
+		if (charLen < strCapacity - 1)
 			return InternalReadString(str, charLen);
 
 		u64 newCharLen = strCapacity - 1;
