[UBA]
* Fixed so unit tests can pass in shipping with asserts enabled

diff --git a/Engine/Source/Programs/UnrealBuildAccelerator/Test/Private/UbaTestBasics.cpp b/Engine/Source/Programs/UnrealBuildAccelerator/Test/Private/UbaTestBasics.cpp
--- a/Engine/Source/Programs/UnrealBuildAccelerator/Test/Private/UbaTestBasics.cpp
+++ b/Engine/Source/Programs/UnrealBuildAccelerator/Test/Private/UbaTestBasics.cpp
@@ -506,20 +506,21 @@ namespace uba
 			return false;
 		#endif
 
-		#if !UBA_DEBUG
-		u8 mem[1024];
-		BinaryWriter writer(mem);
-		writer.WriteString(TC("Foo"));
-		writer.WriteString(TC("Bar"));
-		BinaryReader reader(mem);
-		StringBuffer<3> str;
-		reader.ReadString(str);
-		if (!str.Equals(TC("Fo")))
-			return false;
-		reader.ReadString(str.Clear());
-		if (!str.Equals(TC("Ba")))
-			return false;
-		#endif
+		if (!IsDebug)
+		{
+			u8 mem[1024];
+			BinaryWriter writer(mem);
+			writer.WriteString(TC("Foo"));
+			writer.WriteString(TC("Bar"));
+			BinaryReader reader(mem);
+			StringBuffer<3> str;
+			reader.ReadString(str);
+			if (!str.Equals(TC("Fo")))
+				return false;
+			reader.ReadString(str.Clear());
+			if (!str.Equals(TC("Ba")))
+				return false;
+		}
 
 		return true;
 	}
