[UBA]
* Fixed compile errors when enabling debug detoured functions in shipping

diff --git a/Engine/Source/Programs/UnrealBuildAccelerator/Detours/Private/Windows/UbaDetoursFunctionsKernelBase.inl b/Engine/Source/Programs/UnrealBuildAccelerator/Detours/Private/Windows/UbaDetoursFunctionsKernelBase.inl
--- a/Engine/Source/Programs/UnrealBuildAccelerator/Detours/Private/Windows/UbaDetoursFunctionsKernelBase.inl
+++ b/Engine/Source/Programs/UnrealBuildAccelerator/Detours/Private/Windows/UbaDetoursFunctionsKernelBase.inl
@@ -287,7 +287,7 @@ BOOL Detoured_CreateDirectoryW(LPCWSTR lpPathName, LPSECURITY_ATTRIBUTES lpSecur
 	return res;
 }
 
-#if UBA_DEBUG
+#if DETOURED_INCLUDE_DEBUG
 thread_local bool t_calledRemoveDirectoryW;
 #endif
 
@@ -295,8 +295,8 @@ BOOL Detoured_RemoveDirectoryW(LPCWSTR lpPathName)
 {
 	DETOURED_CALL(RemoveDirectoryW);
 
-	#if UBA_DEBUG
-	t_calledRemoveDirectoryW = 1;
+	#if DETOURED_INCLUDE_DEBUG
+	t_calledRemoveDirectoryW = true;
 	#endif
 
 	StringBuffer<> pathName;
@@ -1552,14 +1552,14 @@ BOOL Detoured_CreateHardLinkW(LPCWSTR lpFileName, LPCWSTR lpExistingFileName, LP
 	return res;
 }
 
-#if UBA_DEBUG
+#if DETOURED_INCLUDE_DEBUG
 thread_local bool t_calledDeleteFileW;
 #endif
 
 BOOL Detoured_DeleteFileW(LPCWSTR lpFileName)
 {
 	DETOURED_CALL(DeleteFileW);
-	#if UBA_DEBUG
+	#if DETOURED_INCLUDE_DEBUG
 	t_calledDeleteFileW = true;
 	#endif
 
@@ -3899,7 +3899,7 @@ BOOL Detoured_GetThreadPreferredUILanguages(DWORD dwFlags, PULONG pulNumLanguage
 	}
 }
 
-#if defined(DETOURED_INCLUDE_DEBUG)
+#if DETOURED_INCLUDE_DEBUG
 
 BOOL Detoured_GetDiskFreeSpaceExA(LPCSTR lpDirectoryName, PULARGE_INTEGER lpFreeBytesAvailableToCaller, PULARGE_INTEGER lpTotalNumberOfBytes, PULARGE_INTEGER lpTotalNumberOfFreeBytes)
 {
diff --git a/Engine/Source/Programs/UnrealBuildAccelerator/Detours/Private/Windows/UbaDetoursFunctionsMiMalloc.inl b/Engine/Source/Programs/UnrealBuildAccelerator/Detours/Private/Windows/UbaDetoursFunctionsMiMalloc.inl
--- a/Engine/Source/Programs/UnrealBuildAccelerator/Detours/Private/Windows/UbaDetoursFunctionsMiMalloc.inl
+++ b/Engine/Source/Programs/UnrealBuildAccelerator/Detours/Private/Windows/UbaDetoursFunctionsMiMalloc.inl
@@ -229,7 +229,7 @@ void* Detoured__recalloc_base(void* memblock, size_t num, size_t size)
 	return mi_recalloc(memblock, num, size);
 }
 
-#if defined(DETOURED_INCLUDE_DEBUG)
+#if DETOURED_INCLUDE_DEBUG
 
 size_t Detoured__aligned_msize(void* p, size_t alignment, size_t offset)
 {
diff --git a/Engine/Source/Programs/UnrealBuildAccelerator/Detours/Private/Windows/UbaDetoursFunctionsNtDll.inl b/Engine/Source/Programs/UnrealBuildAccelerator/Detours/Private/Windows/UbaDetoursFunctionsNtDll.inl
--- a/Engine/Source/Programs/UnrealBuildAccelerator/Detours/Private/Windows/UbaDetoursFunctionsNtDll.inl
+++ b/Engine/Source/Programs/UnrealBuildAccelerator/Detours/Private/Windows/UbaDetoursFunctionsNtDll.inl
@@ -1647,7 +1647,7 @@ NTSTATUS NTAPI Detoured_NtOpenFile(PHANDLE FileHandle, ACCESS_MASK DesiredAccess
 	return Shared_NtCreateFile(false, FileHandle, DesiredAccess, ObjectAttributes, IoStatusBlock, (PLARGE_INTEGER)NULL, 0L, ShareAccess, FILE_OPEN, OpenOptions, (PVOID)NULL, 0L);
 }
 
-#if UBA_DEBUG
+#if DETOURED_INCLUDE_DEBUG
 thread_local bool t_ntReadFileCalled;
 #endif
 
@@ -1655,7 +1655,7 @@ NTSTATUS NTAPI Detoured_NtReadFile(HANDLE FileHandle, HANDLE Event, PIO_APC_ROUT
 {
 	DETOURED_CALL(NtReadFile);
 
-	#if UBA_DEBUG
+	#if DETOURED_INCLUDE_DEBUG
 	t_ntReadFileCalled = true;
 	#endif
 
@@ -2077,7 +2077,7 @@ NTSTATUS Detoured_NtQueryInformationProcess(HANDLE ProcessHandle, PROCESSINFOCLA
 	return res;
 }
 
-#if UBA_DEBUG
+#if DETOURED_INCLUDE_DEBUG
 thread_local bool t_ntQuerySecurityObjectCalled;
 #endif
 
@@ -2085,7 +2085,7 @@ NTSTATUS Detoured_NtQuerySecurityObject(HANDLE Handle, SECURITY_INFORMATION Secu
 {
 	DETOURED_CALL(NtQuerySecurityObject);
 	
-	#if UBA_DEBUG
+	#if DETOURED_INCLUDE_DEBUG
 	t_ntQuerySecurityObjectCalled = true;
 	#endif
 
@@ -2185,7 +2185,7 @@ NTSTATUS Detoured_NtQuerySecurityObject(HANDLE Handle, SECURITY_INFORMATION Secu
 }
 
 
-#if defined(DETOURED_INCLUDE_DEBUG)
+#if DETOURED_INCLUDE_DEBUG
 
 NTSTATUS NTAPI Detoured_NtCreateIoCompletion(PHANDLE IoCompletionHandle, ACCESS_MASK DesiredAccess, POBJECT_ATTRIBUTES ObjectAttributes, ULONG Count)
 {
diff --git a/Engine/Source/Programs/UnrealBuildAccelerator/Detours/Private/Windows/UbaDetoursFunctionsUcrtBase.inl b/Engine/Source/Programs/UnrealBuildAccelerator/Detours/Private/Windows/UbaDetoursFunctionsUcrtBase.inl
--- a/Engine/Source/Programs/UnrealBuildAccelerator/Detours/Private/Windows/UbaDetoursFunctionsUcrtBase.inl
+++ b/Engine/Source/Programs/UnrealBuildAccelerator/Detours/Private/Windows/UbaDetoursFunctionsUcrtBase.inl
@@ -335,7 +335,7 @@ errno_t Detoured__wsplitpath_s(const wchar_t* path, wchar_t* drive, size_t drive
 }
 #endif
 
-#if defined(DETOURED_INCLUDE_DEBUG)
+#if DETOURED_INCLUDE_DEBUG
 int Detoured__wcsnicoll_l(const wchar_t* string1, const wchar_t* string2, size_t count, _locale_t locale)
 {
 	DETOURED_CALL(_wcsnicoll_l);
@@ -429,4 +429,4 @@ int Detoured_fputs(const char* str, FILE* stream)
 }
 
 
-#endif // defined(DETOURED_INCLUDE_DEBUG)
+#endif // DETOURED_INCLUDE_DEBUG
diff --git a/Engine/Source/Programs/UnrealBuildAccelerator/Detours/Private/Windows/UbaDetoursFunctionsWin.cpp b/Engine/Source/Programs/UnrealBuildAccelerator/Detours/Private/Windows/UbaDetoursFunctionsWin.cpp
--- a/Engine/Source/Programs/UnrealBuildAccelerator/Detours/Private/Windows/UbaDetoursFunctionsWin.cpp
+++ b/Engine/Source/Programs/UnrealBuildAccelerator/Detours/Private/Windows/UbaDetoursFunctionsWin.cpp
@@ -38,7 +38,7 @@ DETOURED_FUNCTIONS_MEMORY
 DETOURED_FUNCTIONS_MEMORY_NON_WINE
 #undef DETOURED_FUNCTION
 
-#if !defined(DETOURED_INCLUDE_DEBUG)
+#if !DETOURED_INCLUDE_DEBUG
 #define True_GetEnvironmentVariableW GetEnvironmentVariableW
 #endif
 
@@ -744,7 +744,7 @@ int DetourAttachFunctions(bool runningRemote)
 	if (IsDebuggerPresent())
 	{
 		True_CreateProcessW = nullptr;
-		#if defined(DETOURED_INCLUDE_DEBUG)
+		#if DETOURED_INCLUDE_DEBUG
 		True_CreateProcessA = nullptr;
 		True_CreateProcessAsUserW = nullptr;
 		#endif
diff --git a/Engine/Source/Programs/UnrealBuildAccelerator/Detours/Private/Windows/UbaDetoursFunctionsWin.h b/Engine/Source/Programs/UnrealBuildAccelerator/Detours/Private/Windows/UbaDetoursFunctionsWin.h
--- a/Engine/Source/Programs/UnrealBuildAccelerator/Detours/Private/Windows/UbaDetoursFunctionsWin.h
+++ b/Engine/Source/Programs/UnrealBuildAccelerator/Detours/Private/Windows/UbaDetoursFunctionsWin.h
@@ -13,12 +13,13 @@
 #include <wchar.h>
 #include <stdio.h>
 
-
 #if UBA_DEBUG
-#define DETOURED_INCLUDE_DEBUG
+#define DETOURED_INCLUDE_DEBUG 1
+#else
+#define DETOURED_INCLUDE_DEBUG 0
 #endif
 
-#if defined(DETOURED_INCLUDE_DEBUG)
+#if DETOURED_INCLUDE_DEBUG
 #include <direct.h>
 #endif
 
@@ -235,7 +236,7 @@
 #define DETOURED_FUNCTIONS_MEMORY_NON_WINE
 #endif
 
-#if defined(DETOURED_INCLUDE_DEBUG)
+#if DETOURED_INCLUDE_DEBUG
 
 #define DETOURED_FUNCTIONS_KERNELBASE_DEBUG \
 	DETOURED_FUNCTION(ExitProcess) \
