[UBA]
* Replaced WriteFile/WriteFileEx with NtWriteFile

diff --git a/Engine/Source/Programs/UnrealBuildAccelerator/Detours/Private/Windows/UbaDetoursFunctionsKernelBase.inl b/Engine/Source/Programs/UnrealBuildAccelerator/Detours/Private/Windows/UbaDetoursFunctionsKernelBase.inl
--- a/Engine/Source/Programs/UnrealBuildAccelerator/Detours/Private/Windows/UbaDetoursFunctionsKernelBase.inl
+++ b/Engine/Source/Programs/UnrealBuildAccelerator/Detours/Private/Windows/UbaDetoursFunctionsKernelBase.inl
@@ -634,37 +634,6 @@ BOOL Detoured_DeleteSynchronizationBarrier(LPSYNCHRONIZATION_BARRIER lpBarrier)
 	return false;
 }
 
-StringBuffer<32 * 1024> g_stdFile;
-ReaderWriterLock& g_stdFileLock = *new ReaderWriterLock();
-
-void WriteStdFile(LPCVOID lpBuffer, DWORD nNumberOfBytesToWrite, bool isError)
-{
-	if (!g_conEnabled[isError?0:1] || g_suppressLogging)
-		return;
-
-	SCOPED_WRITE_LOCK(g_stdFileLock, lock);
-	u32 start = 0;
-	u32 i = 0;
-	auto bufferStr = (const char*)lpBuffer;
-	while (i != nNumberOfBytesToWrite)
-	{
-		if (bufferStr[i] == '\n')
-		{
-			int len = i - start;
-			if (len > 0 && bufferStr[i - 1] == '\r')
-				--len;
-			if (len)
-				g_stdFile.Appendf(L"%.*hs", len, bufferStr + start);
-			Rpc_WriteLog(g_stdFile.data, g_stdFile.count, false, isError);
-			g_stdFile.Clear();
-			start = i + 1;
-		}
-		++i;
-	}
-	if (u32 left = nNumberOfBytesToWrite - start)
-		g_stdFile.Appendf(L"%.*hs", left, bufferStr + start);
-}
-
 HANDLE Detoured_CreateToolhelp32Snapshot(DWORD dwFlags, DWORD th32ProcessID)
 {
 	DETOURED_CALL(CreateToolhelp32Snapshot);
@@ -683,110 +652,6 @@ int WriteFileAccessViolation(PEXCEPTION_POINTERS ep)
 	return EXCEPTION_CONTINUE_SEARCH;
 }
 
-BOOL Detoured_WriteFile(HANDLE hFile, LPCVOID lpBuffer, DWORD nNumberOfBytesToWrite, LPDWORD lpNumberOfBytesWritten, LPOVERLAPPED lpOverlapped)
-{
-	DWORD temp = 0;
-	if (!lpNumberOfBytesWritten)
-		lpNumberOfBytesWritten = &temp;
-
-	DETOURED_CALL(WriteFile);
-	HANDLE trueHandle = hFile;
-	UBA_ASSERT(!isListDirectoryHandle(hFile));
-	if (isDetouredHandle(hFile))
-	{
-		auto& dh = asDetouredHandle(hFile);
-		auto& fo = *dh.fileObject;
-
-		if (dh.type >= HandleType_StdErr)
-		{
-			if (dh.type != HandleType_StdIn)
-			{
-				//DEBUG_LOG_DETOURED(L"WriteStdFile1", L"%llu", uintptr_t(hFile));
-				WriteStdFile(lpBuffer, nNumberOfBytesToWrite, dh.type == HandleType_StdErr);
-			}
-			*lpNumberOfBytesWritten = nNumberOfBytesToWrite;
-			SetLastError(ERROR_SUCCESS);
-			return true;
-		}
-
-		auto& fi = *fo.fileInfo;
-		if (MemoryFile* mf = fi.memoryFile)
-		{
-			if (lpOverlapped)
-			{
-				u64 offset = ToLargeInteger(lpOverlapped->OffsetHigh, lpOverlapped->Offset).QuadPart;
-				u64 writtenSize = offset + nNumberOfBytesToWrite;
-				{
-					SCOPED_WRITE_LOCK(mf->lock, lock);
-					mf->EnsureCommitted(dh, writtenSize);
-					mf->writtenSize = Max(mf->writtenSize, writtenSize);
-					mf->isReported = false;
-				}
-				memcpy(mf->baseAddress + offset, lpBuffer, nNumberOfBytesToWrite);
-				//dh.pos = writtenSize;
-				if (lpOverlapped->hEvent)
-					SetEvent(lpOverlapped->hEvent);
-			}
-			else
-			{
-				SCOPED_WRITE_LOCK(mf->lock, lock);
-
-				__try
-				{
-					mf->Write(dh, lpBuffer, nNumberOfBytesToWrite, fi.name);
-				}
-				__except(WriteFileAccessViolation(GetExceptionInformation()))
-				{
-				}
-			}
-
-			*lpNumberOfBytesWritten = nNumberOfBytesToWrite;
-			SetLastError(ERROR_SUCCESS);
-			DEBUG_LOG_DETOURED(L"WriteFile", L"(MEMORY)%s %llu (%ls) ToWrite: %u -> Success", (lpOverlapped ? L" OVERLAPPED" : L""), uintptr_t(hFile), HandleToName(hFile), nNumberOfBytesToWrite);
-			return TRUE;
-		}
-		UBA_ASSERTF(!fi.isFileMap, L"Trying to write to file %ls which is a filemap. This is not supported (%s)", HandleToName(hFile), fi.originalName);
-		UBA_ASSERTF(dh.trueHandle != INVALID_HANDLE_VALUE, L"Trying to write to file %ls which does not have a valid handle (%s)", HandleToName(hFile), fi.originalName);
-		trueHandle = dh.trueHandle;
-	}
-	else if (hFile == PseudoHandle)
-	{
-		DEBUG_LOG_DETOURED(L"WriteFile", L"(PseudoHandle) -> Success");
-		SetLastError(ERROR_SUCCESS);
-		return true;
-	}
-	else if (hFile == g_stdHandle[1] || hFile == g_stdHandle[0])
-	{
-		//DEBUG_LOG_DETOURED(L"WriteStdFile2", L"%llu", uintptr_t(hFile));
-		WriteStdFile(lpBuffer, nNumberOfBytesToWrite, hFile == g_stdHandle[0]);
-		*lpNumberOfBytesWritten = nNumberOfBytesToWrite;
-		SetLastError(ERROR_SUCCESS);
-		return true;
-		//if (GetFileType(trueHandle) != FILE_TYPE_CHAR )
-		//return True_WriteFile(trueHandle, lpBuffer, nNumberOfBytesToWrite, lpNumberOfBytesWritten, lpOverlapped);
-	}
-
-	TimerScope ts(g_kernelStats.writeFile);
-	BOOL res = True_WriteFile(trueHandle, lpBuffer, nNumberOfBytesToWrite, lpNumberOfBytesWritten, lpOverlapped);
-
-	g_kernelStats.writeFile.bytes += *lpNumberOfBytesWritten;
-	DEBUG_LOG_TRUE(L"WriteFile", L"%llu (%ls) -> %ls", uintptr_t(hFile), HandleToName(hFile), ToString(res));
-	return res;
-}
-
-BOOL Detoured_WriteFileEx(HANDLE hFile, LPCVOID lpBuffer, DWORD nNumberOfBytesToWrite, LPOVERLAPPED lpOverlapped, LPOVERLAPPED_COMPLETION_ROUTINE lpCompletionRoutine)
-{
-	DETOURED_CALL(WriteFileEx);
-	//DEBUG_LOG_TRUE(L"WriteFileEx", L"%llu (%ls)", uintptr_t(hFile), HandleToName(hFile));
-	//UBA_ASSERT(!isDetouredHandle(hFile));
-	UBA_ASSERT(isDetouredHandle(hFile));
-	DetouredHandle& h = asDetouredHandle(hFile);
-	UBA_ASSERT(h.trueHandle != INVALID_HANDLE_VALUE);
-
-	TimerScope ts(g_kernelStats.writeFile);
-	return True_WriteFileEx(h.trueHandle, lpBuffer, nNumberOfBytesToWrite, lpOverlapped, lpCompletionRoutine);
-}
-
 BOOL Detoured_DeviceIoControl(HANDLE hDevice, DWORD dwIoControlCode, LPVOID lpInBuffer, DWORD nInBufferSize, LPVOID lpOutBuffer, DWORD  nOutBufferSize, LPDWORD lpBytesReturned, LPOVERLAPPED lpOverlapped)
 {
 	HANDLE trueHandle = hDevice;
@@ -4096,6 +3961,22 @@ BOOL Detoured_ReadFileScatter(HANDLE hFile, FILE_SEGMENT_ELEMENT* aSegmentArray,
 	return True_ReadFileScatter(hFile, aSegmentArray, nNumberOfBytesToRead, lpReserved, lpOverlapped);
 }
 
+BOOL Detoured_WriteFile(HANDLE hFile, LPCVOID lpBuffer, DWORD nNumberOfBytesToWrite, LPDWORD lpNumberOfBytesWritten, LPOVERLAPPED lpOverlapped)
+{
+	t_ntWriteFileCalled = false;
+	BOOL res = True_WriteFile(hFile, lpBuffer, nNumberOfBytesToWrite, lpNumberOfBytesWritten, lpOverlapped);
+	UBA_ASSERT(t_ntWriteFileCalled);
+	return res;
+}
+
+BOOL Detoured_WriteFileEx(HANDLE hFile, LPCVOID lpBuffer, DWORD nNumberOfBytesToWrite, LPOVERLAPPED lpOverlapped, LPOVERLAPPED_COMPLETION_ROUTINE lpCompletionRoutine)
+{
+	t_ntWriteFileCalled = false;
+	BOOL res = True_WriteFileEx(hFile, lpBuffer, nNumberOfBytesToWrite, lpOverlapped, lpCompletionRoutine);
+	UBA_ASSERT(t_ntWriteFileCalled);
+	return res;
+}
+
 #if !defined(_M_ARM64)
 void Detoured_SetLastError(DWORD dwErrCode)
 {
diff --git a/Engine/Source/Programs/UnrealBuildAccelerator/Detours/Private/Windows/UbaDetoursFunctionsNtDll.inl b/Engine/Source/Programs/UnrealBuildAccelerator/Detours/Private/Windows/UbaDetoursFunctionsNtDll.inl
--- a/Engine/Source/Programs/UnrealBuildAccelerator/Detours/Private/Windows/UbaDetoursFunctionsNtDll.inl
+++ b/Engine/Source/Programs/UnrealBuildAccelerator/Detours/Private/Windows/UbaDetoursFunctionsNtDll.inl
@@ -1735,23 +1735,125 @@ NTSTATUS NTAPI Detoured_NtReadFile(HANDLE FileHandle, HANDLE Event, PIO_APC_ROUT
 	return status;
 }
 
+#if DETOURED_INCLUDE_DEBUG
+thread_local bool t_ntWriteFileCalled;
+#endif
+
+StringBuffer<32 * 1024> g_stdFile;
+ReaderWriterLock& g_stdFileLock = *new ReaderWriterLock();
+
+void WriteStdFile(LPCVOID lpBuffer, DWORD nNumberOfBytesToWrite, bool isError)
+{
+	if (!g_conEnabled[isError?0:1] || g_suppressLogging)
+		return;
+
+	SCOPED_WRITE_LOCK(g_stdFileLock, lock);
+	u32 start = 0;
+	u32 i = 0;
+	auto bufferStr = (const char*)lpBuffer;
+	while (i != nNumberOfBytesToWrite)
+	{
+		if (bufferStr[i] == '\n')
+		{
+			int len = i - start;
+			if (len > 0 && bufferStr[i - 1] == '\r')
+				--len;
+			if (len)
+				g_stdFile.Appendf(L"%.*hs", len, bufferStr + start);
+			Rpc_WriteLog(g_stdFile.data, g_stdFile.count, false, isError);
+			g_stdFile.Clear();
+			start = i + 1;
+		}
+		++i;
+	}
+	if (u32 left = nNumberOfBytesToWrite - start)
+		g_stdFile.Appendf(L"%.*hs", left, bufferStr + start);
+}
+
 NTSTATUS NTAPI Detoured_NtWriteFile(HANDLE FileHandle, HANDLE Event, PIO_APC_ROUTINE ApcRoutine, PVOID ApcContext, PIO_STATUS_BLOCK IoStatusBlock, PVOID Buffer, ULONG Length, PLARGE_INTEGER ByteOffset, PULONG Key)
 {
-	DETOURED_CALL(NtReadFile);
+	DETOURED_CALL(NtWriteFile);
+	
+	#if DETOURED_INCLUDE_DEBUG
+	t_ntWriteFileCalled = true;
+	#endif
+
+	auto trueRun = [&](HANDLE fileHandle, const tchar* prefix)
+		{
+			(void)prefix;
+			TimerScope ts(g_kernelStats.writeFile);
+			NTSTATUS status = True_NtWriteFile(fileHandle, Event, ApcRoutine, ApcContext, IoStatusBlock, Buffer, Length, ByteOffset, Key);
+			DEBUG_LOG_TRUE(L"NtWriteFile", L"%s%llu (%ls) -> %s", prefix, uintptr_t(fileHandle), HandleToName(fileHandle), ToString(status).data);
+			g_kernelStats.writeFile.bytes += IoStatusBlock->Information;
+			return status;
+		};
+
 	UBA_ASSERT(!isListDirectoryHandle(FileHandle));
+
+	if (FileHandle == PseudoHandle)
+		return trueRun(FileHandle, TC("(PseudoHandle)"));
+
+	auto successRun = [&]()
+		{
+			IoStatusBlock->Information = Length;
+			IoStatusBlock->Status = STATUS_SUCCESS;
+			if (Event)
+				SetEvent(Event);
+			if (ApcRoutine != nullptr)
+				ApcRoutine(ApcContext, IoStatusBlock, 0);
+			return STATUS_SUCCESS;
+		};
+
+	if (FileHandle == g_stdHandle[1] || FileHandle == g_stdHandle[0])
+	{
+		//DEBUG_LOG_DETOURED(L"WriteStdFile2", L"%llu", uintptr_t(hFile));
+		WriteStdFile(Buffer, Length, FileHandle == g_stdHandle[0]);
+		return successRun();
+	}
+
 	if (!isDetouredHandle(FileHandle))
-		return True_NtWriteFile(FileHandle, Event, ApcRoutine, ApcContext, IoStatusBlock, Buffer, Length, ByteOffset, Key);
+		return trueRun(FileHandle, TC(""));
+
 	auto& dh = asDetouredHandle(FileHandle);
-	UBA_ASSERT(dh.type == HandleType_File);
+
+	if (dh.type >= HandleType_StdErr)
+	{
+		if (dh.type != HandleType_StdIn)
+		{
+			//DEBUG_LOG_DETOURED(L"WriteStdFile1", L"%llu", uintptr_t(hFile));
+			WriteStdFile(Buffer, Length, dh.type == HandleType_StdErr);
+		}
+		return successRun();
+	}
+
 	auto& fo = *dh.fileObject;
-	fo.wasUsed = true;
-	FileInfo& fi = *fo.fileInfo;
-	UBA_ASSERT(!fi.isFileMap && !fi.memoryFile);
-	HANDLE trueHandle = dh.trueHandle;
-	UBA_ASSERT(trueHandle != INVALID_HANDLE_VALUE);
-	NTSTATUS res = True_NtWriteFile(trueHandle, Event, ApcRoutine, ApcContext, IoStatusBlock, Buffer, Length, ByteOffset, Key);
-	DEBUG_LOG_TRUE(L"NtWriteFile", L"%llu (%ls) -> %s", uintptr_t(FileHandle), HandleToName(FileHandle), ToString(res).data);
-	return res;
+	auto& fi = *fo.fileInfo;
+
+	if (!fi.memoryFile)
+	{
+		UBA_ASSERTF(!fi.isFileMap, L"Trying to write to file %ls which is a filemap. This is not supported (%s)", HandleToName(FileHandle), fi.originalName);
+		UBA_ASSERTF(dh.trueHandle != INVALID_HANDLE_VALUE, L"Trying to write to file %ls which does not have a valid handle (%s)", HandleToName(FileHandle), fi.originalName);
+		return trueRun(dh.trueHandle, TC(""));
+	}
+	MemoryFile& mf = *fi.memoryFile;
+
+	u64 offset = ByteOffset ? ByteOffset->QuadPart : dh.pos;
+	u64 offsetAfterWrite = offset + Length;
+	dh.pos = offsetAfterWrite;
+
+	SCOPED_WRITE_LOCK(mf.lock, lock);
+	u64 newSize = Max(mf.writtenSize, offsetAfterWrite);
+	if (mf.isThrowAway)
+	{
+		mf.writtenSize = newSize;
+		return successRun();
+	}
+	mf.EnsureCommitted(dh, newSize);
+	mf.writtenSize = newSize;
+	mf.isReported = false;
+	lock.Leave();
+	memcpy(mf.baseAddress + offset, Buffer, Length);
+	return successRun();
 }
 
 NTSTATUS NTAPI Detoured_NtFsControlFile(HANDLE FileHandle, HANDLE Event, PIO_APC_ROUTINE ApcRoutine, PVOID ApcContext, PIO_STATUS_BLOCK IoStatusBlock, ULONG FsControlCode, PVOID InputBuffer, ULONG InputBufferLength, PVOID OutputBuffer, ULONG OutputBufferLength)
