[UBA]
* Fixed GetFileAttributes edge case when file was written to output remotely but not using memory files
* Converted logging

diff --git a/Engine/Source/Programs/UnrealBuildAccelerator/Detours/Private/UbaDetoursShared.cpp b/Engine/Source/Programs/UnrealBuildAccelerator/Detours/Private/UbaDetoursShared.cpp
--- a/Engine/Source/Programs/UnrealBuildAccelerator/Detours/Private/UbaDetoursShared.cpp
+++ b/Engine/Source/Programs/UnrealBuildAccelerator/Detours/Private/UbaDetoursShared.cpp
@@ -506,6 +506,8 @@ namespace uba
 
 	const tchar* Shared_GetFileAttributes(FileAttributes& outAttr, StringView fileName, bool checkIfDir)
 	{
+		memset(&outAttr, 0, sizeof(FileAttributes));
+
 		StringBuffer<MaxPath> fileNameForKey;
 		fileNameForKey.Append(fileName);
 		if (CaseInsensitiveFs)
diff --git a/Engine/Source/Programs/UnrealBuildAccelerator/Detours/Private/Windows/UbaDetoursFunctionsKernelBase.inl b/Engine/Source/Programs/UnrealBuildAccelerator/Detours/Private/Windows/UbaDetoursFunctionsKernelBase.inl
--- a/Engine/Source/Programs/UnrealBuildAccelerator/Detours/Private/Windows/UbaDetoursFunctionsKernelBase.inl
+++ b/Engine/Source/Programs/UnrealBuildAccelerator/Detours/Private/Windows/UbaDetoursFunctionsKernelBase.inl
@@ -2177,42 +2177,61 @@ BOOL Detoured_GetFileInformationByHandle(HANDLE hFile, LPBY_HANDLE_FILE_INFORMAT
 					return false;
 				}
 
-				UBA_ASSERTF(attr.fileIndex, TC("No file index set for %s"), HandleToName(hFile));
-
-				LARGE_INTEGER li = ToLargeInteger(attr.fileIndex);
-				/*
-				#if UBA_DEBUG_VALIDATE
-				if (g_validateFileAccess)
+				if (attr.fileIndex || dh.trueHandle == INVALID_HANDLE_VALUE)
 				{
-					HANDLE h = True_CreateFileW(fileName, 0, FILE_SHARE_READ, NULL, OPEN_ALWAYS, FILE_FLAG_BACKUP_SEMANTICS, NULL);
-					auto res = True_GetFileInformationByHandle(h, lpFileInformation);
-					True_CloseHandle(h);
-
-					UBA_ASSERT(attr.data.dwFileAttributes == lpFileInformation->dwFileAttributes);
-					UBA_ASSERT(attr.volumeSerial == lpFileInformation->dwVolumeSerialNumber);
-					UBA_ASSERT(li.HighPart == lpFileInformation->nFileIndexHigh);
-					UBA_ASSERT(li.LowPart == lpFileInformation->nFileIndexLow);
-					//return res;
-				}
-				#endif
-				*/
-				SetLastError(ERROR_SUCCESS);
+					UBA_ASSERTF(attr.fileIndex, TC("No file index set for %s"), HandleToName(hFile));
 
-				UBA_ASSERTF(attr.volumeSerial, TC("No volume serial set for %s"), HandleToName(hFile));
-
-				lpFileInformation->dwFileAttributes = attr.data.dwFileAttributes;
-				lpFileInformation->ftCreationTime = attr.data.ftCreationTime;
-				lpFileInformation->ftCreationTime = attr.data.ftCreationTime;
-				lpFileInformation->ftLastAccessTime = attr.data.ftLastAccessTime;
-				lpFileInformation->ftLastWriteTime = attr.data.ftLastWriteTime;
-				lpFileInformation->dwVolumeSerialNumber = attr.volumeSerial;
-				lpFileInformation->nFileIndexHigh = li.HighPart;
-				lpFileInformation->nFileIndexLow = li.LowPart;
-				lpFileInformation->nNumberOfLinks = 1;//~u32(0); // TODO
-				lpFileInformation->nFileSizeHigh = attr.data.nFileSizeHigh;
-				lpFileInformation->nFileSizeLow = attr.data.nFileSizeLow;
-				DEBUG_LOG_DETOURED(L"GetFileInformationByHandle", L"(remote) %llu (%ls) -> Success", uintptr_t(hFile), HandleToName(hFile));
-				return TRUE;
+					// TODO: This scenario can happen in very complicated python scripts.. might need more investigation
+					LARGE_INTEGER test;
+					test.LowPart = attr.data.nFileSizeLow;
+					test.HighPart = attr.data.nFileSizeHigh;
+					if (test.QuadPart == InvalidValue)
+					{
+						//auto findIt = g_mappedFileTable.m_lookup.find(fi.fileNameKey);
+						//UBA_ASSERT(findIt != g_mappedFileTable.m_lookup.end());
+
+						UBA_ASSERT(!fi.isFileMap && dh.trueHandle != INVALID_HANDLE_VALUE);
+						TimerScope ts(g_kernelStats.getFileInfo);
+						auto res = True_GetFileInformationByHandle(dh.trueHandle, lpFileInformation); // Calls NtQueryInformationFile
+						DEBUG_LOG_TRUE(L"GetFileInformationByHandle", L"(SPECIAL) %llu (%ls) -> %u (Size: %llu)", uintptr_t(hFile), HandleToName(hFile), res, ToLargeInteger(lpFileInformation->nFileSizeHigh, lpFileInformation->nFileSizeLow).QuadPart);
+						return res;
+					}
+
+					LARGE_INTEGER li = ToLargeInteger(attr.fileIndex);
+					/*
+					#if UBA_DEBUG_VALIDATE
+					if (g_validateFileAccess)
+					{
+						HANDLE h = True_CreateFileW(fileName, 0, FILE_SHARE_READ, NULL, OPEN_ALWAYS, FILE_FLAG_BACKUP_SEMANTICS, NULL);
+						auto res = True_GetFileInformationByHandle(h, lpFileInformation);
+						True_CloseHandle(h);
+
+						UBA_ASSERT(attr.data.dwFileAttributes == lpFileInformation->dwFileAttributes);
+						UBA_ASSERT(attr.volumeSerial == lpFileInformation->dwVolumeSerialNumber);
+						UBA_ASSERT(li.HighPart == lpFileInformation->nFileIndexHigh);
+						UBA_ASSERT(li.LowPart == lpFileInformation->nFileIndexLow);
+						//return res;
+					}
+					#endif
+					*/
+					SetLastError(ERROR_SUCCESS);
+
+					UBA_ASSERTF(attr.volumeSerial, TC("No volume serial set for %s"), HandleToName(hFile));
+
+					lpFileInformation->dwFileAttributes = attr.data.dwFileAttributes;
+					lpFileInformation->ftCreationTime = attr.data.ftCreationTime;
+					lpFileInformation->ftCreationTime = attr.data.ftCreationTime;
+					lpFileInformation->ftLastAccessTime = attr.data.ftLastAccessTime;
+					lpFileInformation->ftLastWriteTime = attr.data.ftLastWriteTime;
+					lpFileInformation->dwVolumeSerialNumber = attr.volumeSerial;
+					lpFileInformation->nFileIndexHigh = li.HighPart;
+					lpFileInformation->nFileIndexLow = li.LowPart;
+					lpFileInformation->nNumberOfLinks = 1;//~u32(0); // TODO
+					lpFileInformation->nFileSizeHigh = attr.data.nFileSizeHigh;
+					lpFileInformation->nFileSizeLow = attr.data.nFileSizeLow;
+					DEBUG_LOG_DETOURED(L"GetFileInformationByHandle", L"(remote) %llu (%ls) -> Success", uintptr_t(hFile), HandleToName(hFile));
+					return TRUE;
+				}
 			}
 		}
 		UBA_ASSERTF(dh.trueHandle != INVALID_HANDLE_VALUE, TC("GetFileInformationByHandle needs true handle for %ls"), HandleToName(hFile));
diff --git a/Engine/Source/Programs/UnrealBuildAccelerator/Detours/Private/Windows/UbaDetoursFunctionsWin.cpp b/Engine/Source/Programs/UnrealBuildAccelerator/Detours/Private/Windows/UbaDetoursFunctionsWin.cpp
--- a/Engine/Source/Programs/UnrealBuildAccelerator/Detours/Private/Windows/UbaDetoursFunctionsWin.cpp
+++ b/Engine/Source/Programs/UnrealBuildAccelerator/Detours/Private/Windows/UbaDetoursFunctionsWin.cpp
@@ -88,17 +88,15 @@ bool g_debugFileFlushOnWrite = false;
 void WriteDebug(const void* data, u32 dataLen)
 {
 	auto readMem = (const u8*)data;
-	DWORD toWrite = (DWORD)dataLen;
+	ULONG toWrite = (ULONG)dataLen;
 	DWORD lastError = GetLastError();
-	while (true)
+	while (toWrite)
 	{
-		DWORD written;
-		if (True_WriteFile((HANDLE)g_debugFile, readMem, toWrite, &written, NULL))
+		IO_STATUS_BLOCK ioStatusBlock;
+		if (True_NtWriteFile((HANDLE)g_debugFile, NULL, NULL, NULL, &ioStatusBlock, (void*)readMem, toWrite, NULL, NULL) != STATUS_SUCCESS)
 			break;
-		if (GetLastError() != ERROR_IO_PENDING)
-			break;// ExitProcess(1340); // During shutdown this might actually cause error and we just ignore that and break out
-		readMem += written;
-		toWrite -= written;
+		readMem += ioStatusBlock.Information;
+		toWrite -= (ULONG)ioStatusBlock.Information;
 	}
 	if (g_debugFileFlushOnWrite)
 		True_FlushFileBuffers((HANDLE)g_debugFile);
