[UBA]
* Removed detouring implementation of GetFileSize/GetFileSizeEx since NtQueryInformationFile is implemented instead

diff --git a/Engine/Source/Programs/UnrealBuildAccelerator/Detours/Private/Windows/UbaDetoursFunctionsKernelBase.inl b/Engine/Source/Programs/UnrealBuildAccelerator/Detours/Private/Windows/UbaDetoursFunctionsKernelBase.inl
--- a/Engine/Source/Programs/UnrealBuildAccelerator/Detours/Private/Windows/UbaDetoursFunctionsKernelBase.inl
+++ b/Engine/Source/Programs/UnrealBuildAccelerator/Detours/Private/Windows/UbaDetoursFunctionsKernelBase.inl
@@ -720,86 +720,6 @@ BOOL Detoured_FlushFileBuffers(HANDLE hFile)
 	return res;
 }
 
-DWORD Detoured_GetFileSize(HANDLE hFile, LPDWORD lpFileSizeHigh)
-{
-	DETOURED_CALL(GetFileSize);
-
-	HANDLE trueHandle = hFile;
-	if (isDetouredHandle(hFile))
-	{
-		auto& dh = asDetouredHandle(hFile);
-		FileInfo& fi = *dh.fileObject->fileInfo;
-		if (fi.size != InvalidValue)
-		{
-			LARGE_INTEGER li = ToLargeInteger(fi.size);
-			if (lpFileSizeHigh)
-				*lpFileSizeHigh = li.HighPart;
-			DEBUG_LOG_DETOURED(L"GetFileSize", L"%llu (%ls) -> %u", uintptr_t(hFile), HandleToName(hFile), li.LowPart);
-			SetLastError(ERROR_SUCCESS);
-			return li.LowPart;
-		}
-		if (fi.memoryFile)
-		{
-			LARGE_INTEGER li = ToLargeInteger(fi.memoryFile->writtenSize);
-			if (lpFileSizeHigh)
-				*lpFileSizeHigh = li.HighPart;
-			DEBUG_LOG_DETOURED(L"GetFileSize", L"%llu (%ls) -> %u", uintptr_t(hFile), HandleToName(hFile), li.LowPart);
-			SetLastError(ERROR_SUCCESS);
-			return li.LowPart;
-		}
-		trueHandle = dh.trueHandle;
-		UBA_ASSERT(trueHandle != INVALID_HANDLE_VALUE);
-	}
-
-	DEBUG_LOG_TRUE(L"GetFileSize", L"%llu (%ls)", uintptr_t(hFile), HandleToName(hFile));
-	return True_GetFileSize(trueHandle, lpFileSizeHigh); // Calls NtQueryInformationFile
-}
-
-DWORD Detoured_GetFileSizeEx(HANDLE hFile, PLARGE_INTEGER lpFileSize)
-{
-	DETOURED_CALL(GetFileSizeEx);
-	HANDLE trueHandle = hFile;
-	if (isDetouredHandle(hFile))
-	{
-		DetouredHandle& dh = asDetouredHandle(hFile);
-		FileInfo& fi = *dh.fileObject->fileInfo;
-		if (fi.size != InvalidValue)
-		{
-			*lpFileSize = ToLargeInteger(fi.size);
-			SetLastError(ERROR_SUCCESS);
-			DEBUG_LOG_DETOURED(L"GetFileSizeEx", L"%llu (%ls) (Size:%llu) -> 1", uintptr_t(hFile), HandleToName(hFile), fi.size);
-			return 1;
-		}
-		if (fi.memoryFile)
-		{
-			*lpFileSize = ToLargeInteger(fi.memoryFile->writtenSize);
-			SetLastError(ERROR_SUCCESS);
-			DEBUG_LOG_DETOURED(L"GetFileSizeEx", L"%llu (%ls) (Size:%llu) -> 1", uintptr_t(hFile), HandleToName(hFile), fi.memoryFile->writtenSize);
-			return 1;
-		}
-
-		trueHandle = dh.trueHandle;
-		
-		if (trueHandle == INVALID_HANDLE_VALUE)
-		{
-			u32 entryOffset = dh.dirTableOffset;
-			if (entryOffset != ~u32(0))
-			{
-				DirectoryTable::EntryInformation entryInfo;
-				g_directoryTable.GetEntryInformation(entryInfo, entryOffset);
-				*lpFileSize = ToLargeInteger(entryInfo.size);
-				SetLastError(ERROR_SUCCESS);
-				DEBUG_LOG_DETOURED(L"GetFileSizeEx", L"%llu (%ls) (Size:%llu) -> 1", uintptr_t(hFile), HandleToName(hFile), entryInfo.size);
-				return 1;
-			}
-			UBA_ASSERTF(false, L"GetFileSizeEx (%s)", HandleToName(hFile));
-		}
-	}
-	DWORD res = True_GetFileSizeEx(trueHandle, lpFileSize); // This ends up in Detoured_NtQueryInformationFile
-	DEBUG_LOG_TRUE(L"GetFileSizeEx", L"%llu (%ls) -> %u (%llu)", uintptr_t(hFile), HandleToName(hFile), res, lpFileSize->QuadPart);
-	return res;
-}
-
 BOOL Detoured_SetFileTime(HANDLE hFile, const FILETIME* lpCreationTime, const FILETIME* lpLastAccessTime, const FILETIME* lpLastWriteTime)
 {
 	DETOURED_CALL(SetFileTime);
@@ -3807,6 +3727,20 @@ BOOL Detoured_GetVolumePathNameA(LPCSTR lpszFileName, LPSTR lpszVolumePathName,
 	return True_GetVolumePathNameA(lpszFileName, lpszVolumePathName, cchBufferLength);
 }
 
+DWORD Detoured_GetFileSize(HANDLE hFile, LPDWORD lpFileSizeHigh)
+{
+	DETOURED_CALL(GetFileSize);
+	DEBUG_LOG_TRUE(L"GetFileSize", L"%llu (%ls)", uintptr_t(hFile), HandleToName(hFile));
+	return True_GetFileSize(hFile, lpFileSizeHigh); // Calls NtQueryInformationFile
+}
+
+DWORD Detoured_GetFileSizeEx(HANDLE hFile, PLARGE_INTEGER lpFileSize)
+{
+	DETOURED_CALL(GetFileSizeEx);
+	DEBUG_LOG_TRUE(L"GetFileSizeEx", L"%llu (%ls)", uintptr_t(hFile), HandleToName(hFile));
+	return True_GetFileSizeEx(hFile, lpFileSize); // Calls NtQueryInformationFile
+}
+
 DWORD Detoured_SetFilePointer(HANDLE hFile, LONG lDistanceToMove, PLONG lpDistanceToMoveHigh, DWORD dwMoveMethod)
 {
 	DETOURED_CALL(SetFilePointer);
