[UBA]
* Missed one file in change that removed GetFileSize

diff --git a/Engine/Source/Programs/UnrealBuildAccelerator/Core/Public/UbaWinBinDependencyParser.h b/Engine/Source/Programs/UnrealBuildAccelerator/Core/Public/UbaWinBinDependencyParser.h
--- a/Engine/Source/Programs/UnrealBuildAccelerator/Core/Public/UbaWinBinDependencyParser.h
+++ b/Engine/Source/Programs/UnrealBuildAccelerator/Core/Public/UbaWinBinDependencyParser.h
@@ -11,7 +11,6 @@
 #define True_CreateFileMappingW ::CreateFileMappingW
 #define True_MapViewOfFile ::MapViewOfFile
 #define True_UnmapViewOfFile ::UnmapViewOfFile
-#define True_GetFileSizeEx ::GetFileSizeEx
 #endif
 
 namespace uba
@@ -162,7 +161,7 @@ namespace uba
 			return false;
 		auto unmap = MakeGuard([&]() { True_UnmapViewOfFile(mem); });
 		LARGE_INTEGER fileSize;
-		if (!True_GetFileSizeEx(fileHandle, &fileSize))
+		if (!::GetFileSizeEx(fileHandle, &fileSize))
 			return false;
 
 		return FindImportsInMem(filePath.data, mem, fileSize.QuadPart, func);
