[UBA]
* Cleaned up Detoured_NtQueryInformationFile and removed Detoured_GetFileTime

diff --git a/Engine/Source/Programs/UnrealBuildAccelerator/Detours/Private/Windows/UbaDetoursFunctionsKernelBase.inl b/Engine/Source/Programs/UnrealBuildAccelerator/Detours/Private/Windows/UbaDetoursFunctionsKernelBase.inl
--- a/Engine/Source/Programs/UnrealBuildAccelerator/Detours/Private/Windows/UbaDetoursFunctionsKernelBase.inl
+++ b/Engine/Source/Programs/UnrealBuildAccelerator/Detours/Private/Windows/UbaDetoursFunctionsKernelBase.inl
@@ -600,19 +600,6 @@ BOOL Detoured_GetQueuedCompletionStatusEx(HANDLE CompletionPort, LPOVERLAPPED_EN
 	return res;
 }
 
-DWORD Detoured_GetSecurityInfo(HANDLE handle, SE_OBJECT_TYPE ObjectType, SECURITY_INFORMATION SecurityInfo, PSID* ppsidOwner, PSID* ppsidGroup, PACL* ppDacl, PACL* ppSacl, PSECURITY_DESCRIPTOR* ppSecurityDescriptor)
-{
-	DETOURED_CALL(GetSecurityInfo);
-	if (isDetouredHandle(handle))
-	{
-		handle = asDetouredHandle(handle).trueHandle;
-		UBA_ASSERTF(handle != INVALID_HANDLE_VALUE, L"GetSecurityInfo");
-	}
-
-	DEBUG_LOG_TRUE(L"GetSecurityInfo", L"");
-	return True_GetSecurityInfo(handle, ObjectType, SecurityInfo, ppsidOwner, ppsidGroup, ppDacl, ppSacl, ppSecurityDescriptor);
-}
-
 BOOL Detoured_PathSearchAndQualifyW(LPCWSTR pszPath, LPWSTR pszBuf, UINT cchBuf)
 {
 	FixPath2(pszPath, TStrlen(pszPath), g_virtualWorkingDir.data, g_virtualWorkingDir.count, pszBuf, cchBuf, nullptr);
@@ -739,35 +726,6 @@ BOOL Detoured_SetFileTime(HANDLE hFile, const FILETIME* lpCreationTime, const FI
 	return True_SetFileTime(trueHandle, lpCreationTime, lpLastAccessTime, lpLastWriteTime);
 }
 
-BOOL Detoured_GetFileTime(HANDLE hFile, LPFILETIME lpCreationTime, LPFILETIME lpLastAccessTime, LPFILETIME lpLastWriteTime)
-{
-	DETOURED_CALL(GetFileTime);
-	HANDLE trueHandle = hFile;
-	if (isDetouredHandle(hFile))
-	{
-		DetouredHandle& dh = asDetouredHandle(hFile);
-		u32 entryOffset = dh.dirTableOffset;
-		if (entryOffset != ~u32(0))
-		{
-			DirectoryTable::EntryInformation entryInfo;
-			g_directoryTable.GetEntryInformation(entryInfo, entryOffset);
-			if (lpLastWriteTime)
-				*(u64*)lpLastWriteTime = entryInfo.lastWrite;
-			if (lpCreationTime)
-				*(u64*)lpCreationTime = entryInfo.lastWrite;
-			if (lpLastAccessTime)
-				*(u64*)lpLastAccessTime = entryInfo.lastWrite;
-			//UBA_ASSERT(!lpLastAccessTime);
-			DEBUG_LOG_DETOURED(L"GetFileTime", L"%llu (%ls)", uintptr_t(hFile), HandleToName(hFile));
-			return TRUE;
-		}
-		trueHandle = asDetouredHandle(hFile).trueHandle;
-		UBA_ASSERT(trueHandle != INVALID_HANDLE_VALUE);
-	}
-	DEBUG_LOG_TRUE(L"GetFileTime", L"%llu (%ls)", uintptr_t(hFile), HandleToName(hFile));
-	return True_GetFileTime(trueHandle, lpCreationTime, lpLastAccessTime, lpLastWriteTime);
-}
-
 DWORD Detoured_GetFileType(HANDLE hFile)
 {
 	DETOURED_CALL(GetFileType);
@@ -3741,6 +3699,13 @@ DWORD Detoured_GetFileSizeEx(HANDLE hFile, PLARGE_INTEGER lpFileSize)
 	return True_GetFileSizeEx(hFile, lpFileSize); // Calls NtQueryInformationFile
 }
 
+BOOL Detoured_GetFileTime(HANDLE hFile, LPFILETIME lpCreationTime, LPFILETIME lpLastAccessTime, LPFILETIME lpLastWriteTime)
+{
+	DETOURED_CALL(GetFileTime);
+	DEBUG_LOG_TRUE(L"GetFileTime", L"%llu (%ls)", uintptr_t(hFile), HandleToName(hFile));
+	return True_GetFileTime(trueHandle, lpCreationTime, lpLastAccessTime, lpLastWriteTime);
+}
+
 DWORD Detoured_SetFilePointer(HANDLE hFile, LONG lDistanceToMove, PLONG lpDistanceToMoveHigh, DWORD dwMoveMethod)
 {
 	DETOURED_CALL(SetFilePointer);
@@ -4114,6 +4079,19 @@ BOOL Detoured_RemoveDirectoryA(LPCSTR lpPathName)
 	return res;
 }
 /*
+DWORD Detoured_GetSecurityInfo(HANDLE handle, SE_OBJECT_TYPE ObjectType, SECURITY_INFORMATION SecurityInfo, PSID* ppsidOwner, PSID* ppsidGroup, PACL* ppDacl, PACL* ppSacl, PSECURITY_DESCRIPTOR* ppSecurityDescriptor)
+{
+	DETOURED_CALL(GetSecurityInfo);
+	if (isDetouredHandle(handle))
+	{
+		handle = asDetouredHandle(handle).trueHandle;
+		UBA_ASSERTF(handle != INVALID_HANDLE_VALUE, L"GetSecurityInfo");
+	}
+
+	DEBUG_LOG_TRUE(L"GetSecurityInfo", L"");
+	return True_GetSecurityInfo(handle, ObjectType, SecurityInfo, ppsidOwner, ppsidGroup, ppDacl, ppSacl, ppSecurityDescriptor);
+}
+
 BOOL Detoured_DecryptFileW(LPCWSTR lpFileName, DWORD dwReserved)
 {
 	DETOURED_CALL(DecryptFileW);
diff --git a/Engine/Source/Programs/UnrealBuildAccelerator/Detours/Private/Windows/UbaDetoursFunctionsNtDll.inl b/Engine/Source/Programs/UnrealBuildAccelerator/Detours/Private/Windows/UbaDetoursFunctionsNtDll.inl
--- a/Engine/Source/Programs/UnrealBuildAccelerator/Detours/Private/Windows/UbaDetoursFunctionsNtDll.inl
+++ b/Engine/Source/Programs/UnrealBuildAccelerator/Detours/Private/Windows/UbaDetoursFunctionsNtDll.inl
@@ -135,6 +135,9 @@ NTSTATUS Detoured_NtQueryVolumeInformationFile(HANDLE FileHandle, PIO_STATUS_BLO
 
 NTSTATUS Detoured_NtQueryInformationFile(HANDLE FileHandle, PIO_STATUS_BLOCK IoStatusBlock, PVOID FileInformation, ULONG Length, FILE_INFORMATION_CLASS FileInformationClass)
 {
+	IoStatusBlock->Information = 0;
+	IoStatusBlock->Status = STATUS_SUCCESS;
+
 	DETOURED_CALL(NtQueryInformationFile);
 	if (isListDirectoryHandle(FileHandle))
 	{
@@ -191,160 +194,181 @@ NTSTATUS Detoured_NtQueryInformationFile(HANDLE FileHandle, PIO_STATUS_BLOCK IoS
 		}
 	}
 
-	HANDLE TrueHandle = FileHandle;
-	if (isDetouredHandle(FileHandle))
-	{
-		auto& dh = asDetouredHandle(FileHandle);
-		TrueHandle = dh.trueHandle;
-
-		if (TrueHandle == INVALID_HANDLE_VALUE)
+	HANDLE trueHandle = FileHandle;
+	auto runTrue = [&]()
 		{
-			if (FileInformationClass == FileStandardInformation)
-			{
-				UBA_ASSERT(dh.type == HandleType_File);
-				FileInfo& fi = *dh.fileObject->fileInfo;
-				UBA_ASSERT(fi.memoryFile || fi.isFileMap);
-				auto& info = *(FILE_STANDARD_INFORMATION*)FileInformation;
-				u64 size = 0;
-				u64 allocationSize = 0;
-				if (fi.memoryFile)
-				{
-					size = fi.memoryFile->writtenSize;
-					allocationSize = fi.memoryFile->committedSize;
-				}
-				else if (fi.fileMapMem)
-				{
-					size = fi.fileMapMemSize;
-					allocationSize = size;
-				}
-				else
-				{
-					size = fi.size;
-					allocationSize = size;
-				}
-				info.DeletePending = FALSE;
-				info.NumberOfLinks = 0;
-				info.Directory = FALSE;
-				info.AllocationSize.QuadPart = allocationSize;
-				info.EndOfFile.QuadPart = size;
-				return STATUS_SUCCESS;
-			}
-			if (FileInformationClass == FilePositionInformation)
-			{
-				if (dh.type >= HandleType_StdErr)
-				{
-					return STATUS_SUCCESS;
-				}
-				UBA_ASSERT(dh.type == HandleType_File);
-				FileInfo& fi = *dh.fileObject->fileInfo;(void)fi;
-				UBA_ASSERT(fi.memoryFile || fi.isFileMap);
-				auto& info = *(FILE_POSITION_INFORMATION*)FileInformation;
-				info.CurrentByteOffset.QuadPart = dh.pos;
-				return STATUS_SUCCESS;
-			}
-			if (FileInformationClass == FileBasicInformation)
-			{
-				auto& info = *(FILE_BASIC_INFORMATION*)FileInformation;
-				UBA_ASSERT(dh.dirTableOffset != ~0u);
-				DirectoryTable::EntryInformation entryInfo;
-				g_directoryTable.GetEntryInformation(entryInfo, dh.dirTableOffset);
-				UBA_ASSERT(entryInfo.attributes != 0);
-				info.CreationTime.QuadPart = entryInfo.lastWrite;
-				info.LastAccessTime.QuadPart = entryInfo.lastWrite;
-				info.LastWriteTime.QuadPart = entryInfo.lastWrite;
-				info.ChangeTime.QuadPart = entryInfo.lastWrite;
-				info.FileAttributes = entryInfo.attributes;
-				DEBUG_LOG_DETOURED(L"NtQueryInformationFile", L"(FileBasicInformation) %llu (%s) -> %ls", uintptr_t(FileHandle), HandleToName(FileHandle), ToString(STATUS_SUCCESS).data);
-				return STATUS_SUCCESS;
-			}
-			if (FileInformationClass == FileAllInformation)
-			{
-				UBA_ASSERT(dh.dirTableOffset != ~0u);
-				DirectoryTable::EntryInformation entryInfo;
-				g_directoryTable.GetEntryInformation(entryInfo, dh.dirTableOffset);
-				UBA_ASSERT(entryInfo.attributes != 0);
+			UBA_ASSERT(trueHandle != INVALID_HANDLE_VALUE);
+			TimerScope ts(g_kernelStats.getFileInfo);
+			auto res = True_NtQueryInformationFile(trueHandle, IoStatusBlock, FileInformation, Length, FileInformationClass);
+			DEBUG_LOG_TRUE(L"NtQueryInformationFile", L"(%u) %llu (%ls) -> %ls", FileInformationClass, uintptr_t(FileHandle), HandleToName(FileHandle), ToString(res).data);
+			return res;
+		};
 
-				// TODO This code path is here to handle nodejs queries.. Is not properly implemented and miss things
-				auto& info = *(FILE_ALL_INFORMATION*)FileInformation;
-				info.BasicInformation.CreationTime.QuadPart = entryInfo.lastWrite;
-				info.BasicInformation.LastAccessTime.QuadPart = entryInfo.lastWrite;
-				info.BasicInformation.LastWriteTime.QuadPart = entryInfo.lastWrite;
-				info.BasicInformation.ChangeTime.QuadPart = entryInfo.lastWrite;
-				info.BasicInformation.FileAttributes = entryInfo.attributes;
-				info.StandardInformation.AllocationSize.QuadPart = entryInfo.size;
-				info.StandardInformation.EndOfFile.QuadPart = entryInfo.size;
-				info.StandardInformation.NumberOfLinks = 0;
-				info.StandardInformation.DeletePending = false;
-				info.StandardInformation.Directory = false;
-				info.InternalInformation.IndexNumber.QuadPart = entryInfo.fileIndex;
-				DEBUG_LOG_DETOURED(L"NtQueryInformationFile", L"(FileAllInformation) %llu (%s) -> %ls", uintptr_t(FileHandle), HandleToName(FileHandle), ToString(STATUS_SUCCESS).data);
-				return STATUS_SUCCESS;
-			}
-			if (FileInformationClass == FileNetworkOpenInformation)
-			{
-				UBA_ASSERT(dh.dirTableOffset != ~0u);
-				DirectoryTable::EntryInformation entryInfo;
-				g_directoryTable.GetEntryInformation(entryInfo, dh.dirTableOffset);
-				UBA_ASSERT(entryInfo.attributes != 0);
+	if (!isDetouredHandle(FileHandle))
+		return runTrue();
 
-				// TODO This code path is here to handle nodejs queries.. Is not properly implemented and miss things
-				auto& info = *(FILE_NETWORK_OPEN_INFORMATION*)FileInformation;
-				info.CreationTime.QuadPart = entryInfo.lastWrite;
-				info.LastAccessTime.QuadPart = entryInfo.lastWrite;
-				info.LastWriteTime.QuadPart = entryInfo.lastWrite;
-				info.ChangeTime.QuadPart = entryInfo.lastWrite;
-				info.AllocationSize.QuadPart = entryInfo.size;
-				u64 fileSize = dh.fileObject->fileInfo->size;
-				if (fileSize == InvalidValue)
-					fileSize = entryInfo.size;
-				info.EndOfFile.QuadPart = fileSize;
-				info.FileAttributes = entryInfo.attributes;
-				DEBUG_LOG_DETOURED(L"NtQueryInformationFile", L"(FileNetworkOpenInformation) %llu (%s) -> %ls", uintptr_t(FileHandle), HandleToName(FileHandle), ToString(STATUS_SUCCESS).data);
-				return STATUS_SUCCESS;
-			}
-			if (FileInformationClass == FileNormalizedNameInformation)
-			{
-				const tchar* src = TC("");//dh.fileObject->fileInfo->originalName;
-				if (Length < sizeof(ULONG))
-					return STATUS_INFO_LENGTH_MISMATCH;
-				auto& info = *(FILE_NAME_INFORMATION*)FileInformation;
-				u64 destLen = Length - FIELD_OFFSET(FILE_NAME_INFORMATION, FileName);
-				u64 srcLen = TStrlen(src)*sizeof(tchar);
-				u64 toWrite = Min(srcLen, destLen);
-				info.FileNameLength = (ULONG)srcLen;
-				memcpy(info.FileName, src, toWrite);
-				NTSTATUS res = destLen >= srcLen ? STATUS_SUCCESS : STATUS_BUFFER_OVERFLOW;
-				IoStatusBlock->Information = FIELD_OFFSET(FILE_NAME_INFORMATION, FileName) + toWrite;
-				IoStatusBlock->Status = res;
-				DEBUG_LOG_DETOURED(L"NtQueryInformationFile", L"(FileNormalizedNameInformation) %llu (%s) -> %ls", uintptr_t(FileHandle), HandleToName(FileHandle), ToString(res).data);
-				return res;
-			}
-			if (FileInformationClass == FileNameInformation)
-			{
-				const wchar_t* fileName = dh.fileObject->fileInfo->originalName;
-				auto& info = *(FILE_NAME_INFORMATION*)FileInformation;
-				ULONG header = FIELD_OFFSET(FILE_NAME_INFORMATION, FileName);
-				ULONG nameBytes = (ULONG)(wcslen(fileName) * sizeof(WCHAR));
-				ULONG required = header + nameBytes;
-				if (Length < header)
-					return STATUS_INFO_LENGTH_MISMATCH;
-				ULONG copy = Min(nameBytes, Length - header);
-				info.FileNameLength = nameBytes;
-				if (copy)
-					memcpy(info.FileName, fileName, copy);
-				IoStatusBlock->Status = STATUS_SUCCESS;
-				IoStatusBlock->Information = Min(Length, required);
-				DEBUG_LOG_DETOURED(L"NtQueryInformationFile", L"(FileNameInformation) %llu (%s) -> %ls", uintptr_t(FileHandle), HandleToName(FileHandle), ToString(STATUS_SUCCESS).data);
-				return STATUS_SUCCESS;
-			}
-			UBA_ASSERTF(false, L"NtQueryInformationFile (%u) failed using detoured handle %ls (%ls)", FileInformationClass, dh.fileObject->fileInfo->name, dh.fileObject->fileInfo->originalName);
+	auto& dh = asDetouredHandle(FileHandle);
+	trueHandle = dh.trueHandle;
+
+	if (FileInformationClass == FileStandardInformation)
+	{
+		UBA_ASSERT(dh.type == HandleType_File);
+		FileInfo& fi = *dh.fileObject->fileInfo;
+		u64 size = 0;
+		u64 allocationSize = 0;
+		if (fi.memoryFile)
+		{
+			size = fi.memoryFile->writtenSize;
+			allocationSize = fi.memoryFile->committedSize;
 		}
+		else if (fi.fileMapMem)
+		{
+			size = fi.fileMapMemSize;
+			allocationSize = size;
+		}
+		else if (fi.size == InvalidValue)
+		{
+			return runTrue();
+		}
+		else
+		{
+			size = fi.size;
+			allocationSize = size;
+		}
+
+		auto& info = *(FILE_STANDARD_INFORMATION*)FileInformation;
+		info.DeletePending = FALSE;
+		info.NumberOfLinks = 0;
+		info.Directory = FALSE;
+		info.AllocationSize.QuadPart = allocationSize;
+		info.EndOfFile.QuadPart = size;
+		DEBUG_LOG_DETOURED(L"NtQueryInformationFile", L"(FileStandardInformation) %llu (%s) -> %ls", uintptr_t(FileHandle), HandleToName(FileHandle), ToString(STATUS_SUCCESS).data);
+		return STATUS_SUCCESS;
+	}
+	if (FileInformationClass == FilePositionInformation)
+	{
+		if (dh.type >= HandleType_StdErr)
+		{
+			if (dh.trueHandle != INVALID_HANDLE_VALUE)
+				return runTrue();
+			return STATUS_SUCCESS;
+		}
+		UBA_ASSERT(dh.type == HandleType_File);
+		FileInfo& fi = *dh.fileObject->fileInfo;
+		if (!fi.memoryFile && !fi.isFileMap)
+			return runTrue();
+		auto& info = *(FILE_POSITION_INFORMATION*)FileInformation;
+		info.CurrentByteOffset.QuadPart = dh.pos;
+		DEBUG_LOG_DETOURED(L"NtQueryInformationFile", L"(FilePositionInformation) %llu (%s) -> %ls", uintptr_t(FileHandle), HandleToName(FileHandle), ToString(STATUS_SUCCESS).data);
+		return STATUS_SUCCESS;
+	}
+	if (FileInformationClass == FileBasicInformation)
+	{
+		auto& info = *(FILE_BASIC_INFORMATION*)FileInformation;
+		if (dh.dirTableOffset == ~0u)
+			return runTrue();
+		DirectoryTable::EntryInformation entryInfo;
+		g_directoryTable.GetEntryInformation(entryInfo, dh.dirTableOffset);
+		UBA_ASSERT(entryInfo.attributes != 0);
+		info.CreationTime.QuadPart = entryInfo.lastWrite;
+		info.LastAccessTime.QuadPart = entryInfo.lastWrite;
+		info.LastWriteTime.QuadPart = entryInfo.lastWrite;
+		info.ChangeTime.QuadPart = entryInfo.lastWrite;
+		info.FileAttributes = entryInfo.attributes;
+		DEBUG_LOG_DETOURED(L"NtQueryInformationFile", L"(FileBasicInformation) %llu (%s) -> %ls", uintptr_t(FileHandle), HandleToName(FileHandle), ToString(STATUS_SUCCESS).data);
+		return STATUS_SUCCESS;
+	}
+	if (FileInformationClass == FileAllInformation)
+	{
+		if (dh.dirTableOffset == ~0u)
+			return runTrue();
+		DirectoryTable::EntryInformation entryInfo;
+		g_directoryTable.GetEntryInformation(entryInfo, dh.dirTableOffset);
+		UBA_ASSERT(entryInfo.attributes != 0);
+
+		// TODO This code path is here to handle nodejs queries.. Is not properly implemented and miss things
+		auto& info = *(FILE_ALL_INFORMATION*)FileInformation;
+		info.BasicInformation.CreationTime.QuadPart = entryInfo.lastWrite;
+		info.BasicInformation.LastAccessTime.QuadPart = entryInfo.lastWrite;
+		info.BasicInformation.LastWriteTime.QuadPart = entryInfo.lastWrite;
+		info.BasicInformation.ChangeTime.QuadPart = entryInfo.lastWrite;
+		info.BasicInformation.FileAttributes = entryInfo.attributes;
+		info.StandardInformation.AllocationSize.QuadPart = entryInfo.size;
+		info.StandardInformation.EndOfFile.QuadPart = entryInfo.size;
+		info.StandardInformation.NumberOfLinks = 0;
+		info.StandardInformation.DeletePending = false;
+		info.StandardInformation.Directory = false;
+		info.InternalInformation.IndexNumber.QuadPart = entryInfo.fileIndex;
+		//IoStatusBlock->Information = FIELD_OFFSET(FILE_ALL_INFORMATION, NameInformation.FileName) + name_len;
+
+		DEBUG_LOG_DETOURED(L"NtQueryInformationFile", L"NOT IMPLEMENTED (FileAllInformation) %llu (%s) -> %ls", uintptr_t(FileHandle), HandleToName(FileHandle), ToString(STATUS_SUCCESS).data);
+		return STATUS_SUCCESS;
+	}
+	if (FileInformationClass == FileNetworkOpenInformation)
+	{
+		if (dh.dirTableOffset == ~0u)
+			return runTrue();
+		DirectoryTable::EntryInformation entryInfo;
+		g_directoryTable.GetEntryInformation(entryInfo, dh.dirTableOffset);
+		UBA_ASSERT(entryInfo.attributes != 0);
+
+		// TODO This code path is here to handle nodejs queries.. Is not properly implemented and miss things
+		auto& info = *(FILE_NETWORK_OPEN_INFORMATION*)FileInformation;
+		info.CreationTime.QuadPart = entryInfo.lastWrite;
+		info.LastAccessTime.QuadPart = entryInfo.lastWrite;
+		info.LastWriteTime.QuadPart = entryInfo.lastWrite;
+		info.ChangeTime.QuadPart = entryInfo.lastWrite;
+		info.AllocationSize.QuadPart = entryInfo.size;
+		u64 fileSize = dh.fileObject->fileInfo->size;
+		if (fileSize == InvalidValue)
+			fileSize = entryInfo.size;
+		info.EndOfFile.QuadPart = fileSize;
+		info.FileAttributes = entryInfo.attributes;
+		DEBUG_LOG_DETOURED(L"NtQueryInformationFile", L"NOT IMPLEMENTED (FileNetworkOpenInformation) %llu (%s) -> %ls", uintptr_t(FileHandle), HandleToName(FileHandle), ToString(STATUS_SUCCESS).data);
+		return STATUS_SUCCESS;
+	}
+	if (FileInformationClass == FileNormalizedNameInformation)
+	{
+		if (dh.trueHandle != INVALID_HANDLE_VALUE)
+			return runTrue();
+		const tchar* src = TC("");//dh.fileObject->fileInfo->originalName;
+		if (Length < sizeof(ULONG))
+			return STATUS_INFO_LENGTH_MISMATCH;
+		auto& info = *(FILE_NAME_INFORMATION*)FileInformation;
+		u64 destLen = Length - FIELD_OFFSET(FILE_NAME_INFORMATION, FileName);
+		u64 srcLen = TStrlen(src)*sizeof(tchar);
+		u64 toWrite = Min(srcLen, destLen);
+		info.FileNameLength = (ULONG)srcLen;
+		memcpy(info.FileName, src, toWrite);
+		NTSTATUS res = destLen >= srcLen ? STATUS_SUCCESS : STATUS_BUFFER_OVERFLOW;
+		IoStatusBlock->Information = FIELD_OFFSET(FILE_NAME_INFORMATION, FileName) + toWrite;
+		IoStatusBlock->Status = res;
+		DEBUG_LOG_DETOURED(L"NtQueryInformationFile", L"(FileNormalizedNameInformation) %llu (%s) -> %ls", uintptr_t(FileHandle), HandleToName(FileHandle), ToString(res).data);
+		return res;
+	}
+	if (FileInformationClass == FileNameInformation)
+	{
+		if (dh.trueHandle != INVALID_HANDLE_VALUE)
+			return runTrue();
+		const wchar_t* fileName = dh.fileObject->fileInfo->originalName;
+		auto& info = *(FILE_NAME_INFORMATION*)FileInformation;
+		ULONG header = FIELD_OFFSET(FILE_NAME_INFORMATION, FileName);
+		ULONG nameBytes = (ULONG)(wcslen(fileName) * sizeof(WCHAR));
+		ULONG required = header + nameBytes;
+		if (Length < header)
+			return STATUS_INFO_LENGTH_MISMATCH;
+		ULONG copy = Min(nameBytes, Length - header);
+		info.FileNameLength = nameBytes;
+		if (copy)
+			memcpy(info.FileName, fileName, copy);
+		IoStatusBlock->Status = STATUS_SUCCESS;
+		IoStatusBlock->Information = Min(Length, required);
+		DEBUG_LOG_DETOURED(L"NtQueryInformationFile", L"(FileNameInformation) %llu (%s) -> %ls", uintptr_t(FileHandle), HandleToName(FileHandle), ToString(STATUS_SUCCESS).data);
+		return STATUS_SUCCESS;
 	}
 
-	TimerScope ts(g_kernelStats.getFileInfo);
-	auto res = True_NtQueryInformationFile(TrueHandle, IoStatusBlock, FileInformation, Length, FileInformationClass);
-	DEBUG_LOG_TRUE(L"NtQueryInformationFile", L"(%u) %llu (%ls) -> %ls", FileInformationClass, uintptr_t(FileHandle), HandleToName(FileHandle), ToString(res).data);
-	return res;
+	UBA_ASSERTF(false, L"NtQueryInformationFile (%u) failed using detoured handle %ls (%ls)", FileInformationClass, dh.fileObject->fileInfo->name, dh.fileObject->fileInfo->originalName);
+	return runTrue();
 }
 
 NTSTATUS NTAPI Detoured_NtQueryDirectoryFile(HANDLE FileHandle, HANDLE Event, PIO_APC_ROUTINE ApcRoutine, PVOID ApcContext, PIO_STATUS_BLOCK IoStatusBlock, PVOID FileInformation, ULONG Length,
