[UBA]
* Fixed compile error in debug configs

diff --git a/Engine/Source/Programs/UnrealBuildAccelerator/Detours/Private/Windows/UbaDetoursFunctionsKernelBase.inl b/Engine/Source/Programs/UnrealBuildAccelerator/Detours/Private/Windows/UbaDetoursFunctionsKernelBase.inl
--- a/Engine/Source/Programs/UnrealBuildAccelerator/Detours/Private/Windows/UbaDetoursFunctionsKernelBase.inl
+++ b/Engine/Source/Programs/UnrealBuildAccelerator/Detours/Private/Windows/UbaDetoursFunctionsKernelBase.inl
@@ -3703,7 +3703,7 @@ BOOL Detoured_GetFileTime(HANDLE hFile, LPFILETIME lpCreationTime, LPFILETIME lp
 {
 	DETOURED_CALL(GetFileTime);
 	DEBUG_LOG_TRUE(L"GetFileTime", L"%llu (%ls)", uintptr_t(hFile), HandleToName(hFile));
-	return True_GetFileTime(trueHandle, lpCreationTime, lpLastAccessTime, lpLastWriteTime);
+	return True_GetFileTime(hFile, lpCreationTime, lpLastAccessTime, lpLastWriteTime);
 }
 
 DWORD Detoured_SetFilePointer(HANDLE hFile, LONG lDistanceToMove, PLONG lpDistanceToMoveHigh, DWORD dwMoveMethod)
