[UBA]
* Cleaned up Detoured_NtSetInformationFile

diff --git a/Engine/Source/Programs/UnrealBuildAccelerator/Detours/Private/Windows/UbaDetoursFunctionsNtDll.inl b/Engine/Source/Programs/UnrealBuildAccelerator/Detours/Private/Windows/UbaDetoursFunctionsNtDll.inl
--- a/Engine/Source/Programs/UnrealBuildAccelerator/Detours/Private/Windows/UbaDetoursFunctionsNtDll.inl
+++ b/Engine/Source/Programs/UnrealBuildAccelerator/Detours/Private/Windows/UbaDetoursFunctionsNtDll.inl
@@ -602,143 +602,116 @@ NTSTATUS Detoured_NtSetInformationFile(HANDLE FileHandle, PIO_STATUS_BLOCK IoSta
 {
 	DETOURED_CALL(NtSetInformationFile);
 
-
-	u8 tempBuffer[sizeof(FILE_RENAME_INFORMATION) + 512];
+	IoStatusBlock->Information = 1;
+	IoStatusBlock->Status = STATUS_SUCCESS;
 
 	HANDLE trueHandle = FileHandle;
-	if (isDetouredHandle(FileHandle))
-	{
-		DetouredHandle& h = asDetouredHandle(FileHandle);
 
-		if (FileInformationClass == FilePositionInformation)
+	auto runTrue = [&]()
 		{
-			auto& info = *(FILE_POSITION_INFORMATION*)FileInformation;
-			if (h.type >= HandleType_StdErr)
-			{
-				info.CurrentByteOffset.QuadPart = 0;
-				return STATUS_SUCCESS;
-			}
-			UBA_ASSERT(h.type == HandleType_File);
-			if (h.type >= HandleType_StdErr)
-			{
-				UBA_ASSERT(false); // TODO: What should we do with this?
-				return STATUS_SUCCESS;
-			}
-
-			FileInfo& fi = *h.fileObject->fileInfo;
-			if (fi.memoryFile || fi.isFileMap)
-			{
-				h.pos = info.CurrentByteOffset.QuadPart;
-				return STATUS_SUCCESS;
-			}
-			trueHandle = h.trueHandle;
 			UBA_ASSERT(trueHandle != INVALID_HANDLE_VALUE);
-		}
-		if (FileInformationClass == FileEndOfFileInformation)
-		{
-			UBA_ASSERT(h.type == HandleType_File);
-			FileInfo& fi = *h.fileObject->fileInfo;
-			if (MemoryFile* mf = fi.memoryFile)
-			{
-				DEBUG_LOG_DETOURED(L"SetEndOfFile (MEMORY)", L"%llu NewSize: %llu (%ls) -> Success", uintptr_t(FileHandle), mf->writtenSize, HandleToName(FileHandle));
-				mf->writtenSize = h.pos;
-				mf->isReported = false;
-				mf->EnsureCommitted(h, mf->writtenSize);
+			TimerScope ts(g_kernelStats.setFileInfo);
+			auto res = True_NtSetInformationFile(trueHandle, IoStatusBlock, FileInformation, Length, FileInformationClass);
+			DEBUG_LOG_TRUE(L"NtSetInformationFile", L"(%u) %llu (%ls) -> %ls", FileInformationClass, uintptr_t(FileHandle), HandleToName(FileHandle), ToString(res).data);
+			return res;
+		};
 
-				auto& info = *(FILE_END_OF_FILE_INFORMATION*)FileInformation;
-				info.EndOfFile.QuadPart = h.pos;
-				return STATUS_SUCCESS;
-			}
-			trueHandle = h.trueHandle;
-			UBA_ASSERT(trueHandle != INVALID_HANDLE_VALUE);
-		}
-		else if (FileInformationClass == FileRenameInformation)
-		{
-			// We can end up in here through MoveFileEx and SetFileInformationByHandle
-			auto& info = *(FILE_RENAME_INFORMATION*)FileInformation;
-			StringView nameView(info.FileName, info.FileNameLength / 2);
-			if (nameView.StartsWith(L"\\??\\"))
-				nameView = nameView.Skip(4);
-			StringBuffer<> newName;
-			FixPath(newName, nameView.data, nameView.count);
-			DevirtualizePath(newName);
-			FileObject& fo = *h.fileObject;
-			fo.newName = newName.ToString();
-
-			StringKey newFileNameKey = ToStringKeyLower(newName);
-
-			// TODO: Revisit this.. don't know what could go wrong with this
-#if 0
-			FileInfo* newFileInfo = nullptr;
-			{
-				SCOPED_READ_LOCK(g_mappedFileTable.m_lookupLock, _);
-				auto findIt = g_mappedFileTable.m_lookup.find(newFileNameKey);
-				if (findIt != g_mappedFileTable.m_lookup.end())
-					newFileInfo = &findIt->second;
-			}
+	if (!isDetouredHandle(FileHandle))
+		return runTrue();
 
-			UBA_ASSERTF(newFileInfo, L"File info already exists for the rename we are doing from %ls to %ls", fo.fileInfo->originalName, newName.data); // TODO: Implement when we find a test case
-#endif
-			DEBUG_LOG_DETOURED(L"NtSetInformationFile", L"File is set to be renamed on close (from %ls to %ls)", HandleToName(FileHandle), fo.newName.c_str());
+	DetouredHandle& h = asDetouredHandle(FileHandle);
+	trueHandle = h.trueHandle;
 
-			if (auto memoryFile = fo.fileInfo->memoryFile)
-			{
-				memoryFile->isReported = false;
-				return STATUS_SUCCESS;
-			}
+	if (FileInformationClass == FilePositionInformation)
+	{
+		auto& info = *(FILE_POSITION_INFORMATION*)FileInformation;
+		if (h.type >= HandleType_StdErr)
+		{
+			info.CurrentByteOffset.QuadPart = 0;
+			return STATUS_SUCCESS;
+		}
+		UBA_ASSERT(h.type == HandleType_File);
+		if (h.type >= HandleType_StdErr)
+		{
+			UBA_ASSERT(false); // TODO: What should we do with this?
+			return STATUS_SUCCESS;
+		}
 
-			UBA_ASSERT(!fo.fileInfo->isFileMap);
+		FileInfo& fi = *h.fileObject->fileInfo;
+		if (fi.memoryFile || fi.isFileMap)
+		{
+			h.pos = info.CurrentByteOffset.QuadPart;
+			DEBUG_LOG_DETOURED(L"NtSetInformationFile", L"(FilePosition MEMORY) %llu %s -> Success", uintptr_t(FileHandle), HandleToName(FileHandle));
+			return STATUS_SUCCESS;
+		}
+		return runTrue();
+	}
+	if (FileInformationClass == FileEndOfFileInformation)
+	{
+		UBA_ASSERT(h.type == HandleType_File);
+		FileInfo& fi = *h.fileObject->fileInfo;
+		if (MemoryFile* mf = fi.memoryFile)
+		{
+			DEBUG_LOG_DETOURED(L"NtSetInformationFile", L"(FileEndOfFile MEMORY) %llu NewSize: %llu (%ls) -> Success", uintptr_t(FileHandle), mf->writtenSize, HandleToName(FileHandle));
+			mf->writtenSize = h.pos;
+			mf->isReported = false;
+			mf->EnsureCommitted(h, mf->writtenSize);
 
-#if 0
-			// This is for the odd rename logic in clang.. foo.so -> foo.so123123.tmp
-			// Clang first create a tmp file, to see if it exists..
-			// It then renames the old file to that file..
-			// ..and finally opens it again with attributes only and deleteonclose flag.. and closes it.
-			// TODO: This is wrong.. we just delete the file.. but what we should do is copy it over in memory to the memory file before setting it to delete
-			SCOPED_READ_LOCK(g_mappedFileTable.m_lookupLock, _);
-			auto findIt = g_mappedFileTable.m_lookup.find(newFileNameKey);
-			if (findIt != g_mappedFileTable.m_lookup.end())
-			{
-				FileInfo& info = findIt->second;
-				if (info.memoryFile)
-				{
-					Rpc_WriteLogf(L"HNNMMMM");
-					FILE_DISPOSITION_INFO info;
-					info.DeleteFileW = true;
-					True_SetFileInformationByHandle(FileHandle, FileDispositionInfo, &info, sizeof(info));
-					return STATUS_SUCCESS;
-				}
-			}
-#endif
+			auto& info = *(FILE_END_OF_FILE_INFORMATION*)FileInformation;
+			info.EndOfFile.QuadPart = h.pos;
+			return STATUS_SUCCESS;
+		}
+		return runTrue();
+	}
+	if (FileInformationClass == FileRenameInformation)
+	{
+		// We can end up in here through MoveFileEx and SetFileInformationByHandle
+		auto& info = *(FILE_RENAME_INFORMATION*)FileInformation;
+		StringView nameView(info.FileName, info.FileNameLength / 2);
+		if (nameView.StartsWith(L"\\??\\"))
+			nameView = nameView.Skip(4);
+		StringBuffer<> newName;
+		FixPath(newName, nameView.data, nameView.count);
+		DevirtualizePath(newName);
+		FileObject& fo = *h.fileObject;
+		fo.newName = newName.ToString();
 
-			if (g_runningRemote) // This needs a proper solution as the comments above.
-				return STATUS_SUCCESS;
+		StringKey newFileNameKey = ToStringKeyLower(newName);
 
-			// In case we are using vfs we need to replace the information before calling the true NtSetInformationFile
-			auto& info2 = *(FILE_RENAME_INFORMATION*)tempBuffer;
-			memcpy(&info2, &info, sizeof(FILE_RENAME_INFORMATION));
-			memcpy(info2.FileName, TC("\\??\\"), 8);
-			memcpy(info2.FileName + 4, newName.data, newName.count * 2);
-			info2.FileNameLength = (newName.count+4) * 2;
-			
-			FileInformation = &info2;
-			Length = sizeof(FILE_RENAME_INFORMATION) + info2.FileNameLength + 2;
-		}
+		DEBUG_LOG_DETOURED(L"NtSetInformationFile", L"File is set to be renamed on close (from %ls to %ls)", HandleToName(FileHandle), fo.newName.c_str());
 
-		trueHandle = h.trueHandle;
-		if (trueHandle == INVALID_HANDLE_VALUE)
+		if (auto memoryFile = fo.fileInfo->memoryFile)
 		{
-			//UBA_ASSERT(!g_runningRemote);
-			// TODO: This needs to be sent back to Session.. so session can set whatever needs to be set.
-			DEBUG_LOG_DETOURED(L"NtSetInformationFile", L"(%u) SKIPPED!!!!!!!!! %llu (%ls) -> Skipped", FileInformationClass, uintptr_t(FileHandle), HandleToName(FileHandle));
+			memoryFile->isReported = false;
 			return STATUS_SUCCESS;
 		}
+
+		UBA_ASSERT(!fo.fileInfo->isFileMap);
+
+		if (g_runningRemote) // This needs a proper solution as the comments above.
+			return STATUS_SUCCESS;
+
+		// In case we are using vfs we need to replace the information before calling the true NtSetInformationFile
+		u8 tempBuffer[sizeof(FILE_RENAME_INFORMATION) + 512];
+		auto& info2 = *(FILE_RENAME_INFORMATION*)tempBuffer;
+		memcpy(&info2, &info, sizeof(FILE_RENAME_INFORMATION));
+		memcpy(info2.FileName, TC("\\??\\"), 8);
+		memcpy(info2.FileName + 4, newName.data, newName.count * 2);
+		info2.FileNameLength = (newName.count+4) * 2;
+			
+		FileInformation = &info2;
+		Length = sizeof(FILE_RENAME_INFORMATION) + info2.FileNameLength + 2;
+		return runTrue();
 	}
 
-	TimerScope ts(g_kernelStats.setFileInfo);
-	auto res = True_NtSetInformationFile(trueHandle, IoStatusBlock, FileInformation, Length, FileInformationClass);
-	DEBUG_LOG_TRUE(L"NtSetInformationFile", L"(%u) %llu (%ls) -> %ls", FileInformationClass, uintptr_t(FileHandle), HandleToName(FileHandle), ToString(res).data);
-	return res;
+	if (trueHandle == INVALID_HANDLE_VALUE)
+	{
+		//UBA_ASSERT(!g_runningRemote);
+		// TODO: This needs to be sent back to Session.. so session can set whatever needs to be set.
+		DEBUG_LOG_DETOURED(L"NtSetInformationFile", L"(%u) SKIPPED!!!!!!!!! %llu (%ls) -> Skipped", FileInformationClass, uintptr_t(FileHandle), HandleToName(FileHandle));
+		return STATUS_SUCCESS;
+	}
+	return runTrue();
 }
 
 NTSTATUS NTAPI Detoured_NtSetInformationObject(HANDLE ObjectHandle, OBJECT_INFORMATION_CLASS ObjectInformationClass, PVOID ObjectInformation, ULONG Length)
