[UBA] Fix potential buffer overflow in SearchPathForFile
Increase fullPath buffer from 512 to MaxPath (1024).

#Synced-CL 51708345
#Synced-CL 51708345

diff --git a/Engine/Source/Programs/UnrealBuildAccelerator/Common/Private/UbaFile.cpp b/Engine/Source/Programs/UnrealBuildAccelerator/Common/Private/UbaFile.cpp
--- a/Engine/Source/Programs/UnrealBuildAccelerator/Common/Private/UbaFile.cpp
+++ b/Engine/Source/Programs/UnrealBuildAccelerator/Common/Private/UbaFile.cpp
@@ -448,7 +448,7 @@ namespace uba
 		bool traverseRes = TraverseDir(logger, ToView(dir),
 			[&](const DirectoryEntry& e)
 			{
-				StringBuffer<> fullPath(dir);
+				StringBuffer<MaxPath> fullPath(dir);
 				fullPath.EnsureEndsWithSlash().Append(e.name);
 
 				if (IsReadOnly(e.attributes))
@@ -499,7 +499,7 @@ namespace uba
 	{
 		UBA_ASSERT(!IsAbsolutePath(file));
 
-		StringBuffer<> fullPath;
+		StringBuffer<MaxPath> fullPath;
 
 		StringView fileView = ToView(file);
 		StringView ext;
