[UBA]
* Improved asserts and logging
* Improved implementation of NtQueryInformationFile FileAllInformation

diff --git a/Engine/Source/Programs/UnrealBuildAccelerator/Common/Private/UbaProcess.cpp b/Engine/Source/Programs/UnrealBuildAccelerator/Common/Private/UbaProcess.cpp
--- a/Engine/Source/Programs/UnrealBuildAccelerator/Common/Private/UbaProcess.cpp
+++ b/Engine/Source/Programs/UnrealBuildAccelerator/Common/Private/UbaProcess.cpp
@@ -993,7 +993,7 @@ namespace uba
 			UBA_PROCESS_MESSAGES
 			#undef UBA_PROCESS_MESSAGE
 		}
-		return CancelWithError(m_session.m_logger.Error(TC("Unknown message type %u when communicating with detoured process %s. (Message index %u after %s)"), messageType, m_startInfo.GetDescription(), m_messageCount, TimeToText(GetTime() - m_startTime).str));
+		return CancelWithError(m_session.m_logger.Error(TC("Unknown message type %u when communicating with detoured process %s (%u). (Message index %u after %s)"), messageType, m_startInfo.GetDescription(), m_id, m_messageCount, TimeToText(GetTime() - m_startTime).str));
 	}
 
 	void ProcessImpl::CallSession(MessageType type, bool success)
@@ -1443,7 +1443,7 @@ namespace uba
 			if (!process.m_nativeProcessHandle || process.m_nativeProcessHandle == InvalidProcHandle)
 				return CancelWithError(m_session.m_logger.Error(TC("Failed to duplicate handle for child process")));
 			if (!DuplicateHandle((HANDLE)m_nativeProcessHandle, nativeThreadHandle, GetCurrentProcess(), &process.m_nativeThreadHandle, 0, false, DUPLICATE_SAME_ACCESS))
-				return CancelWithError(m_session.m_logger.Error(TC("Failed to duplicate handle for child thread")));
+				return CancelWithError(m_session.m_logger.Error(TC("Failed to duplicate handle for child thread %llu in process %llu (%s)"), nativeThreadHandle, m_nativeProcessHandle, LastErrorToText().data));
 			if (!process.m_nativeThreadHandle || process.m_nativeThreadHandle == INVALID_HANDLE_VALUE)
 				return CancelWithError(m_session.m_logger.Error(TC("Failed to duplicate handle for child thread")));
 			process.m_nativeProcessId = nativeProcessId;
@@ -2168,7 +2168,7 @@ namespace uba
 
 			if (!DetourCopyPayloadToProcess((HANDLE)m_nativeProcessHandle, DetoursPayloadGuid, &payload, sizeof(payload)))
 			{
-				logger.Error(TC("Failed to copy payload to process (%s)"), LastErrorToText().data);
+				logger.Error(TC("Failed to copy payload to process %llu - %s (%s)"), m_nativeProcessHandle, m_realApplication.c_str(), LastErrorToText().data);
 				return UBA_EXIT_CODE(9);
 			}
 		}
diff --git a/Engine/Source/Programs/UnrealBuildAccelerator/Core/Public/UbaWinTypes.h b/Engine/Source/Programs/UnrealBuildAccelerator/Core/Public/UbaWinTypes.h
--- a/Engine/Source/Programs/UnrealBuildAccelerator/Core/Public/UbaWinTypes.h
+++ b/Engine/Source/Programs/UnrealBuildAccelerator/Core/Public/UbaWinTypes.h
@@ -102,21 +102,101 @@ struct FILE_STANDARD_INFORMATION {
 	BOOLEAN       Directory;
 };
 
+struct FILE_RENAME_INFORMATION {
+	union {
+		BOOLEAN ReplaceIfExists;  // FileRenameInformation
+		ULONG Flags;              // FileRenameInformationEx
+	} DUMMYUNIONNAME;
+	HANDLE RootDirectory;
+	ULONG FileNameLength;
+	WCHAR FileName[1];
+};
+
+struct FILE_IS_REMOTE_DEVICE_INFORMATION
+{
+	BOOLEAN IsRemote;
+};
+
+struct FILE_ID_INFORMATION
+{
+	ULONGLONG   VolumeSerialNumber;
+	FILE_ID_128 FileId;
+};
+
+struct FILE_NAME_INFORMATION
+{
+	ULONG FileNameLength;
+	WCHAR FileName[1];
+};
+
+struct FILE_INTERNAL_INFORMATION
+{
+	LARGE_INTEGER IndexNumber;
+};
+
+struct FILE_EA_INFORMATION
+{
+	ULONG EaSize;
+};
+
+struct FILE_ACCESS_INFORMATION
+{
+	ACCESS_MASK AccessFlags;
+};
+
+struct FILE_MODE_INFORMATION
+{
+	ULONG Mode;
+};
+
+struct FILE_ALIGNMENT_INFORMATION
+{
+	ULONG AlignmentRequirement;
+};
+
+struct FILE_ALL_INFORMATION
+{
+	FILE_BASIC_INFORMATION     BasicInformation;
+	FILE_STANDARD_INFORMATION  StandardInformation;
+	FILE_INTERNAL_INFORMATION  InternalInformation;
+	FILE_EA_INFORMATION        EaInformation;
+	FILE_ACCESS_INFORMATION    AccessInformation;
+	FILE_POSITION_INFORMATION  PositionInformation;
+	FILE_MODE_INFORMATION      ModeInformation;
+	FILE_ALIGNMENT_INFORMATION AlignmentInformation;
+	FILE_NAME_INFORMATION      NameInformation;
+};
+
+struct FILE_FS_VOLUME_INFORMATION {
+	LARGE_INTEGER VolumeCreationTime;
+	ULONG         VolumeSerialNumber;
+	ULONG         VolumeLabelLength;
+	BOOLEAN       SupportsObjects;
+	WCHAR         VolumeLabel[1];
+};
+
+
+
+#define FILE_INFO_CLASSES \
+	FILE_INFO_CLASS(FileBasicInformation, 4) \
+	FILE_INFO_CLASS(FileStandardInformation, 5) \
+	FILE_INFO_CLASS(FileNameInformation, 9) \
+	FILE_INFO_CLASS(FileRenameInformation, 10) \
+	FILE_INFO_CLASS(FileDispositionInformation, 13) \
+	FILE_INFO_CLASS(FilePositionInformation, 14) \
+	FILE_INFO_CLASS(FileAllInformation, 18) \
+	FILE_INFO_CLASS(FileEndOfFileInformation, 20) \
+	FILE_INFO_CLASS(FileNetworkOpenInformation, 34) \
+	FILE_INFO_CLASS(FileNormalizedNameInformation, 35) \
+	FILE_INFO_CLASS(FileIdBothDirectoryInformation, 37) \
+	FILE_INFO_CLASS(FileIsRemoteDeviceInformation, 51) \
+	FILE_INFO_CLASS(FileIdInformation, 59) \
+	FILE_INFO_CLASS(FileRenameInformationEx, 65) \
+
 
-constexpr FILE_INFORMATION_CLASS FileBasicInformation = (FILE_INFORMATION_CLASS)4;
-constexpr FILE_INFORMATION_CLASS FileStandardInformation = (FILE_INFORMATION_CLASS)5;
-constexpr FILE_INFORMATION_CLASS FileNameInformation = (FILE_INFORMATION_CLASS)9;
-constexpr FILE_INFORMATION_CLASS FileRenameInformation = (FILE_INFORMATION_CLASS)10;
-constexpr FILE_INFORMATION_CLASS FileDispositionInformation = (FILE_INFORMATION_CLASS)13;
-constexpr FILE_INFORMATION_CLASS FilePositionInformation = (FILE_INFORMATION_CLASS)14;
-constexpr FILE_INFORMATION_CLASS FileAllInformation = (FILE_INFORMATION_CLASS)18;
-constexpr FILE_INFORMATION_CLASS FileEndOfFileInformation = (FILE_INFORMATION_CLASS)20;
-constexpr FILE_INFORMATION_CLASS FileNetworkOpenInformation = (FILE_INFORMATION_CLASS)34;
-constexpr FILE_INFORMATION_CLASS FileNormalizedNameInformation = (FILE_INFORMATION_CLASS)35;
-constexpr FILE_INFORMATION_CLASS FileIdBothDirectoryInformation = (FILE_INFORMATION_CLASS)37;
-constexpr FILE_INFORMATION_CLASS FileIsRemoteDeviceInformation = (FILE_INFORMATION_CLASS)51;
-constexpr FILE_INFORMATION_CLASS FileIdInformation = (FILE_INFORMATION_CLASS)59;
-constexpr FILE_INFORMATION_CLASS FileRenameInformationEx = (FILE_INFORMATION_CLASS)65;
+#define FILE_INFO_CLASS(name, value) constexpr FILE_INFORMATION_CLASS name = (FILE_INFORMATION_CLASS)value;
+FILE_INFO_CLASSES
+#undef FILE_INFO_CLASS
 
 constexpr OBJECT_INFORMATION_CLASS ObjectNameInformation = (OBJECT_INFORMATION_CLASS)1;
 
diff --git a/Engine/Source/Programs/UnrealBuildAccelerator/Detours/Private/Windows/UbaDetoursFunctionsNtDll.inl b/Engine/Source/Programs/UnrealBuildAccelerator/Detours/Private/Windows/UbaDetoursFunctionsNtDll.inl
--- a/Engine/Source/Programs/UnrealBuildAccelerator/Detours/Private/Windows/UbaDetoursFunctionsNtDll.inl
+++ b/Engine/Source/Programs/UnrealBuildAccelerator/Detours/Private/Windows/UbaDetoursFunctionsNtDll.inl
@@ -1,5 +1,18 @@
 // Copyright Epic Games, Inc. All Rights Reserved.
 
+StringBuffer<64> ToString(FILE_INFORMATION_CLASS c)
+{
+	switch (u32(c))
+	{
+		#define FILE_INFO_CLASS(name, value) case value: return StringBuffer<64>(TC(#name));
+		FILE_INFO_CLASSES
+		#undef FILE_INFO_CLASS
+	default:
+		return StringBuffer<64>(TCV("class ")).AppendValue(u32(c));
+	}
+}
+
+
 bool IsContentWrite(u32 desiredAccess, u32 createDisposition)
 {
 	if (desiredAccess & (FILE_WRITE_DATA | FILE_APPEND_DATA | GENERIC_WRITE))
@@ -190,7 +203,7 @@ NTSTATUS Detoured_NtQueryInformationFile(HANDLE FileHandle, PIO_STATUS_BLOCK IoS
 		*/
 		else
 		{
-			FatalError(1348, L"NtQueryInformationFile with class %u not implemented", FileInformationClass);
+			FatalError(1348, L"NtQueryInformationFile (%s) not implemented", ToString(FileInformationClass).data);
 		}
 	}
 
@@ -200,7 +213,7 @@ NTSTATUS Detoured_NtQueryInformationFile(HANDLE FileHandle, PIO_STATUS_BLOCK IoS
 			UBA_ASSERT(trueHandle != INVALID_HANDLE_VALUE);
 			TimerScope ts(g_kernelStats.getFileInfo);
 			auto res = True_NtQueryInformationFile(trueHandle, IoStatusBlock, FileInformation, Length, FileInformationClass);
-			DEBUG_LOG_TRUE(L"NtQueryInformationFile", L"(%u) %llu (%ls) -> %ls", FileInformationClass, uintptr_t(FileHandle), HandleToName(FileHandle), ToString(res).data);
+			DEBUG_LOG_TRUE(L"NtQueryInformationFile", L"(%s) %llu (%ls) -> %ls", ToString(FileInformationClass).data, uintptr_t(FileHandle), HandleToName(FileHandle), ToString(res).data);
 			return res;
 		};
 
@@ -210,38 +223,76 @@ NTSTATUS Detoured_NtQueryInformationFile(HANDLE FileHandle, PIO_STATUS_BLOCK IoS
 	auto& dh = asDetouredHandle(FileHandle);
 	trueHandle = dh.trueHandle;
 
+	auto setBasicInformation = [](FILE_BASIC_INFORMATION& out, DirectoryTable::EntryInformation& entryInfo)
+		{
+			out.CreationTime.QuadPart = entryInfo.lastWrite;
+			out.LastAccessTime.QuadPart = entryInfo.lastWrite;
+			out.LastWriteTime.QuadPart = entryInfo.lastWrite;
+			out.ChangeTime.QuadPart = entryInfo.lastWrite;
+			out.FileAttributes = entryInfo.attributes;
+			return STATUS_SUCCESS;
+		};
+
+	auto setStandardInformation = [&](FILE_STANDARD_INFORMATION& out, FileInfo& fi, DirectoryTable::EntryInformation* entryInfo)
+		{
+			u64 size = 0;
+			u64 allocationSize = 0;
+			if (fi.memoryFile)
+			{
+				size = fi.memoryFile->writtenSize;
+				allocationSize = fi.memoryFile->committedSize;
+			}
+			else if (fi.fileMapMem)
+			{
+				size = fi.fileMapMemSize;
+				allocationSize = size;
+			}
+			else if (fi.size != InvalidValue)
+			{
+				size = fi.size;
+				allocationSize = size;
+			}
+			else if (entryInfo)
+			{
+				size = entryInfo->size;
+				allocationSize = size;
+			}
+			else
+				return STATUS_INVALID_HANDLE;
+			out.DeletePending = FALSE;
+			out.NumberOfLinks = 0;
+			out.Directory = FALSE;
+			out.AllocationSize.QuadPart = allocationSize;
+			out.EndOfFile.QuadPart = size;
+			return STATUS_SUCCESS;
+		};
+
+	auto setNameInformation = [&](FILE_NAME_INFORMATION& out, FileInfo& fi)
+		{
+			const wchar_t* fileName = fi.originalName;
+			auto& info = *(FILE_NAME_INFORMATION*)FileInformation;
+			ULONG header = FIELD_OFFSET(FILE_NAME_INFORMATION, FileName);
+			ULONG nameBytes = (ULONG)(wcslen(fileName) * sizeof(WCHAR));
+			ULONG required = header + nameBytes;
+			if (Length < header)
+			{
+				IoStatusBlock->Status = STATUS_INFO_LENGTH_MISMATCH;
+				return STATUS_INFO_LENGTH_MISMATCH;
+			}
+			ULONG copy = Min(nameBytes, Length - header);
+			info.FileNameLength = nameBytes;
+			if (copy)
+				memcpy(info.FileName, fileName, copy);
+			IoStatusBlock->Information = Min(Length, required);
+			return STATUS_SUCCESS;
+		};
+
 	if (FileInformationClass == FileStandardInformation)
 	{
 		UBA_ASSERT(dh.type == HandleType_File);
 		FileInfo& fi = *dh.fileObject->fileInfo;
-		u64 size = 0;
-		u64 allocationSize = 0;
-		if (fi.memoryFile)
-		{
-			size = fi.memoryFile->writtenSize;
-			allocationSize = fi.memoryFile->committedSize;
-		}
-		else if (fi.fileMapMem)
-		{
-			size = fi.fileMapMemSize;
-			allocationSize = size;
-		}
-		else if (fi.size == InvalidValue)
-		{
+		if (setStandardInformation(*(FILE_STANDARD_INFORMATION*)FileInformation, fi, nullptr) != STATUS_SUCCESS)
 			return runTrue();
-		}
-		else
-		{
-			size = fi.size;
-			allocationSize = size;
-		}
-
-		auto& info = *(FILE_STANDARD_INFORMATION*)FileInformation;
-		info.DeletePending = FALSE;
-		info.NumberOfLinks = 0;
-		info.Directory = FALSE;
-		info.AllocationSize.QuadPart = allocationSize;
-		info.EndOfFile.QuadPart = size;
 		DEBUG_LOG_DETOURED(L"NtQueryInformationFile", L"(FileStandardInformation) %llu (%s) -> %ls", uintptr_t(FileHandle), HandleToName(FileHandle), ToString(STATUS_SUCCESS).data);
 		return STATUS_SUCCESS;
 	}
@@ -257,24 +308,19 @@ NTSTATUS Detoured_NtQueryInformationFile(HANDLE FileHandle, PIO_STATUS_BLOCK IoS
 		FileInfo& fi = *dh.fileObject->fileInfo;
 		if (!fi.memoryFile && !fi.isFileMap)
 			return runTrue();
-		auto& info = *(FILE_POSITION_INFORMATION*)FileInformation;
+		auto& info =  *(FILE_POSITION_INFORMATION*)FileInformation;
 		info.CurrentByteOffset.QuadPart = dh.pos;
 		DEBUG_LOG_DETOURED(L"NtQueryInformationFile", L"(FilePositionInformation) %llu (%s) -> %ls", uintptr_t(FileHandle), HandleToName(FileHandle), ToString(STATUS_SUCCESS).data);
 		return STATUS_SUCCESS;
 	}
 	if (FileInformationClass == FileBasicInformation)
 	{
-		auto& info = *(FILE_BASIC_INFORMATION*)FileInformation;
 		if (dh.dirTableOffset == ~0u)
 			return runTrue();
 		DirectoryTable::EntryInformation entryInfo;
 		g_directoryTable.GetEntryInformation(entryInfo, dh.dirTableOffset);
 		UBA_ASSERT(entryInfo.attributes != 0);
-		info.CreationTime.QuadPart = entryInfo.lastWrite;
-		info.LastAccessTime.QuadPart = entryInfo.lastWrite;
-		info.LastWriteTime.QuadPart = entryInfo.lastWrite;
-		info.ChangeTime.QuadPart = entryInfo.lastWrite;
-		info.FileAttributes = entryInfo.attributes;
+		setBasicInformation(*(FILE_BASIC_INFORMATION*)FileInformation, entryInfo);
 		DEBUG_LOG_DETOURED(L"NtQueryInformationFile", L"(FileBasicInformation) %llu (%s) -> %ls", uintptr_t(FileHandle), HandleToName(FileHandle), ToString(STATUS_SUCCESS).data);
 		return STATUS_SUCCESS;
 	}
@@ -282,27 +328,30 @@ NTSTATUS Detoured_NtQueryInformationFile(HANDLE FileHandle, PIO_STATUS_BLOCK IoS
 	{
 		if (dh.dirTableOffset == ~0u)
 			return runTrue();
+		if (dh.trueHandle != INVALID_HANDLE_VALUE)
+			return runTrue();
+		UBA_ASSERT(dh.type == HandleType_File);
+
 		DirectoryTable::EntryInformation entryInfo;
 		g_directoryTable.GetEntryInformation(entryInfo, dh.dirTableOffset);
 		UBA_ASSERT(entryInfo.attributes != 0);
+		FileInfo& fi = *dh.fileObject->fileInfo;
 
 		// TODO This code path is here to handle nodejs queries.. Is not properly implemented and miss things
 		auto& info = *(FILE_ALL_INFORMATION*)FileInformation;
-		info.BasicInformation.CreationTime.QuadPart = entryInfo.lastWrite;
-		info.BasicInformation.LastAccessTime.QuadPart = entryInfo.lastWrite;
-		info.BasicInformation.LastWriteTime.QuadPart = entryInfo.lastWrite;
-		info.BasicInformation.ChangeTime.QuadPart = entryInfo.lastWrite;
-		info.BasicInformation.FileAttributes = entryInfo.attributes;
-		info.StandardInformation.AllocationSize.QuadPart = entryInfo.size;
-		info.StandardInformation.EndOfFile.QuadPart = entryInfo.size;
-		info.StandardInformation.NumberOfLinks = 0;
-		info.StandardInformation.DeletePending = false;
-		info.StandardInformation.Directory = false;
+		NTSTATUS res = STATUS_SUCCESS;
+		
+		res = res == STATUS_SUCCESS ? setBasicInformation(info.BasicInformation, entryInfo) : res;
+		res = res == STATUS_SUCCESS ? setStandardInformation(info.StandardInformation, fi, &entryInfo) : res;
+		res = res == STATUS_SUCCESS ? setNameInformation(info.NameInformation, fi) : res;
 		info.InternalInformation.IndexNumber.QuadPart = entryInfo.fileIndex;
-		//IoStatusBlock->Information = FIELD_OFFSET(FILE_ALL_INFORMATION, NameInformation.FileName) + name_len;
-
-		DEBUG_LOG_DETOURED(L"NtQueryInformationFile", L"NOT IMPLEMENTED (FileAllInformation) %llu (%s) -> %ls", uintptr_t(FileHandle), HandleToName(FileHandle), ToString(STATUS_SUCCESS).data);
-		return STATUS_SUCCESS;
+		info.EaInformation.EaSize = 0;
+		info.AccessInformation.AccessFlags = dh.fileObject->desiredAccess;
+		info.PositionInformation.CurrentByteOffset.QuadPart = 0;
+		info.ModeInformation.Mode = 0;
+		info.AlignmentInformation.AlignmentRequirement = 0;
+		DEBUG_LOG_DETOURED(L"NtQueryInformationFile", L"(FileAllInformation) %llu (%s) -> %ls", uintptr_t(FileHandle), HandleToName(FileHandle), ToString(res).data);
+		return res;
 	}
 	if (FileInformationClass == FileNetworkOpenInformation)
 	{
@@ -350,21 +399,9 @@ NTSTATUS Detoured_NtQueryInformationFile(HANDLE FileHandle, PIO_STATUS_BLOCK IoS
 	{
 		if (dh.trueHandle != INVALID_HANDLE_VALUE)
 			return runTrue();
-		const wchar_t* fileName = dh.fileObject->fileInfo->originalName;
-		auto& info = *(FILE_NAME_INFORMATION*)FileInformation;
-		ULONG header = FIELD_OFFSET(FILE_NAME_INFORMATION, FileName);
-		ULONG nameBytes = (ULONG)(wcslen(fileName) * sizeof(WCHAR));
-		ULONG required = header + nameBytes;
-		if (Length < header)
-			return STATUS_INFO_LENGTH_MISMATCH;
-		ULONG copy = Min(nameBytes, Length - header);
-		info.FileNameLength = nameBytes;
-		if (copy)
-			memcpy(info.FileName, fileName, copy);
-		IoStatusBlock->Status = STATUS_SUCCESS;
-		IoStatusBlock->Information = Min(Length, required);
-		DEBUG_LOG_DETOURED(L"NtQueryInformationFile", L"(FileNameInformation) %llu (%s) -> %ls", uintptr_t(FileHandle), HandleToName(FileHandle), ToString(STATUS_SUCCESS).data);
-		return STATUS_SUCCESS;
+		NTSTATUS res = setNameInformation(*(FILE_NAME_INFORMATION*)FileInformation, *dh.fileObject->fileInfo);
+		DEBUG_LOG_DETOURED(L"NtQueryInformationFile", L"(FileNameInformation) %llu (%s) -> %ls", uintptr_t(FileHandle), HandleToName(FileHandle), ToString(res).data);
+		return res;
 	}
 
 	UBA_ASSERTF(false, L"NtQueryInformationFile (%u) failed using detoured handle %ls (%ls)", FileInformationClass, dh.fileObject->fileInfo->name, dh.fileObject->fileInfo->originalName);
@@ -505,7 +542,7 @@ NTSTATUS NTAPI Detoured_NtQueryDirectoryFile(HANDLE FileHandle, HANDLE Event, PI
 	{
 		DetouredHandle& h = asDetouredHandle(FileHandle);
 		trueHandle = h.trueHandle;
-		UBA_ASSERTF(trueHandle != INVALID_HANDLE_VALUE, L"NtQueryDirectoryFile for using class %u not implemented for detoured handles (%ls)", FileInformationClass, HandleToName(FileHandle));
+		UBA_ASSERTF(trueHandle != INVALID_HANDLE_VALUE, L"NtQueryDirectoryFile (%s) not implemented for detoured handles (%ls)", ToString(FileInformationClass).data, HandleToName(FileHandle));
 	}
 
 	NTSTATUS res = True_NtQueryDirectoryFile(trueHandle, Event, ApcRoutine, ApcContext, IoStatusBlock, FileInformation, Length, FileInformationClass, ReturnSingleEntry, FileName, RestartScan);
@@ -612,7 +649,7 @@ NTSTATUS Detoured_NtSetInformationFile(HANDLE FileHandle, PIO_STATUS_BLOCK IoSta
 			UBA_ASSERT(trueHandle != INVALID_HANDLE_VALUE);
 			TimerScope ts(g_kernelStats.setFileInfo);
 			auto res = True_NtSetInformationFile(trueHandle, IoStatusBlock, FileInformation, Length, FileInformationClass);
-			DEBUG_LOG_TRUE(L"NtSetInformationFile", L"(%u) %llu (%ls) -> %ls", FileInformationClass, uintptr_t(FileHandle), HandleToName(FileHandle), ToString(res).data);
+			DEBUG_LOG_TRUE(L"NtSetInformationFile", L"(%s) %llu (%ls) -> %ls", ToString(FileInformationClass).data, uintptr_t(FileHandle), HandleToName(FileHandle), ToString(res).data);
 			return res;
 		};
 
@@ -622,6 +659,18 @@ NTSTATUS Detoured_NtSetInformationFile(HANDLE FileHandle, PIO_STATUS_BLOCK IoSta
 	DetouredHandle& h = asDetouredHandle(FileHandle);
 	trueHandle = h.trueHandle;
 
+	if (FileInformationClass == FileBasicInformation)
+	{
+		auto& info = *(FILE_BASIC_INFORMATION*)FileInformation;
+		u32 attr = info.FileAttributes;
+		bool attrUnchanged = attr == (g_isRunningWine ? 0 : 0xFFFFFFFF);
+		if (attrUnchanged && info.CreationTime.QuadPart == -1 && info.LastWriteTime.QuadPart == -1)
+		{
+			DEBUG_LOG_DETOURED(L"NtSetInformationFile", L"(FileBasic) %llu IGNORE (%ls)", uintptr_t(FileHandle), HandleToName(FileHandle));
+			return TRUE;
+		}
+		return runTrue();
+	}
 	if (FileInformationClass == FilePositionInformation)
 	{
 		auto& info = *(FILE_POSITION_INFORMATION*)FileInformation;
@@ -708,7 +757,7 @@ NTSTATUS Detoured_NtSetInformationFile(HANDLE FileHandle, PIO_STATUS_BLOCK IoSta
 	{
 		//UBA_ASSERT(!g_runningRemote);
 		// TODO: This needs to be sent back to Session.. so session can set whatever needs to be set.
-		DEBUG_LOG_DETOURED(L"NtSetInformationFile", L"(%u) SKIPPED!!!!!!!!! %llu (%ls) -> Skipped", FileInformationClass, uintptr_t(FileHandle), HandleToName(FileHandle));
+		DEBUG_LOG_DETOURED(L"NtSetInformationFile", L"(%s) SKIPPED!!!!!!!!! %llu (%ls) -> Skipped", ToString(FileInformationClass).data, uintptr_t(FileHandle), HandleToName(FileHandle));
 		return STATUS_SUCCESS;
 	}
 	return runTrue();
diff --git a/Engine/Source/Programs/UnrealBuildAccelerator/Detours/Private/Windows/UbaDetoursFunctionsWin.h b/Engine/Source/Programs/UnrealBuildAccelerator/Detours/Private/Windows/UbaDetoursFunctionsWin.h
--- a/Engine/Source/Programs/UnrealBuildAccelerator/Detours/Private/Windows/UbaDetoursFunctionsWin.h
+++ b/Engine/Source/Programs/UnrealBuildAccelerator/Detours/Private/Windows/UbaDetoursFunctionsWin.h
@@ -477,56 +477,6 @@ extern "C" {
 DETOURED_FUNCTIONS
 #undef DETOURED_FUNCTION
 
-struct FILE_RENAME_INFORMATION {
-    union {
-        BOOLEAN ReplaceIfExists;  // FileRenameInformation
-        ULONG Flags;              // FileRenameInformationEx
-    } DUMMYUNIONNAME;
-    HANDLE RootDirectory;
-    ULONG FileNameLength;
-    WCHAR FileName[1];
-};
-
-struct FILE_IS_REMOTE_DEVICE_INFORMATION
-{
-	BOOLEAN IsRemote;
-};
-
-struct FILE_ID_INFORMATION {
-	ULONGLONG   VolumeSerialNumber;
-	FILE_ID_128 FileId;
-};
-
-struct FILE_NAME_INFORMATION {
-	ULONG FileNameLength;
-	WCHAR FileName[1];
-};
-
-struct FILE_INTERNAL_INFORMATION
-{
-	LARGE_INTEGER IndexNumber;
-};
-
-struct FILE_ALL_INFORMATION {
-	FILE_BASIC_INFORMATION     BasicInformation;
-	FILE_STANDARD_INFORMATION  StandardInformation;
-	FILE_INTERNAL_INFORMATION  InternalInformation;
-	//FILE_EA_INFORMATION        EaInformation;
-	//FILE_ACCESS_INFORMATION    AccessInformation;
-	//FILE_POSITION_INFORMATION  PositionInformation;
-	//FILE_MODE_INFORMATION      ModeInformation;
-	//FILE_ALIGNMENT_INFORMATION AlignmentInformation;
-	//FILE_NAME_INFORMATION      NameInformation;
-};
-
-struct FILE_FS_VOLUME_INFORMATION {
-	LARGE_INTEGER VolumeCreationTime;
-	ULONG         VolumeSerialNumber;
-	ULONG         VolumeLabelLength;
-	BOOLEAN       SupportsObjects;
-	WCHAR         VolumeLabel[1];
-};
-
 namespace uba
 {
 	struct DetoursPayload;
