[UBA]
* Added testing code to check if we can detour the path conversion instead of detouring both CreateFile and NtCreateFile (disabled by default)

diff --git a/Engine/Source/Programs/UnrealBuildAccelerator/Core/Public/UbaWinTypes.h b/Engine/Source/Programs/UnrealBuildAccelerator/Core/Public/UbaWinTypes.h
--- a/Engine/Source/Programs/UnrealBuildAccelerator/Core/Public/UbaWinTypes.h
+++ b/Engine/Source/Programs/UnrealBuildAccelerator/Core/Public/UbaWinTypes.h
@@ -204,3 +204,22 @@ enum FS_INFORMATION_CLASS {};
 
 constexpr FS_INFORMATION_CLASS FileFsVolumeInformation = (FS_INFORMATION_CLASS)1;
 constexpr FS_INFORMATION_CLASS FileFsDeviceInformation = (FS_INFORMATION_CLASS)4;
+
+enum RTL_PATH_TYPE
+{
+	RtlPathTypeUnknown = 0,
+	RtlPathTypeUncAbsolute,     // "//foo"
+	RtlPathTypeDriveAbsolute,   // "c:/foo"
+	RtlPathTypeDriveRelative,   // "c:foo"
+	RtlPathTypeRooted,          // "/foo"
+	RtlPathTypeRelative,        // "foo"
+	RtlPathTypeLocalDevice,     // "//./foo"
+	RtlPathTypeRootLocalDevice  // "//."
+};
+
+struct RTL_RELATIVE_NAME_U
+{
+	UNICODE_STRING RelativeName;
+	HANDLE ContainingDirectory;
+	PVOID CurDirRef;
+};
\ No newline at end of file
diff --git a/Engine/Source/Programs/UnrealBuildAccelerator/Detours/Private/Windows/UbaDetoursFunctionsKernelBase.inl b/Engine/Source/Programs/UnrealBuildAccelerator/Detours/Private/Windows/UbaDetoursFunctionsKernelBase.inl
--- a/Engine/Source/Programs/UnrealBuildAccelerator/Detours/Private/Windows/UbaDetoursFunctionsKernelBase.inl
+++ b/Engine/Source/Programs/UnrealBuildAccelerator/Detours/Private/Windows/UbaDetoursFunctionsKernelBase.inl
@@ -201,8 +201,17 @@ HANDLE Detoured_CreateFileW(LPCWSTR lpFileName, DWORD dwDesiredAccess, DWORD dwS
 	DEBUG_LOG_DETOURED(L"CreateFileW", L"%ls", lpFileName);
 	u32 disallowDetour = !lpFileName || Equals(lpFileName, L"nul");
 	t_disallowDetour += disallowDetour;
+
+	#if UBA_CHECK_PATH_CONVERSION
+	t_rtlCalled = false;
+	#endif
+
 	t_createFileFileName = lpFileName;
 	HANDLE h = True_CreateFileW(lpFileName, dwDesiredAccess, dwShareMode, lpSecurityAttributes, dwCreationDisposition, dwFlagsAndAttributes, hTemplateFile);
+
+	#if UBA_CHECK_PATH_CONVERSION
+	UBA_ASSERT(t_rtlCalled);
+	#endif
 	t_createFileFileName = nullptr;
 	t_disallowDetour -= disallowDetour;
 	return h;
@@ -217,8 +226,18 @@ HANDLE Detoured_CreateFileA(LPCSTR lpFileName, DWORD dwDesiredAccess, DWORD dwSh
 	fileName.Appendf(L"%hs", lpFileName);
 	u32 disallowDetour = fileName.Equals(L"nul");
 	t_disallowDetour += disallowDetour;
+
+	#if UBA_CHECK_PATH_CONVERSION
+	t_rtlCalled = false;
+	#endif
+
 	t_createFileFileName = fileName.data;
 	HANDLE h = True_CreateFileA(lpFileName, dwDesiredAccess, dwShareMode, lpSecurityAttributes, dwCreationDisposition, dwFlagsAndAttributes, hTemplateFile);
+
+	#if UBA_CHECK_PATH_CONVERSION
+	UBA_ASSERT(t_rtlCalled);
+	#endif
+
 	t_createFileFileName = nullptr;
 	t_disallowDetour -= disallowDetour;
 	return h;
diff --git a/Engine/Source/Programs/UnrealBuildAccelerator/Detours/Private/Windows/UbaDetoursFunctionsNtDll.inl b/Engine/Source/Programs/UnrealBuildAccelerator/Detours/Private/Windows/UbaDetoursFunctionsNtDll.inl
--- a/Engine/Source/Programs/UnrealBuildAccelerator/Detours/Private/Windows/UbaDetoursFunctionsNtDll.inl
+++ b/Engine/Source/Programs/UnrealBuildAccelerator/Detours/Private/Windows/UbaDetoursFunctionsNtDll.inl
@@ -942,14 +942,24 @@ NTSTATUS NTAPI Local_NtCreateFile(bool IsCreateFunc, PHANDLE hFileHandle, ACCESS
 	}
 }
 
+#define UBA_CHECK_PATH_CONVERSION 0
+
+#if UBA_CHECK_PATH_CONVERSION
+thread_local bool t_rtlCalled;
+#endif
+
 NTSTATUS NTAPI Shared_NtCreateFile(bool IsCreateFunc, PHANDLE hFileHandle, ACCESS_MASK DesiredAccess, POBJECT_ATTRIBUTES ObjectAttributes, PIO_STATUS_BLOCK IoStatusBlock, PLARGE_INTEGER AllocationSize, ULONG FileAttributes,
 	ULONG ShareAccess, ULONG CreateDisposition, ULONG CreateOptions, PVOID EaBuffer, ULONG EaLength)
 {
 	*hFileHandle = INVALID_HANDLE_VALUE;
 
-#if UBA_DEBUG_LOG_ENABLED
+	#if UBA_DEBUG_LOG_ENABLED
 	const wchar_t* funcName = IsCreateFunc ? L"NtCreateFile" : L"NtOpenFile"; (void)funcName;
-#endif
+	#endif
+
+	#if UBA_CHECK_PATH_CONVERSION
+	UBA_ASSERT(!t_createFileFileName || t_rtlCalled);
+	#endif
 
 	const wchar_t* createFileName = t_createFileFileName;
 
@@ -2474,9 +2484,44 @@ BOOLEAN Detoured_RtlValidateHeap(HANDLE HeapPtr, ULONG Flags, PVOID Block)
 
 NTSTATUS Detoured_RtlDosPathNameToNtPathName_U_WithStatus(PCWSTR dos_path, PUNICODE_STRING ntpath, PWSTR* file_part, VOID* reserved)
 {
+	#if UBA_CHECK_PATH_CONVERSION
+	t_rtlCalled = true;
+	#endif
 	auto res = True_RtlDosPathNameToNtPathName_U_WithStatus(dos_path, ntpath, file_part, reserved);
 	//DEBUG_LOG_TRUE(L"Detoured_RtlDosPathNameToNtPathName_U_WithStatus", L"-> %s", ToString(res).data);
 	return res;
 }
 
+NTSTATUS Detoured_RtlDosPathNameToRelativeNtPathName_U_WithStatus(PCWSTR DosFileName, PUNICODE_STRING NtFileName, PWSTR *FilePart, RTL_RELATIVE_NAME_U* RelativeName)
+{
+	#if UBA_CHECK_PATH_CONVERSION
+	t_rtlCalled = true;
+	#endif
+	return True_RtlDosPathNameToRelativeNtPathName_U_WithStatus(DosFileName, NtFileName, FilePart, RelativeName);
+}
+
+DWORD WINAPI Detoured_RtlGetFullPathName_U(const WCHAR* name, ULONG size, WCHAR* buffer, WCHAR** file_part)
+{
+	#if UBA_CHECK_PATH_CONVERSION
+	t_rtlCalled = true;
+	#endif
+	return True_RtlGetFullPathName_U(name, size, buffer, file_part);
+}
+
+ULONG WINAPI Detoured_RtlGetFullPathName_UEx(const WCHAR* name, ULONG size, WCHAR* buffer, WCHAR** file_part, RTL_PATH_TYPE* type)
+{
+	#if UBA_CHECK_PATH_CONVERSION
+	t_rtlCalled = true;
+	#endif
+	return True_RtlGetFullPathName_UEx(name, size, buffer, file_part, type);
+}
+
+BOOLEAN NTAPI Detoured_RtlDosPathNameToRelativeNtPathName_U(PCWSTR DosName, PUNICODE_STRING NtName, PWSTR *FilePart, RTL_RELATIVE_NAME_U* RelativeName)
+{
+	#if UBA_CHECK_PATH_CONVERSION
+	t_rtlCalled = true;
+	#endif
+	return True_RtlDosPathNameToRelativeNtPathName_U(DosName, NtName, FilePart, RelativeName);
+}
+
 #endif
diff --git a/Engine/Source/Programs/UnrealBuildAccelerator/Detours/Private/Windows/UbaDetoursFunctionsWin.h b/Engine/Source/Programs/UnrealBuildAccelerator/Detours/Private/Windows/UbaDetoursFunctionsWin.h
--- a/Engine/Source/Programs/UnrealBuildAccelerator/Detours/Private/Windows/UbaDetoursFunctionsWin.h
+++ b/Engine/Source/Programs/UnrealBuildAccelerator/Detours/Private/Windows/UbaDetoursFunctionsWin.h
@@ -362,6 +362,10 @@
 	DETOURED_FUNCTION(RtlReAllocateHeap) \
 	DETOURED_FUNCTION(RtlValidateHeap) \
 	DETOURED_FUNCTION(RtlDosPathNameToNtPathName_U_WithStatus) \
+	DETOURED_FUNCTION(RtlDosPathNameToRelativeNtPathName_U_WithStatus) \
+	DETOURED_FUNCTION(RtlGetFullPathName_U) \
+	DETOURED_FUNCTION(RtlGetFullPathName_UEx) \
+	DETOURED_FUNCTION(RtlDosPathNameToRelativeNtPathName_U) \
 	DETOURED_FUNCTION(NtCreateIoCompletion) \
 	DETOURED_FUNCTION(NtFlushBuffersFileEx) \
 	DETOURED_FUNCTION(NtTerminateProcess) \
@@ -463,6 +467,10 @@ extern "C" {
 	BOOLEAN NTAPI RtlFreeHeap(PVOID HeapHandle, ULONG Flags, PVOID HeapBase);
 	BOOLEAN NTAPI RtlValidateHeap(HANDLE HeapPtr, ULONG Flags, PVOID Block);
 	NTSTATUS NTAPI RtlDosPathNameToNtPathName_U_WithStatus(PCWSTR dos_path, PUNICODE_STRING ntpath, PWSTR* file_part, VOID* reserved);
+	NTSTATUS NTAPI RtlDosPathNameToRelativeNtPathName_U_WithStatus(PCWSTR DosFileName, PUNICODE_STRING NtFileName, PWSTR *FilePart, RTL_RELATIVE_NAME_U* RelativeName);
+	DWORD WINAPI RtlGetFullPathName_U(const WCHAR* name, ULONG size, WCHAR* buffer, WCHAR** file_part);
+	BOOLEAN NTAPI RtlDosPathNameToRelativeNtPathName_U(PCWSTR DosName, PUNICODE_STRING NtName, PWSTR *FilePart, RTL_RELATIVE_NAME_U* RelativeName);
+	ULONG NTAPI WINAPI RtlGetFullPathName_UEx(const WCHAR* name, ULONG size, WCHAR* buffer, WCHAR** file_part, RTL_PATH_TYPE* type);
 	VOID NTAPI RtlExitUserProcess(NTSTATUS ExitStatus);
 	VOID NTSYSAPI RtlMapGenericMask(PACCESS_MASK AccessMask, const GENERIC_MAPPING *GenericMapping);
 	PVOID NTAPI RtlAllocateHeap( PVOID HeapHandle, ULONG Flags, SIZE_T Size);
