[UBA]
* Fixed bug where invalid process handle was treated as -1 on win32 (it is null)
* Improved asserts
* Added more logic around terminating child processes
* Added code for debugging wine file system calls
* Fixed crash in visualizer when hovering over processes that got cancelled really early

diff --git a/Engine/Source/Programs/UnrealBuildAccelerator/Common/Private/UbaProcess.cpp b/Engine/Source/Programs/UnrealBuildAccelerator/Common/Private/UbaProcess.cpp
--- a/Engine/Source/Programs/UnrealBuildAccelerator/Common/Private/UbaProcess.cpp
+++ b/Engine/Source/Programs/UnrealBuildAccelerator/Common/Private/UbaProcess.cpp
@@ -1337,7 +1337,7 @@ namespace uba
 		TString applicationStr = reader.ReadString();
 		TString commandLineWithoutApplication = reader.ReadLongString();
 
-		UBA_ASSERTF(!applicationStr.empty() && (applicationStr[0] != '\"' || applicationStr[1] != '\"'), TC("Invalid application name: %s"), applicationStr.c_str());
+		UBA_ASSERTF(!applicationStr.empty() && (applicationStr[0] != '\"' || applicationStr[1] != '\"'), TC("Invalid application name: %s  (Length %llu)"), applicationStr.c_str(), applicationStr.size());
 
 		StringBuffer<> currentDir;
 		reader.ReadString(currentDir);
@@ -1384,6 +1384,7 @@ namespace uba
 
 		const char* detoursLib = m_session.m_detoursLibrary[IsLinux][IsArmBinary].c_str();
 		u32 detoursLibLen = u32(m_session.m_detoursLibrary[IsLinux][IsArmBinary].size());
+		UBA_ASSERT(detoursLibLen < 1024);
 
 		writer.WriteU32(childProcessId);
 		writer.WriteU32(info.rules->index);
@@ -1481,11 +1482,13 @@ namespace uba
 	{
 		// TODO: Should this traverse down? This is used for when child processes are terminated through job being closed. used by node.exe (building chrome)
 		u32 nativeProcessId = reader.ReadU32();
+		u32 exitCode = reader.ReadU32();
 		for (auto& child : m_childProcesses)
 		{
 			auto& process = *(ProcessImpl*)child.m_process;
 			if (process.m_nativeProcessId != nativeProcessId)
 				continue;
+			m_session.m_logger.Info(TC("CHILD TERMINATED. Exit code %u"), exitCode);
 			process.m_parentReportedExit = true;
 			return true;
 		}
@@ -1517,7 +1520,7 @@ namespace uba
 
 	bool ProcessImpl::HandleInputDependencies(BinaryReader& reader, BinaryWriter& writer)
 	{
-		UBA_ASSERT(m_startInfo.trackInputs);
+		UBA_ASSERTF(m_startInfo.trackInputs, TC("Got InputDependencies message even though track input dependencies was not set (%u)"), m_id);
 
 		if (u64 reserveSize = reader.Read7BitEncoded())
 			m_trackedInputs.reserve(m_trackedInputs.size() + reserveSize);
@@ -2081,15 +2084,15 @@ namespace uba
 				return UBA_EXIT_CODE(6);
 			}
 
-			//auto closeThreadHandle = MakeGuard([&]() { CloseHandle(pi.hThread); });
-			//auto closeProcessHandle = MakeGuard([&]() { CloseHandle(pi.hProcess); });
-
 			destroyAttr.Execute();
 
 			m_nativeProcessHandle = (ProcHandle)(u64)processInfo.hProcess;
 			m_nativeProcessId = processInfo.dwProcessId;
 			m_nativeThreadHandle = processInfo.hThread;
 
+			if (m_nativeProcessHandle == InvalidProcHandle) 
+				logger.Error(TC("Seems like DetourCreateProcessWithDllEx succeeded starting %s but returned a process handle that is invalid (0)?"), commandLine.c_str());
+
 			#if UBA_DEBUG
 			USHORT processMachine = 0;
 			USHORT nativeMachine = 0;
@@ -2168,7 +2171,8 @@ namespace uba
 
 			if (!DetourCopyPayloadToProcess((HANDLE)m_nativeProcessHandle, DetoursPayloadGuid, &payload, sizeof(payload)))
 			{
-				logger.Error(TC("Failed to copy payload to process %llu - %s (%s)"), m_nativeProcessHandle, m_realApplication.c_str(), LastErrorToText().data);
+				if (!m_cancelled && !m_parentReportedExit)
+					logger.Error(TC("Process with id %u failed to copy payload to %sprocess %llu - %s (%s)"), m_id, m_parentProcess?TC("child "):TC(""), m_nativeProcessHandle, m_realApplication.c_str(), LastErrorToText().data);
 				return UBA_EXIT_CODE(9);
 			}
 		}
@@ -2194,7 +2198,8 @@ namespace uba
 
 		if (!m_startInfo.startSuspended && ResumeThread(m_nativeThreadHandle) == -1)
 		{
-			logger.Error(TC("Failed to resume thread for"));//% ls. (% ls)", commandLine.c_str(), LastErrorToText().data);
+			LogEntryType logType = (m_cancelled || m_parentReportedExit) ? LogEntryType_Info : LogEntryType_Error;
+			logger.Logf(logType, TC("Failed to resume thread %llu for processid %u - %s. (%s)"), m_nativeThreadHandle, m_id, m_realApplication.c_str(), LastErrorToText().data);
 			return UBA_EXIT_CODE(11);
 		}
 
diff --git a/Engine/Source/Programs/UnrealBuildAccelerator/Common/Private/UbaSessionClient.cpp b/Engine/Source/Programs/UnrealBuildAccelerator/Common/Private/UbaSessionClient.cpp
--- a/Engine/Source/Programs/UnrealBuildAccelerator/Common/Private/UbaSessionClient.cpp
+++ b/Engine/Source/Programs/UnrealBuildAccelerator/Common/Private/UbaSessionClient.cpp
@@ -1016,7 +1016,7 @@ namespace uba
 			}
 		}
 
-		UBA_ASSERTF(false, TC("Code path not implemented.. should likely send message to server"));
+		UBA_ASSERTF(false, TC("NotImplemented - chmod on non-written file %s (should be sent message to server"), msg.fileName.data);
 		return true;
 	}
 
diff --git a/Engine/Source/Programs/UnrealBuildAccelerator/Common/Public/UbaEnvironment.h b/Engine/Source/Programs/UnrealBuildAccelerator/Common/Public/UbaEnvironment.h
--- a/Engine/Source/Programs/UnrealBuildAccelerator/Common/Public/UbaEnvironment.h
+++ b/Engine/Source/Programs/UnrealBuildAccelerator/Common/Public/UbaEnvironment.h
@@ -13,7 +13,7 @@ namespace uba
 	void CloseMutex(MutexHandle mutex);
 
 	enum ProcHandle : u64 {};
-	inline constexpr ProcHandle InvalidProcHandle = ((ProcHandle)(u64)-1);
+	inline constexpr ProcHandle InvalidProcHandle = (ProcHandle)0ull;
 	ProcHandle GetCurrentProcessHandle();
 
 	u32 GetEnvironmentVariableW(const tchar* name, tchar* buffer, u32 nSize);
diff --git a/Engine/Source/Programs/UnrealBuildAccelerator/Detours/Private/Windows/UbaDetoursFunctionsKernelBase.inl b/Engine/Source/Programs/UnrealBuildAccelerator/Detours/Private/Windows/UbaDetoursFunctionsKernelBase.inl
--- a/Engine/Source/Programs/UnrealBuildAccelerator/Detours/Private/Windows/UbaDetoursFunctionsKernelBase.inl
+++ b/Engine/Source/Programs/UnrealBuildAccelerator/Detours/Private/Windows/UbaDetoursFunctionsKernelBase.inl
@@ -210,7 +210,8 @@ HANDLE Detoured_CreateFileW(LPCWSTR lpFileName, DWORD dwDesiredAccess, DWORD dwS
 	HANDLE h = True_CreateFileW(lpFileName, dwDesiredAccess, dwShareMode, lpSecurityAttributes, dwCreationDisposition, dwFlagsAndAttributes, hTemplateFile);
 
 	#if UBA_CHECK_PATH_CONVERSION
-	UBA_ASSERT(t_rtlCalled);
+	bool doNotCheckRtlCalled = g_isRunningWine && (!lpFileName || !lpFileName[0] || dwCreationDisposition < CREATE_NEW || dwCreationDisposition > TRUNCATE_EXISTING);
+	UBA_ASSERTF(doNotCheckRtlCalled || t_rtlCalled, TC("Rtl path conversion not called fo rpath %s"), lpFileName);
 	#endif
 	t_createFileFileName = nullptr;
 	t_disallowDetour -= disallowDetour;
@@ -323,10 +324,9 @@ BOOL Detoured_RemoveDirectoryW(LPCWSTR lpPathName)
 
 	if (pathName.StartsWith(g_systemTemp.data))
 	{
+		DEBUG_LOG_TRUE(L"RemoveDirectoryW", L"(NODETOUR) %ls", lpPathName);
 		SuppressCreateFileDetourScope s; // TODO: Revisit this.. will not work remotely
-		BOOL res = True_RemoveDirectoryW(lpPathName);
-		DEBUG_LOG_TRUE(L"RemoveDirectoryW", L"%ls -> %ls", lpPathName, ToString(res));
-		return res;
+		return True_RemoveDirectoryW(lpPathName);
 	}
 
 	DevirtualizePath(pathName);
@@ -2222,6 +2222,8 @@ BOOL Detoured_SetFileInformationByHandle(HANDLE hFile, FILE_INFO_BY_HANDLE_CLASS
 	}
 	else if (FileInformationClass == FileEndOfFileInfo)
 	{
+		// TODO: THIS CAN BE REMOVED!!
+
 		if (MemoryFile* mf = fi.memoryFile)
 		{
 			auto& info = *(FILE_END_OF_FILE_INFO*)lpFileInformation;
@@ -3086,9 +3088,13 @@ BOOL Detoured_CreateProcessW(LPCWSTR lpApplicationName, LPWSTR lpCommandLine, LP
 		writer.Flush();
 		BinaryReader reader;
 		processId = reader.ReadU32();
-		UBA_ASSERTF(processId > 0, L"Failed to create process %s", originalCmd.data);
 		if (!processId)
+		{
+			pcs.Leave();
+			UBA_ASSERTF(processId > 0, L"Failed to create process %s", originalCmd.data);
+			SetLastError(ERROR_ACCESS_DENIED); // This is wrong but there should be info in the log
 			return FALSE;
+		}
 
 		reader.Skip(sizeof(u32)); // Rules index
 
@@ -3209,16 +3215,27 @@ BOOL Detoured_CreateProcessW(LPCWSTR lpApplicationName, LPWSTR lpCommandLine, LP
 		}
 	}
 
+	#if UBA_DEBUG
+	bool debugWineCalls = false;//application.EndsWith(L"node.exe");
+	if (debugWineCalls && !lpEnvironment)
+		lpEnvironment = True_GetEnvironmentStringsW();
+	#endif
+
 	// Sometimes programs decides to change environment variables such as TEMP/TMP/SystemRoot etc which will break when running remotely
 	Vector<tchar> tempEnvironment;
 	if (lpEnvironment)
 	{
-		auto addChars = [&](const tchar* text, u64 len)
+		u32 reserveSize = 1;
+		for(StringView env = ToView((tchar*)lpEnvironment); env.count; env = ToView(env.data + env.count + 1))
+			reserveSize += env.count+1;
+		tempEnvironment.reserve(reserveSize);
+
+		auto addChars = [&](StringView text)
 			{
 				u64 destPos = tempEnvironment.size();
-				tempEnvironment.resize(tempEnvironment.size() + len);
+				tempEnvironment.resize(tempEnvironment.size() + text.count);
 				tchar* dest = tempEnvironment.data() + destPos;
-				memcpy(dest, text, len*sizeof(tchar));
+				memcpy(dest, text.data, text.count*sizeof(tchar));
 			};
 		UBA_ASSERT(dwCreationFlags & CREATE_UNICODE_ENVIRONMENT);
 		for(StringView env = ToView((tchar*)lpEnvironment); env.count; env = ToView(env.data + env.count + 1))
@@ -3228,15 +3245,20 @@ BOOL Detoured_CreateProcessW(LPCWSTR lpApplicationName, LPWSTR lpCommandLine, LP
 			StringView name = env.SubStr(0, equalPos);
 			if (name.Equals(TCV("TEMP")) || name.Equals(TCV("TMP")))
 			{
-				addChars(name.data, name.count+1);
-				addChars(g_systemTemp.data, g_systemTemp.count+1);
+				addChars({name.data, name.count+1});
+				addChars({g_systemTemp.data, g_systemTemp.count+1});
 			}
 			else
 			{
-				addChars(env.data, env.count+1);
+				addChars({env.data, env.count+1});
 			}
 		}
 
+		#if UBA_DEBUG
+		if (debugWineCalls)
+			addChars(TCV("WINEDEBUG=,+file\0")); //"WINEDEBUG=+relay,+process,+file\0"
+		#endif
+
 		tempEnvironment.push_back(0);
 		lpEnvironment = tempEnvironment.data();
 	}
@@ -3316,8 +3338,7 @@ BOOL Detoured_CreateProcessW(LPCWSTR lpApplicationName, LPWSTR lpCommandLine, LP
 	UBA_ASSERTF(res, L"Failed to spawn process %ls (Error code: %u)", finalCommandLine, lastError);
 
 	HANDLE trueHandle = lpProcessInformation->hProcess;
-
-	if (!res || trueHandle == INVALID_HANDLE_VALUE)
+	if (!res || trueHandle == NULL)
 	{
 		DEBUG_LOG_DETOURED(L"CreateProcessW", L"FAILED");
 		return FALSE;
@@ -3386,6 +3407,15 @@ BOOL Detoured_TerminateProcess(HANDLE hProcess, UINT uExitCode)
 			SendExitMessage((DWORD)uExitCode, GetTime());
 			PostDeinit();
 		}
+		else
+		{
+			DWORD pid = GetProcessId(hProcess);
+			DEBUG_LOG(L"PROCESS %u WILL BE TERMINATED", pid);
+			RPC_MESSAGE(ExitChildProcess, detach) // Let it be part of detach timer for now
+			writer.WriteU32(pid);
+			writer.WriteU32(uExitCode);
+			writer.Flush();
+		}
 	}
 
 	return True_TerminateProcess(hProcess, uExitCode);
@@ -3410,6 +3440,7 @@ BOOL Detoured_AssignProcessToJobObject(HANDLE hJob, HANDLE hProcess)
 				DEBUG_LOG(L"PROCESS %u WILL BE TERMINATED from job %llu", pid, u64(hJob));
 				RPC_MESSAGE(ExitChildProcess, detach) // Let it be part of detach timer for now
 				writer.WriteU32(pid);
+				writer.WriteU32(0xC000010A);
 				writer.Flush();
 			}
 		}
diff --git a/Engine/Source/Programs/UnrealBuildAccelerator/Detours/Private/Windows/UbaDetoursFunctionsNtDll.inl b/Engine/Source/Programs/UnrealBuildAccelerator/Detours/Private/Windows/UbaDetoursFunctionsNtDll.inl
--- a/Engine/Source/Programs/UnrealBuildAccelerator/Detours/Private/Windows/UbaDetoursFunctionsNtDll.inl
+++ b/Engine/Source/Programs/UnrealBuildAccelerator/Detours/Private/Windows/UbaDetoursFunctionsNtDll.inl
@@ -63,6 +63,8 @@ StringBuffer<32> ToString(NTSTATUS s)
 		return res.Append(L"STATUS_SHARING_VIOLATION");
 	if (s == STATUS_ACCESS_DENIED)
 		return res.Append(L"STATUS_ACCESS_DENIED");
+	if (s == STATUS_DIRECTORY_NOT_EMPTY)
+		return res.Append(L"STATUS_DIRECTORY_NOT_EMPTY");
 	return res.Appendf(L"Error (0x%x)", u32(s));
 }
 #endif
@@ -649,12 +651,15 @@ NTSTATUS Detoured_NtSetInformationFile(HANDLE FileHandle, PIO_STATUS_BLOCK IoSta
 			UBA_ASSERT(trueHandle != INVALID_HANDLE_VALUE);
 			TimerScope ts(g_kernelStats.setFileInfo);
 			auto res = True_NtSetInformationFile(trueHandle, IoStatusBlock, FileInformation, Length, FileInformationClass);
-			DEBUG_LOG_TRUE(L"NtSetInformationFile", L"(%s) %llu (%ls) -> %ls", ToString(FileInformationClass).data, uintptr_t(FileHandle), HandleToName(FileHandle), ToString(res).data);
+			DEBUG_LOG_TRUE(L"NtSetInformationFile", L"(%s) %llu (%s) -> %s", ToString(FileInformationClass).data, uintptr_t(FileHandle), HandleToName(FileHandle), ToString(res).data);
 			return res;
 		};
 
 	if (!isDetouredHandle(FileHandle))
+	{
+		UBA_ASSERT(!isListDirectoryHandle(FileHandle));
 		return runTrue();
+	}
 
 	DetouredHandle& h = asDetouredHandle(FileHandle);
 	trueHandle = h.trueHandle;
@@ -2492,11 +2497,13 @@ NTSTATUS Detoured_RtlDosPathNameToNtPathName_U_WithStatus(PCWSTR dos_path, PUNIC
 	return res;
 }
 
-NTSTATUS Detoured_RtlDosPathNameToRelativeNtPathName_U_WithStatus(PCWSTR DosFileName, PUNICODE_STRING NtFileName, PWSTR *FilePart, RTL_RELATIVE_NAME_U* RelativeName)
+NTSTATUS Detoured_RtlDosPathNameToRelativeNtPathName_U_WithStatus(PCWSTR DosFileName, PUNICODE_STRING NtFileName, PWSTR* FilePart, RTL_RELATIVE_NAME_U* RelativeName)
 {
 	#if UBA_CHECK_PATH_CONVERSION
 	t_rtlCalled = true;
 	#endif
+	//UBA_ASSERT(!FilePart);
+	//UBA_ASSERT(!RelativeName);
 	return True_RtlDosPathNameToRelativeNtPathName_U_WithStatus(DosFileName, NtFileName, FilePart, RelativeName);
 }
 
diff --git a/Engine/Source/Programs/UnrealBuildAccelerator/Visualizer/Private/UbaVisualizer.cpp b/Engine/Source/Programs/UnrealBuildAccelerator/Visualizer/Private/UbaVisualizer.cpp
--- a/Engine/Source/Programs/UnrealBuildAccelerator/Visualizer/Private/UbaVisualizer.cpp
+++ b/Engine/Source/Programs/UnrealBuildAccelerator/Visualizer/Private/UbaVisualizer.cpp
@@ -2614,7 +2614,7 @@ namespace uba
 		if (hasExited && process.exitCode != 0)
 			out.Info(TC("  ExitCode:  %u"), process.exitCode);
 
-		if (process.stop != ~u64(0))
+		if (process.stop != ~u64(0) && process.stats.size())
 		{
 			out.Info(TC(""));
 
