[UBA]
* Disabled SignalObjectAndWait on wine because there is a bug somewhere causing it to return too early (Can't repro bug if doing SetEvent/WaitForSingleObject
* Moved detoured TerminateProcess to NtTerminateProcess to catch all calls
* Added code to write/read guard in all messages between host/detoured process to be able to find race conditions or other problems (such as SignalObjectAndWait on wine). Enabled with UBA_PROTOCOL_GUARD
* Added more handling of assert/callstack resolving happening inside rpc message
* Changed posix code to use RPC_MESSAGE
* Added more info to asserts

diff --git a/Engine/Source/Programs/UnrealBuildAccelerator/Cli/Private/UbaCli.cpp b/Engine/Source/Programs/UnrealBuildAccelerator/Cli/Private/UbaCli.cpp
--- a/Engine/Source/Programs/UnrealBuildAccelerator/Cli/Private/UbaCli.cpp
+++ b/Engine/Source/Programs/UnrealBuildAccelerator/Cli/Private/UbaCli.cpp
@@ -13,6 +13,7 @@
 #include "UbaNetworkBackendTcp.h"
 #include "UbaPathUtils.h"
 #include "UbaPlatform.h"
+#include "UbaProcess.h"
 #include "UbaProtocol.h"
 #include "UbaRootPaths.h"
 #include "UbaScheduler.h"
@@ -1749,6 +1750,16 @@ namespace uba
 						return false;
 					bool res = false;
 
+					// Don't want to retry
+					if (commandType == CommandType_Remote)
+					{
+						sessionServer.SetRemoteProcessReturnedEvent([&](Process& process)
+							{
+								if (sessionServer.WasEverConnected())
+									process.Cancel();
+							});
+					}
+
 					for (u32 procIt=0; procIt!=procLoopCount; ++procIt)
 					{
 						switch (commandType)
diff --git a/Engine/Source/Programs/UnrealBuildAccelerator/Common/Private/UbaProcess.cpp b/Engine/Source/Programs/UnrealBuildAccelerator/Common/Private/UbaProcess.cpp
--- a/Engine/Source/Programs/UnrealBuildAccelerator/Common/Private/UbaProcess.cpp
+++ b/Engine/Source/Programs/UnrealBuildAccelerator/Common/Private/UbaProcess.cpp
@@ -691,6 +691,14 @@ namespace uba
 					u64 startTime = GetTime();
 					BinaryReader reader(comMemory, 0, comMemorySize);
 					BinaryWriter writer(comMemory, 0, comMemorySize);
+
+					#if UBA_PROTOCOL_GUARD
+					u32 deadbeef = reader.ReadU32();
+					if (deadbeef != 0xdeadbeef)
+						m_session.m_logger.Error(TC("SHOULD NOT HAPPEN%s"), deadbeef == 0xbaadf00d ? TC(" (Seems like nothing was written?)") : TC(""));
+					writer.WriteU32(0xbaadf00d);
+					#endif
+
 					loop = HandleMessage(reader, writer);
 					SetWritten();
 					m_processStats.hostTotalTime += GetTime() - startTime;
diff --git a/Engine/Source/Programs/UnrealBuildAccelerator/Common/Private/UbaSessionClient.cpp b/Engine/Source/Programs/UnrealBuildAccelerator/Common/Private/UbaSessionClient.cpp
--- a/Engine/Source/Programs/UnrealBuildAccelerator/Common/Private/UbaSessionClient.cpp
+++ b/Engine/Source/Programs/UnrealBuildAccelerator/Common/Private/UbaSessionClient.cpp
@@ -1524,6 +1524,7 @@ namespace uba
 		#if PLATFORM_WINDOWS
 		AddEnvironmentVariableNoLock(TCV("TEMP"), m_tempPath);
 		AddEnvironmentVariableNoLock(TCV("TMP"), m_tempPath);
+		//AddEnvironmentVariableNoLock(TCV("WINEDEBUG"), TCV("+process"));
 		#else
 		AddEnvironmentVariableNoLock(TCV("TMPDIR"), m_tempPath);
 		#endif
diff --git a/Engine/Source/Programs/UnrealBuildAccelerator/Common/Private/UbaSessionServer.cpp b/Engine/Source/Programs/UnrealBuildAccelerator/Common/Private/UbaSessionServer.cpp
--- a/Engine/Source/Programs/UnrealBuildAccelerator/Common/Private/UbaSessionServer.cpp
+++ b/Engine/Source/Programs/UnrealBuildAccelerator/Common/Private/UbaSessionServer.cpp
@@ -871,6 +871,12 @@ namespace uba
 		return ph;
 	}
 
+	bool SessionServer::WasEverConnected()
+	{
+		SCOPED_CRITICAL_SECTION(m_remoteProcessAndSessionLock, lock);
+		return !m_clientSessions.empty();
+	}
+
 	void SessionServer::OnDisconnected(const Guid& clientUid, u32 clientId)
 	{
 		u32 returnCount = 0;
diff --git a/Engine/Source/Programs/UnrealBuildAccelerator/Common/Public/UbaSessionServer.h b/Engine/Source/Programs/UnrealBuildAccelerator/Common/Public/UbaSessionServer.h
--- a/Engine/Source/Programs/UnrealBuildAccelerator/Common/Public/UbaSessionServer.h
+++ b/Engine/Source/Programs/UnrealBuildAccelerator/Common/Public/UbaSessionServer.h
@@ -95,6 +95,7 @@ namespace uba
 		bool GetRemoteTraceEnabled();
 
 		ProcessHandle GetNewestLocalProcess();
+		bool WasEverConnected();
 
 	protected:
 		struct ClientSession;
diff --git a/Engine/Source/Programs/UnrealBuildAccelerator/Core/Public/UbaBinaryReaderWriter.h b/Engine/Source/Programs/UnrealBuildAccelerator/Core/Public/UbaBinaryReaderWriter.h
--- a/Engine/Source/Programs/UnrealBuildAccelerator/Core/Public/UbaBinaryReaderWriter.h
+++ b/Engine/Source/Programs/UnrealBuildAccelerator/Core/Public/UbaBinaryReaderWriter.h
@@ -7,6 +7,8 @@
 
 namespace uba
 {
+	enum ProcessCommunication { ProcessCommunication_Write };
+
 	struct BinaryWriter
 	{
 		inline void WriteByte(u8 value);
@@ -45,8 +47,8 @@ namespace uba
 		inline void Reset(u8* data, u64 offset = 0, u64 capacity = InvalidValue) { m_begin = data; m_pos = data + offset; m_end = data + capacity; }
 		inline BinaryWriter(u8* data, u64 offset = 0, u64 capacity = InvalidValue) { Reset(data, offset, capacity); }
 
-		void Flush(bool waitOnResponse = true);
-		BinaryWriter();
+		BinaryReader Flush(bool waitOnResponse = true);
+		BinaryWriter(ProcessCommunication);
 
 	protected:
 		u8* m_begin;
@@ -100,14 +102,15 @@ namespace uba
 		inline BinaryReader(const u8* data, u64 offset = 0, u64 size = InvalidValue) { m_begin = data; m_pos = data + offset; m_end = data + size; }
 		inline BinaryReader(const Vector<u8>& data, u64 offset = 0) { m_begin = data.data(); m_pos = m_begin + offset; m_end = m_begin + data.size(); }
 
-		BinaryReader();
-
 	protected:
 		inline u64 InternalReadString(tchar* str, u64 charLen);
 		inline void InternalSkipString(u64 charLen);
 		const u8* m_begin;
 		const u8* m_pos;
 		const u8* m_end;
+	private:
+		friend struct BinaryWriter;
+		BinaryReader() = delete;
 	};
 
 	template<u32 Capacity>
diff --git a/Engine/Source/Programs/UnrealBuildAccelerator/Core/Public/UbaDefinitions.h b/Engine/Source/Programs/UnrealBuildAccelerator/Core/Public/UbaDefinitions.h
--- a/Engine/Source/Programs/UnrealBuildAccelerator/Core/Public/UbaDefinitions.h
+++ b/Engine/Source/Programs/UnrealBuildAccelerator/Core/Public/UbaDefinitions.h
@@ -15,6 +15,8 @@
 // Enable to add logging of file mappings
 #define UBA_DEBUG_FILE_MAPPING 0
 
+#define UBA_PROTOCOL_GUARD 0
+
 namespace uba
 {
 	inline constexpr bool EventIsNative = bool(PLATFORM_WINDOWS && !UBA_TEST_WINDOWS_SEMAPHORES);
diff --git a/Engine/Source/Programs/UnrealBuildAccelerator/Detours/Private/PosixOS/UbaDetoursFunctionsPosix.cpp b/Engine/Source/Programs/UnrealBuildAccelerator/Detours/Private/PosixOS/UbaDetoursFunctionsPosix.cpp
--- a/Engine/Source/Programs/UnrealBuildAccelerator/Detours/Private/PosixOS/UbaDetoursFunctionsPosix.cpp
+++ b/Engine/Source/Programs/UnrealBuildAccelerator/Detours/Private/PosixOS/UbaDetoursFunctionsPosix.cpp
@@ -965,14 +965,10 @@ UBA_EXPORT int UBA_WRAPPER(mkdir)(const char* path, mode_t mode)
 	StringKey pathNameKey = ToFilenameKey(pathName);
 
 	{
-		TimerScope ts(g_stats.createFile);
-		SCOPED_FUTEX(g_communicationLock, pcs);
-		BinaryWriter writer;
-		writer.WriteByte(MessageType_CreateDirectory);
+		RPC_MESSAGE(CreateDirectory, createFile);
 		writer.WriteStringKey(pathNameKey);
 		writer.WriteString(pathName);
-		writer.Flush();
-		BinaryReader reader;
+		BinaryReader reader = writer.Flush();
 		res = reader.ReadBool() ? 0 : -1;
 		errorCode = reader.ReadU32();
 		directoryTableSize = reader.ReadU32();
@@ -1005,14 +1001,10 @@ UBA_EXPORT int UBA_WRAPPER(rmdir)(const char* path)
 	StringKey pathNameKey = ToFilenameKey(pathName);
 
 	{
-		TimerScope ts(g_stats.deleteFile);
-		SCOPED_FUTEX(g_communicationLock, pcs);
-		BinaryWriter writer;
-		writer.WriteByte(MessageType_RemoveDirectory);
+		RPC_MESSAGE(RemoveDirectory, deleteFile);
 		writer.WriteStringKey(pathNameKey);
 		writer.WriteString(pathName);
-		writer.Flush();
-		BinaryReader reader;
+		BinaryReader reader = writer.Flush();;
 		res = reader.ReadBool();
 		errorCode = reader.ReadU32();
 		directoryTableSize = reader.ReadU32();
@@ -1798,17 +1790,13 @@ UBA_EXPORT int UBA_WRAPPER(rename)(const char* oldpath, const char* newpath)
 	u32 errorCode;
 	bool result;
 	{
-		TimerScope ts(g_stats.moveFile);
-		SCOPED_FUTEX(g_communicationLock, pcs);
-		BinaryWriter writer;
-		writer.WriteByte(MessageType_MoveFile);
+		RPC_MESSAGE(MoveFile, moveFile);
 		writer.WriteStringKey(oldKey);
 		writer.WriteString(fixedOldPath);
 		writer.WriteStringKey(newKey);
 		writer.WriteString(fixedNewPath);
 		writer.WriteU32(0);// dwFlags);
-		writer.Flush();
-		BinaryReader reader;
+		BinaryReader reader = writer.Flush();
 		result = reader.ReadBool();
 		errorCode = reader.ReadU32();
 		directoryTableSize = reader.ReadU32();
@@ -1862,15 +1850,11 @@ UBA_EXPORT int UBA_WRAPPER(chmod)(const char* pathname, mode_t mode)
 	StringKey key = ToFilenameKey(fixedName);
 	u32 errorCode;
 	{
-		TimerScope ts(g_stats.chmod);
-		SCOPED_FUTEX(g_communicationLock, pcs);
-		BinaryWriter writer;
-		writer.WriteByte(MessageType_Chmod);
+		RPC_MESSAGE(Chmod, chmod);
 		writer.WriteStringKey(key);
 		writer.WriteString(fixedName);
 		writer.WriteU32(mode);
-		writer.Flush();
-		BinaryReader reader;
+		BinaryReader reader = writer.Flush();
 		errorCode = reader.ReadU32();
 		//DEBUG_LOG_PIPE(L"MoveFile", L"%ls to %ls", lpExistingFileName, lpNewFileName);
 	}
@@ -1929,8 +1913,7 @@ UBA_EXPORT int UBA_WRAPPER(symlink)(const char* path1, const char* path2)
 		writer.WriteString(fromPath);
 		writer.WriteStringKey(toKey);
 		writer.WriteString(toPath);
-		writer.Flush();
-		BinaryReader reader;
+		BinaryReader reader = writer.Flush();
 		reader.ReadString(newFromName);
 		reader.ReadString(newToName);
 		closeId = reader.ReadU32();
@@ -1997,15 +1980,11 @@ int Shared_DeleteFile(const char* funcName, const char* pathname)
 	u32 errorCode;
 	{
 		u32 closeId = 0;
-		TimerScope ts(g_stats.deleteFile);
-		SCOPED_FUTEX(g_communicationLock, pcs);
-		BinaryWriter writer;
-		writer.WriteByte(MessageType_DeleteFile);
+		RPC_MESSAGE(DeleteFile, deleteFile);
 		writer.WriteString(fixedName);
 		writer.WriteStringKey(fileNameKey);
 		writer.WriteU32(closeId);
-		writer.Flush();
-		BinaryReader reader;
+		BinaryReader reader = writer.Flush();
 		result = reader.ReadBool();
 		errorCode = reader.ReadU32();
 		directoryTableSize = reader.ReadU32();
@@ -2096,18 +2075,14 @@ bool ExecuteHostRun(StringBufferBase& out, const char* const* argv, bool removeL
 		FlattenArgs(command, argv);
 		DEBUG_LOG_DETOURED("HostRun", "%s", command.data)
 	}
-	TimerScope ts(g_stats.getFullFileName);
-	SCOPED_FUTEX(g_communicationLock, pcs);
-	BinaryWriter writer;
-	writer.WriteByte(MessageType_HostRun);
+	RPC_MESSAGE(HostRun, getFullFileName);
 	u16& size = *(u16*)writer.AllocWrite(2);
 	u64 pos = writer.GetPosition();
 	for (u32 i = 0; argv[i]; ++i)
 		writer.WriteString(argv[i]);
 	size = u16(writer.GetPosition() - pos);
-	writer.Flush();
 
-	BinaryReader reader;
+	BinaryReader reader = writer.Flush();
 	bool success = reader.ReadBool();
 	reader.ReadString(out);
 	
@@ -2256,8 +2231,7 @@ int shared_posix_spawn(pid_t* pid, const char* path, const posix_spawn_file_acti
 	#endif
 	
 	{
-		TimerScope ts(g_stats.createProcess);
-		SCOPED_FUTEX(g_communicationLock, pcs);
+		RPC_MESSAGE(CreateProcess, createProcess);
 
 		const char* pwd = "";
 		//while (char* env = envp[i++])
@@ -2270,16 +2244,13 @@ int shared_posix_spawn(pid_t* pid, const char* path, const posix_spawn_file_acti
 		//	}
 		//}
 
-		BinaryWriter writer;
-		writer.WriteByte(MessageType_CreateProcess);
 		writer.WriteString(path); // Application
 		writer.WriteLongString(cmdLineWithoutApplication, CommunicationMemSize - 4096);
 		writer.WriteString(pwd); // Current dir
 		writer.WriteBool(false); // Start suspended
 		writer.WriteBool(true); // Is child
-		writer.Flush();
+		BinaryReader reader = writer.Flush();
 
-		BinaryReader reader;
 		processId = reader.ReadU32();
 
 		if (!processId) // Can happen if session client got disconnected from server session
@@ -2346,10 +2317,7 @@ int shared_posix_spawn(pid_t* pid, const char* path, const posix_spawn_file_acti
 	bool success = res == 0;
 
 	{
-		TimerScope ts(g_stats.createProcess);
-		SCOPED_FUTEX(g_communicationLock, pcs);
-		BinaryWriter writer;
-		writer.WriteByte(MessageType_StartProcess);
+		RPC_MESSAGE(StartProcess, createProcess);
 		writer.WriteU32(processId);
 		writer.WriteBool(success);
 		writer.WriteU32(res);
@@ -2564,10 +2532,7 @@ int Internal_execve(const char* pathname, char* const _Nullable argv[], char* co
 	}
 
 	{
-		TimerScope ts(g_stats.createProcess);
-		SCOPED_FUTEX(g_communicationLock, pcs);
-		BinaryWriter writer;
-		writer.WriteByte(MessageType_ExitChildProcess);
+		RPC_MESSAGE(ExitChildProcess, createProcess);
 		writer.WriteU32(pid);
 	}
 
@@ -2935,12 +2900,8 @@ namespace uba
 		u32 overlayFilesSize;
 
 		{
-			TimerScope ts(g_stats.init);
-			SCOPED_FUTEX(g_communicationLock, pcs);
-			BinaryWriter writer;
-			writer.WriteByte(MessageType_Init);
-			writer.Flush();
-			BinaryReader reader;
+			RPC_MESSAGE(Init, init);
+			BinaryReader reader = writer.Flush();
 
 			g_processId = reader.ReadU32();
 			g_isChild = reader.ReadBool();
@@ -3065,8 +3026,7 @@ namespace uba
 
 		SendInput();
 
-		BinaryWriter writer;
-		writer.WriteByte(MessageType_Exit);
+		RPC_MESSAGE(Exit, log);
 		writer.WriteU32(0); // Exit code
 		writer.WriteString(""); // Log name
 		g_stats.Write(writer);
@@ -3124,8 +3084,7 @@ namespace uba
 			if (!allowTerminate)
 				return;
 
-			BinaryWriter writer;
-			writer.WriteByte(MessageType_Exit);
+			RPC_MESSAGE(Exit, log);
 			writer.WriteU32(terminateCode); // Exit code
 			writer.WriteString(""); // Log name
 			g_stats.Write(writer);
@@ -3172,16 +3131,13 @@ extern "C"
 		u32 overlayFilesSize = 0;
 		u32 directoryTableSize = 0;
 		{
-			SCOPED_FUTEX(g_communicationLock, pcs);
-			BinaryWriter writer;
-			writer.WriteByte(MessageType_GetNextProcess);
+			RPC_MESSAGE(GetNextProcess, log);
 			writer.WriteU32(prevExitCode);
 			writer.WriteU32(timeOutMs);
 			g_stats.Write(writer);
 			g_kernelStats.Write(writer);
 
-			writer.Flush();
-			BinaryReader reader;
+			BinaryReader reader = writer.Flush();
 			newProcess = reader.ReadBool();
 			*outShouldExit = reader.ReadBool();
 			if (newProcess)
diff --git a/Engine/Source/Programs/UnrealBuildAccelerator/Detours/Private/PosixOS/UbaDetoursMainPosix.cpp b/Engine/Source/Programs/UnrealBuildAccelerator/Detours/Private/PosixOS/UbaDetoursMainPosix.cpp
--- a/Engine/Source/Programs/UnrealBuildAccelerator/Detours/Private/PosixOS/UbaDetoursMainPosix.cpp
+++ b/Engine/Source/Programs/UnrealBuildAccelerator/Detours/Private/PosixOS/UbaDetoursMainPosix.cpp
@@ -210,19 +210,28 @@ namespace uba
 {
 	constexpr u32 WritableMemSize = CommunicationMemSize - sizeof(SharedEvent) * 3;
 
-	BinaryWriter::BinaryWriter()
+	BinaryWriter::BinaryWriter(ProcessCommunication)
 	{
 		m_begin = g_messageMappingMem;
 		m_pos = m_begin;
 		m_end = m_begin + WritableMemSize;
+
+		#if UBA_PROTOCOL_GUARD
+		WriteU32(0xdeadbeef);
+		#endif
 	}
 
-	void BinaryWriter::Flush(bool waitOnResponse)
+	BinaryReader BinaryWriter::Flush(bool waitOnResponse)
 	{
+		#if UBA_PROTOCOL_GUARD
+		if (BinaryReader(g_messageMappingMem, 0, WritableMemSize).ReadU32() != 0xdeadbeef)
+			exit(2345);
+		#endif
+
 		g_writeEvent->Set();
 
 		if (!waitOnResponse)
-			return;
+			return BinaryReader(nullptr, 0, 0);
 
 		TimerScope ts(g_stats.waitOnResponse);
 		do
@@ -243,12 +252,19 @@ namespace uba
 				exit(1339);
 			}
 		} while (true);
-	}
 
-	BinaryReader::BinaryReader()
-	{
-		m_begin = g_messageMappingMem;
-		m_pos = m_begin;
-		m_end = m_begin + WritableMemSize;
+		BinaryReader reader(g_messageMappingMem, 0, WritableMemSize);
+
+		#if UBA_PROTOCOL_GUARD
+		u32 baadfood = reader.ReadU32();
+		if (baadfood != 0xbaadf00d)
+			exit(2345);
+		BinaryWriter(g_messageMappingMem, 0, CommunicationMemSize).WriteU32(0xcafebabe);
+		#endif
+
+		m_begin = nullptr;
+		m_pos = 0;
+		m_end = nullptr;
+		return reader;
 	}
 }
diff --git a/Engine/Source/Programs/UnrealBuildAccelerator/Detours/Private/UbaDetoursFileMappingTable.cpp b/Engine/Source/Programs/UnrealBuildAccelerator/Detours/Private/UbaDetoursFileMappingTable.cpp
--- a/Engine/Source/Programs/UnrealBuildAccelerator/Detours/Private/UbaDetoursFileMappingTable.cpp
+++ b/Engine/Source/Programs/UnrealBuildAccelerator/Detours/Private/UbaDetoursFileMappingTable.cpp
@@ -91,8 +91,7 @@ namespace uba
 		writer.WriteString(fileName);
 		writer.WriteStringKey(fileNameKey);
 		writer.WriteByte(access);
-		writer.Flush();
-		BinaryReader reader;
+		BinaryReader reader = writer.Flush();
 		reader.ReadString(outNewName, newNameCapacity);
 		outSize = reader.ReadU64();
 		outCloseId = reader.ReadU32();
@@ -113,8 +112,7 @@ namespace uba
 		RPC_MESSAGE(CheckRemapping, createFile)
 		writer.WriteString(fileName);
 		writer.WriteStringKey(fileNameKey);
-		writer.Flush();
-		BinaryReader reader;
+		BinaryReader reader = writer.Flush();
 		u32 mappedFileTableSize = reader.ReadU32();
 		pcs.Leave();
 		g_mappedFileTable.ParseNoLock(mappedFileTableSize);
@@ -128,8 +126,7 @@ namespace uba
 			RPC_MESSAGE(ListDirectory, listDirectory)
 			writer.WriteString(dirName, dirNameLen);
 			writer.WriteStringKey(dirKey);
-			writer.Flush();
-			BinaryReader reader;
+			BinaryReader reader = writer.Flush();
 			directoryTableSize = u32(reader.ReadU32());
 			u64 to = reader.ReadU32();
 			tableOffset = to == InvalidTableOffset ? ~u32(0) : (u32)to;
@@ -168,8 +165,7 @@ namespace uba
 			}
 			else
 				writer.WriteStringKey(StringKeyZero);
-			writer.Flush();
-			BinaryReader reader;
+			BinaryReader reader = writer.Flush();
 			directoryTableSize = reader.ReadU32();
 			pcs.Leave();
 			DEBUG_LOG_PIPE(L"CloseFile", L"");
@@ -244,8 +240,7 @@ namespace uba
 		{
 			RPC_MESSAGE_NO_LOCK(GetWrittenFiles, updateTables)
 			writer.WriteBool(isInit);
-			writer.Flush();
-			BinaryReader reader;
+			BinaryReader reader = writer.Flush();
 			if (UpdateWrittenFiles(reader))
 				break;
 		}
@@ -260,8 +255,7 @@ namespace uba
 		{
 			RPC_MESSAGE_NO_LOCK(UpdateTables, updateTables)
 			writer.WriteBool(false);
-			writer.Flush();
-			BinaryReader reader;
+			BinaryReader reader = writer.Flush();
 			directoryTableSize = reader.ReadU32();
 			fileMappingTableSize = reader.ReadU32();
 			done = UpdateWrittenFiles(reader);
@@ -288,8 +282,7 @@ namespace uba
 		RPC_MESSAGE(CreateSharedMemory, createSharedMemory)
 		writer.WriteU64(initialSize);
 		writer.WriteString(hint);
-		writer.Flush();
-		BinaryReader reader;
+		BinaryReader reader = writer.Flush();
 		SharedMemoryHandle handle{ reader.Read7BitEncoded() };
 
 		out.Init(allocatorHandle, capacity);
@@ -305,8 +298,7 @@ namespace uba
 		RPC_MESSAGE(CommitSharedMemory, commitSharedMemory)
 		writer.WriteU64(handle.ToU64());
 		writer.WriteU64(size);
-		writer.Flush();
-		BinaryReader reader;
+		BinaryReader reader = writer.Flush();
 		view.AddRequestedMemory(allocatorHandle, reader);
 	}
 
@@ -319,8 +311,7 @@ namespace uba
 		writer.WriteU64(handle.ToU64());
 		writer.WriteU64(offset);
 		writer.WriteU64(size);
-		writer.Flush();
-		BinaryReader reader;
+		BinaryReader reader = writer.Flush();
 		u64 committed = size;
 		if (!committed)
 			committed = reader.Read7BitEncoded();
@@ -335,8 +326,7 @@ namespace uba
 		writer.WriteU64(handle.ToU64());
 		writer.WriteU64(0);
 		writer.WriteU64(0);
-		writer.Flush();
-		BinaryReader reader;
+		BinaryReader reader = writer.Flush();
 		u64 committed = reader.Read7BitEncoded(); (void)committed;
 		UBA_ASSERT(committed <= capacity);
 		view.Init(allocatorHandle, capacity);
@@ -425,8 +415,7 @@ namespace uba
 				for (auto i=loaderPaths; *i; ++i)
 					writer.WriteString(*i);
 			bytes = u16(writer.GetPosition() - pos);
-			writer.Flush();
-			BinaryReader reader;
+			BinaryReader reader = writer.Flush();
 			reader.ReadString(tempBuf);
 			if (useVirtualName)
 			{
@@ -479,8 +468,7 @@ namespace uba
 				for (auto i=loaderPaths; *i; ++i)
 					writer.WriteString(*i);
 			bytes = u16(writer.GetPosition() - pos);
-			writer.Flush();
-			BinaryReader reader;
+			BinaryReader reader = writer.Flush();
 			reader.ReadString(outReal);
 			reader.ReadString(outVirtual);
 			mappedFileTableSize = reader.ReadU32();
diff --git a/Engine/Source/Programs/UnrealBuildAccelerator/Detours/Private/UbaDetoursShared.cpp b/Engine/Source/Programs/UnrealBuildAccelerator/Detours/Private/UbaDetoursShared.cpp
--- a/Engine/Source/Programs/UnrealBuildAccelerator/Detours/Private/UbaDetoursShared.cpp
+++ b/Engine/Source/Programs/UnrealBuildAccelerator/Detours/Private/UbaDetoursShared.cpp
@@ -407,15 +407,34 @@ namespace uba
 		#endif
 	}
 
+	bool TakeLockForRpc()
+	{
+		for (u32 i=0; i!=5; ++i)
+		{
+			if (g_communicationLock.TryEnter())
+				return true;
+			Sleep(100);
+		}
+		DEBUG_LOG(TC("TakeLockForRpc failed"));
+		return false;
+	}
+
 	void Rpc_WriteLog(const tchar* text, u64 textCharLength, bool printInSession, bool isError)
 	{
 		DEBUG_LOG(TC("LOG  %.*s"), u32(textCharLength), text); // TODO: Investigate, deadlocks on non-windows
 		// DEBUG_LOG(TC("LOG [%7u] %.*s"), GetCurrentThreadId(), u32(textCharLength), text);
-		RPC_MESSAGE(Log, log)
+
+		if (!TakeLockForRpc())
+		{
+			FlushDebug();
+			return;
+		}
+		RPC_MESSAGE_NO_LOCK(Log, log)
 		writer.WriteBool(printInSession);
 		writer.WriteBool(isError);
 		writer.WriteString(text, textCharLength);
 		writer.Flush();
+		g_communicationLock.Leave();
 	}
 
 	void Rpc_WriteLogf(const tchar* format, ...)
@@ -430,24 +449,18 @@ namespace uba
 
 	UBA_NOINLINE void Rpc_ResolveCallstack(StringBufferBase& out, u32 skipCallstackCount, void* context)
 	{
-		u32 tryCount = 0;
-		bool hasLock = false;
-		while (tryCount++ < 5)
+		if (!TakeLockForRpc())
 		{
-			hasLock = g_communicationLock.TryEnter();
-			if (hasLock)
-				break;
-			Sleep(100);
+			out.Append(TCV("\n  Can't fetch callstack while in rpc"));
+			return;
 		}
 
-		BinaryWriter writer;
-		writer.WriteByte(MessageType_ResolveCallstack);
+		RPC_MESSAGE_NO_LOCK(ResolveCallstack, log)
 		auto written = (u32*)writer.AllocWrite(4);
 		if (WriteCallstackInfo(writer, skipCallstackCount, context))
 		{
 			*written = u32(writer.GetPosition()) - 5;
-			writer.Flush();
-			BinaryReader reader;
+			BinaryReader reader = writer.Flush();
 			reader.ReadString(out);
 		}
 		else
diff --git a/Engine/Source/Programs/UnrealBuildAccelerator/Detours/Private/UbaDetoursShared.h b/Engine/Source/Programs/UnrealBuildAccelerator/Detours/Private/UbaDetoursShared.h
--- a/Engine/Source/Programs/UnrealBuildAccelerator/Detours/Private/UbaDetoursShared.h
+++ b/Engine/Source/Programs/UnrealBuildAccelerator/Detours/Private/UbaDetoursShared.h
@@ -146,12 +146,12 @@ namespace uba
 		DEBUG_LOG(TC("RPC_MESSAGE %s"), TC(#messageName)); \
 		TimerScope ts(g_stats.timerName); \
 		SCOPED_FUTEX(g_communicationLock, pcs); \
-		BinaryWriter writer; \
+		BinaryWriter writer(ProcessCommunication_Write); \
 		writer.WriteByte(MessageType_##messageName);
 
 	#define RPC_MESSAGE_NO_LOCK(messageName, timerName) \
 		DEBUG_LOG(TC("RPC_MESSAGE %s"), TC(#messageName)); \
 		TimerScope ts(g_stats.timerName); \
-		BinaryWriter writer; \
+		BinaryWriter writer(ProcessCommunication_Write); \
 		writer.WriteByte(MessageType_##messageName);
 }
diff --git a/Engine/Source/Programs/UnrealBuildAccelerator/Detours/Private/Windows/UbaDetoursFunctionsKernelBase.inl b/Engine/Source/Programs/UnrealBuildAccelerator/Detours/Private/Windows/UbaDetoursFunctionsKernelBase.inl
--- a/Engine/Source/Programs/UnrealBuildAccelerator/Detours/Private/Windows/UbaDetoursFunctionsKernelBase.inl
+++ b/Engine/Source/Programs/UnrealBuildAccelerator/Detours/Private/Windows/UbaDetoursFunctionsKernelBase.inl
@@ -42,8 +42,7 @@ DWORD Local_GetLongPathNameW(LPCWSTR lpszShortPath, LPWSTR lpszLongPath, DWORD c
 		{
 			RPC_MESSAGE(GetLongPathName, longPathName)
 			writer.WriteString(lpszShortPath);
-			writer.Flush();
-			BinaryReader reader;
+			BinaryReader reader = writer.Flush();
 			errorCode = reader.ReadU32();
 			reader.ReadString(longName);
 		}
@@ -293,8 +292,7 @@ BOOL Detoured_CreateDirectoryW(LPCWSTR lpPathName, LPSECURITY_ATTRIBUTES lpSecur
 		RPC_MESSAGE(CreateDirectory, createFile)
 		writer.WriteStringKey(pathNameKey);
 		writer.WriteString(pathName);
-		writer.Flush();
-		BinaryReader reader;
+		BinaryReader reader = writer.Flush();
 		res = reader.ReadBool();
 		errorCode = reader.ReadU32();
 		directoryTableSize = reader.ReadU32();
@@ -340,8 +338,7 @@ BOOL Detoured_RemoveDirectoryW(LPCWSTR lpPathName)
 		RPC_MESSAGE(RemoveDirectory, deleteFile)
 		writer.WriteStringKey(pathNameKey);
 		writer.WriteString(pathName);
-		writer.Flush();
-		BinaryReader reader;
+		BinaryReader reader = writer.Flush();
 		res = reader.ReadBool();
 		errorCode = reader.ReadU32();
 		directoryTableSize = reader.ReadU32();
@@ -1209,8 +1206,7 @@ BOOL Shared_CopyFileExW(const tchar* func, LPCWSTR lpExistingFileName, LPCWSTR l
 		writer.WriteString(fromName);
 		writer.WriteStringKey(toKey);
 		writer.WriteString(toName);
-		writer.Flush();
-		BinaryReader reader;
+		BinaryReader reader = writer.Flush();
 		reader.ReadString(newFromName);
 		reader.ReadString(newToName);
 		closeId = reader.ReadU32();
@@ -1293,8 +1289,7 @@ BOOL Detoured_CreateHardLinkW(LPCWSTR lpFileName, LPCWSTR lpExistingFileName, LP
 		writer.WriteString(fromName);
 		writer.WriteStringKey(toKey);
 		writer.WriteString(toName);
-		writer.Flush();
-		BinaryReader reader;
+		BinaryReader reader = writer.Flush();
 		reader.ReadString(newFromName);
 		reader.ReadString(newToName);
 		closeId = reader.ReadU32();
@@ -1352,8 +1347,7 @@ BOOL Detoured_DeleteFileW(LPCWSTR lpFileName)
 		writer.WriteString(fixedName);
 		writer.WriteStringKey(fileNameKey);
 		writer.WriteU32(closeId);
-		writer.Flush();
-		BinaryReader reader;
+		BinaryReader reader = writer.Flush();
 		result = reader.ReadBool();
 		errorCode = reader.ReadU32();
 		directoryTableSize = reader.ReadU32();
@@ -1459,8 +1453,7 @@ bool Shared_MoveFile(const tchar* func, LPCWSTR lpExistingFileName, LPCWSTR lpNe
 		writer.WriteStringKey(destKey);
 		writer.WriteString(dest);
 		writer.WriteU32(dwFlags);
-		writer.Flush();
-		BinaryReader reader;
+		BinaryReader reader = writer.Flush();
 		result = reader.ReadBool();
 		errorCode = reader.ReadU32();
 		directoryTableSize = reader.ReadU32();
@@ -2363,8 +2356,7 @@ LPVOID Detoured_MapViewOfFileEx(HANDLE hFileMappingObject, DWORD dwDesiredAccess
 				writer.WriteU64(fi.trueFileMapHandle.ToU64());
 				writer.WriteU64(offset + fi.trueFileMapOffset);
 				writer.WriteU64(dwNumberOfBytesToMap);
-				writer.Flush();
-				BinaryReader reader;
+				BinaryReader reader = writer.Flush();
 				view.Init(g_sharedMemoryAllocator, dwNumberOfBytesToMap, lpBaseAddress);
 				view.AddRequestedMemory(g_sharedMemoryAllocator, reader, SharedMemoryMapType_CopyOnWrite, 0);
 				SetLastError(ERROR_SUCCESS);
@@ -3085,8 +3077,7 @@ BOOL Detoured_CreateProcessW(LPCWSTR lpApplicationName, LPWSTR lpCommandLine, LP
 		writer.WriteString(workingDir);
 		writer.WriteBool(startSuspended);
 		writer.WriteBool(isChild);
-		writer.Flush();
-		BinaryReader reader;
+		BinaryReader reader = writer.Flush();
 		processId = reader.ReadU32();
 		if (!processId)
 		{
@@ -3321,6 +3312,7 @@ BOOL Detoured_CreateProcessW(LPCWSTR lpApplicationName, LPWSTR lpCommandLine, LP
 	}
 	#endif
 
+	u32 trueProcessId = lpProcessInformation->dwProcessId;
 	if (isChild)
 	{
 		RPC_MESSAGE(StartProcess, createProcess)
@@ -3329,7 +3321,7 @@ BOOL Detoured_CreateProcessW(LPCWSTR lpApplicationName, LPWSTR lpCommandLine, LP
 		writer.WriteBool(res);
 		writer.WriteU32(lastError);
 		writer.WriteU64(u64(lpProcessInformation->hProcess));
-		writer.WriteU32(lpProcessInformation->dwProcessId);
+		writer.WriteU32(trueProcessId);
 		writer.WriteU64(u64(lpProcessInformation->hThread));
 		writer.Flush();
 		DEBUG_LOG_PIPE(L"StartProcess", L"%ls %ls", application.data, originalCmd.data);
@@ -3347,7 +3339,7 @@ BOOL Detoured_CreateProcessW(LPCWSTR lpApplicationName, LPWSTR lpCommandLine, LP
 	detouredHandle->trueHandle = trueHandle;
 	lpProcessInformation->hProcess = makeDetouredHandle(detouredHandle);
 
-	DEBUG_LOG_DETOURED(L"CreateProcessW", L"(RESULT) %s %llu (0x%llx)", application.data, u64(lpProcessInformation->hProcess), u64(trueHandle));
+	DEBUG_LOG_DETOURED(L"CreateProcessW", L"(RESULT) %s %llu (true: %llu pid: %u)", application.data, u64(lpProcessInformation->hProcess), u64(trueHandle), trueProcessId);
 	return TRUE;
 }
 
@@ -3384,43 +3376,6 @@ BOOL Detoured_CreateProcessA(LPCSTR lpApplicationName, LPSTR lpCommandLine, LPSE
 	return Detoured_CreateProcessW(lpApplicationNameW, lpCommandLineW, lpProcessAttributes, lpThreadAttributes, bInheritHandles, dwCreationFlags, lpEnvironment, lpCurrentDirectoryW, &lpStartupInfoW, lpProcessInformation);
 }
 
-BOOL Detoured_TerminateProcess(HANDLE hProcess, UINT uExitCode)
-{
-	DETOURED_CALL(TerminateProcess);
-	DEBUG_LOG_DETOURED(L"TerminateProcess", L"%llu (%ls) ExitCode: %u", u64(hProcess), HandleToName(hProcess), uExitCode);
-
-	// Some processes actually call terminateprocess on themselves when exiting, ugh.
-	if (hProcess == INVALID_HANDLE_VALUE)
-	{
-		CloseCaches();
-		SendExitMessage((DWORD)uExitCode, GetTime());
-		PostDeinit();
-	}
-	else
-	{
-		if (isDetouredHandle(hProcess))
-			hProcess = asDetouredHandle(hProcess).trueHandle;
-
-		if (GetProcessId(hProcess) == GetCurrentProcessId())
-		{
-			CloseCaches();
-			SendExitMessage((DWORD)uExitCode, GetTime());
-			PostDeinit();
-		}
-		else
-		{
-			DWORD pid = GetProcessId(hProcess);
-			DEBUG_LOG(L"PROCESS %u WILL BE TERMINATED", pid);
-			RPC_MESSAGE(ExitChildProcess, detach) // Let it be part of detach timer for now
-			writer.WriteU32(pid);
-			writer.WriteU32(uExitCode);
-			writer.Flush();
-		}
-	}
-
-	return True_TerminateProcess(hProcess, uExitCode);
-}
-
 BOOL Detoured_AssignProcessToJobObject(HANDLE hJob, HANDLE hProcess)
 {
 	HANDLE trueHandle = hProcess;
diff --git a/Engine/Source/Programs/UnrealBuildAccelerator/Detours/Private/Windows/UbaDetoursFunctionsNtDll.inl b/Engine/Source/Programs/UnrealBuildAccelerator/Detours/Private/Windows/UbaDetoursFunctionsNtDll.inl
--- a/Engine/Source/Programs/UnrealBuildAccelerator/Detours/Private/Windows/UbaDetoursFunctionsNtDll.inl
+++ b/Engine/Source/Programs/UnrealBuildAccelerator/Detours/Private/Windows/UbaDetoursFunctionsNtDll.inl
@@ -768,7 +768,7 @@ NTSTATUS Detoured_NtSetInformationFile(HANDLE FileHandle, PIO_STATUS_BLOCK IoSta
 	return runTrue();
 }
 
-NTSTATUS NTAPI Detoured_NtSetInformationObject(HANDLE ObjectHandle, OBJECT_INFORMATION_CLASS ObjectInformationClass, PVOID ObjectInformation, ULONG Length)
+NTSTATUS Detoured_NtSetInformationObject(HANDLE ObjectHandle, OBJECT_INFORMATION_CLASS ObjectInformationClass, PVOID ObjectInformation, ULONG Length)
 {
 	if (isDetouredHandle(ObjectHandle))
 	{
@@ -781,6 +781,55 @@ NTSTATUS NTAPI Detoured_NtSetInformationObject(HANDLE ObjectHandle, OBJECT_INFOR
 	return res;
 }
 
+NTSTATUS NTAPI Detoured_NtTerminateProcess(HANDLE ProcessHandle, NTSTATUS ExitStatus)
+{
+	// NtTerminateProcess can be reached in three ways.. one is through ExitProcess, TerminateProcess and NtTerminateProcess
+	// If called through ExitMessage we don't want to do anymore work here.
+	DETOURED_CALL(TerminateProcess);
+	DEBUG_LOG_DETOURED(L"NtTerminateProcess", L"%llu (%ls) ExitCode: %u", u64(ProcessHandle), HandleToName(ProcessHandle), ExitStatus);
+
+	if (ProcessHandle == NULL)
+		return True_NtTerminateProcess(ProcessHandle, ExitStatus);
+
+	HANDLE trueProcess = ProcessHandle;
+
+	bool isSelf = trueProcess == GetCurrentProcess();
+	if (!isSelf)
+		if (isDetouredHandle(trueProcess))
+			trueProcess = asDetouredHandle(trueProcess).trueHandle;
+
+	DWORD pid = GetProcessId(trueProcess);
+	if (!isSelf)
+		isSelf = pid == GetCurrentProcessId();
+
+	if (isSelf)
+	{
+		// Not called through ExitProcess
+		if (!g_exitMessageSent)
+		{
+			CloseCaches();
+			SendExitMessage((DWORD)ExitStatus, GetTime());
+			PostDeinit();
+			#if UBA_DEBUG_LOG_ENABLED
+			FlushDebugLog();
+			#endif
+		}
+		return True_NtTerminateProcess(trueProcess, ExitStatus);
+	}
+
+	pid = GetProcessId(trueProcess);
+	if (pid != GetCurrentProcessId())
+	{
+		DEBUG_LOG(L"PROCESS %u WILL BE TERMINATED", pid);
+		RPC_MESSAGE(ExitChildProcess, detach) // Let it be part of detach timer for now
+		writer.WriteU32(pid);
+		writer.WriteU32(ExitStatus);
+		writer.Flush();
+	}
+
+	return True_NtTerminateProcess(ProcessHandle, ExitStatus);
+}
+
 NTSTATUS NTAPI Detoured_NtCreateSection(PHANDLE SectionHandle, ACCESS_MASK DesiredAccess, POBJECT_ATTRIBUTES ObjectAttributes, PLARGE_INTEGER MaximumSize, ULONG SectionPageProtection, ULONG AllocationAttributes, HANDLE FileHandle)
 {
 	DETOURED_CALL(NtCreateSection);
@@ -2382,12 +2431,6 @@ NTSTATUS NTAPI Detoured_NtFlushBuffersFileEx(HANDLE FileHandle, ULONG Flags, PVO
 	return True_NtFlushBuffersFileEx(FileHandle, Flags, Parameters, ParametersSize, IoStatusBlock);
 }
 
-NTSTATUS NTAPI Detoured_NtTerminateProcess(HANDLE ProcessHandle, NTSTATUS ExitStatus)
-{
-	// NtTerminateProcess runs after DllMain DLL_PROCESS_DETACH which means we can't use this to send exit messaage etc
-	return True_NtTerminateProcess(ProcessHandle, ExitStatus);
-}
-
 NTSTATUS NTAPI Detoured_NtAlpcCreatePort(PHANDLE PortHandle, POBJECT_ATTRIBUTES ObjectAttributes, PALPC_PORT_ATTRIBUTES PortAttributes)
 {
 	//DEBUG_LOG_TRUE(L"NtAlpcCreatePort", L"");
diff --git a/Engine/Source/Programs/UnrealBuildAccelerator/Detours/Private/Windows/UbaDetoursFunctionsShell32.inl b/Engine/Source/Programs/UnrealBuildAccelerator/Detours/Private/Windows/UbaDetoursFunctionsShell32.inl
--- a/Engine/Source/Programs/UnrealBuildAccelerator/Detours/Private/Windows/UbaDetoursFunctionsShell32.inl
+++ b/Engine/Source/Programs/UnrealBuildAccelerator/Detours/Private/Windows/UbaDetoursFunctionsShell32.inl
@@ -13,8 +13,7 @@ HRESULT Detoured_SHGetKnownFolderPath(REFKNOWNFOLDERID rfid, DWORD dwFlags, HAND
 		RPC_MESSAGE(SHGetKnownFolderPath, getFullFileName)
 		writer.WriteBytes(&rfid, sizeof(KNOWNFOLDERID));
 		writer.WriteU32(dwFlags);
-		writer.Flush();
-		BinaryReader reader;
+		BinaryReader reader = writer.Flush();
 		HRESULT res = reader.ReadU32();
 		*ppszPath = NULL;
 		if (res == S_OK)
diff --git a/Engine/Source/Programs/UnrealBuildAccelerator/Detours/Private/Windows/UbaDetoursFunctionsWin.cpp b/Engine/Source/Programs/UnrealBuildAccelerator/Detours/Private/Windows/UbaDetoursFunctionsWin.cpp
--- a/Engine/Source/Programs/UnrealBuildAccelerator/Detours/Private/Windows/UbaDetoursFunctionsWin.cpp
+++ b/Engine/Source/Programs/UnrealBuildAccelerator/Detours/Private/Windows/UbaDetoursFunctionsWin.cpp
@@ -1066,8 +1066,7 @@ void Init(const DetoursPayload& payload, u64 startTime)
 
 	{
 		RPC_MESSAGE(Init, init)
-		writer.Flush();
-		BinaryReader reader;
+		BinaryReader reader = writer.Flush();
 
 		g_processId = reader.ReadU32();
 		g_isChild = reader.ReadBool();
@@ -1231,8 +1230,7 @@ extern "C"
 		RPC_MESSAGE(Custom, log)
 		writer.WriteU32(sendSize);
 		writer.WriteBytes(send, sendSize);
-		writer.Flush();
-		BinaryReader reader;
+		BinaryReader reader = writer.Flush();
 		u32 recvSize = reader.ReadU32();
 		UBA_ASSERT(recvSize < recvCapacity);
 		reader.ReadBytes(recv, recvSize);
@@ -1242,8 +1240,7 @@ extern "C"
 	UBA_DETOURED_API bool UbaFlushWrittenFiles()
 	{
 		RPC_MESSAGE(FlushWrittenFiles, log)
-		writer.Flush();
-		BinaryReader reader;
+		BinaryReader reader = writer.Flush();
 		return reader.ReadBool();
 	}
 
@@ -1253,8 +1250,7 @@ extern "C"
 			RPC_MESSAGE(UpdateEnvironment, log)
 			writer.WriteString(reason ? reason : L"");
 			writer.WriteBool(resetStats);
-			writer.Flush();
-			BinaryReader reader;
+			BinaryReader reader = writer.Flush();
 			if (!reader.ReadBool())
 				return false;
 		}
@@ -1292,8 +1288,7 @@ extern "C"
 			g_stats.Write(writer);
 			g_kernelStats.Write(writer);
 
-			writer.Flush();
-			BinaryReader reader;
+			BinaryReader reader = writer.Flush();
 			newProcess = reader.ReadBool();
 			*outShouldExit = reader.ReadBool();
 			if (newProcess)
diff --git a/Engine/Source/Programs/UnrealBuildAccelerator/Detours/Private/Windows/UbaDetoursFunctionsWin.h b/Engine/Source/Programs/UnrealBuildAccelerator/Detours/Private/Windows/UbaDetoursFunctionsWin.h
--- a/Engine/Source/Programs/UnrealBuildAccelerator/Detours/Private/Windows/UbaDetoursFunctionsWin.h
+++ b/Engine/Source/Programs/UnrealBuildAccelerator/Detours/Private/Windows/UbaDetoursFunctionsWin.h
@@ -92,7 +92,6 @@
 	DETOURED_FUNCTION(GetFinalPathNameByHandleW) \
 	DETOURED_FUNCTION(CreateProcessW) \
 	DETOURED_FUNCTION(CreateProcessA) \
-	DETOURED_FUNCTION(TerminateProcess) \
 	DETOURED_FUNCTION(SearchPathW) \
 	DETOURED_FUNCTION(LoadLibraryExW) \
 	DETOURED_FUNCTION(GetStdHandle) \
@@ -148,6 +147,7 @@
 	DETOURED_FUNCTION(NtQuerySecurityObject) \
 	DETOURED_FUNCTION(NtSetInformationFile) \
 	DETOURED_FUNCTION(NtSetInformationObject) \
+	DETOURED_FUNCTION(NtTerminateProcess) \
 	DETOURED_FUNCTION(NtCreateSection) \
 	DETOURED_FUNCTION(RtlExitUserProcess) \
 	DETOURED_FUNCTION(RtlSizeHeap) \
@@ -368,7 +368,6 @@
 	DETOURED_FUNCTION(RtlDosPathNameToRelativeNtPathName_U) \
 	DETOURED_FUNCTION(NtCreateIoCompletion) \
 	DETOURED_FUNCTION(NtFlushBuffersFileEx) \
-	DETOURED_FUNCTION(NtTerminateProcess) \
 	DETOURED_FUNCTION(NtAlpcCreatePort) \
 	DETOURED_FUNCTION(NtAlpcConnectPort) \
 	DETOURED_FUNCTION(NtAlpcCreatePortSection) \
@@ -443,6 +442,12 @@ extern "C" {
 	NTSTATUS NTAPI NtTerminateProcess(HANDLE ProcessHandle, NTSTATUS ExitStatus);
 	NTSTATUS NTAPI NtSetInformationFile(HANDLE FileHandle, PIO_STATUS_BLOCK IoStatusBlock, PVOID FileInformation, ULONG Length, FILE_INFORMATION_CLASS FileInformationClass);
 	NTSTATUS NTAPI NtSetInformationObject(HANDLE ObjectHandle, OBJECT_INFORMATION_CLASS ObjectInformationClass, PVOID ObjectInformation, ULONG Length);
+	NTSTATUS NTAPI NtCreateJobObject(HANDLE* handle, ACCESS_MASK access, const OBJECT_ATTRIBUTES* attr);
+	NTSTATUS NTAPI NtOpenJobObject(HANDLE* handle, ACCESS_MASK access, const OBJECT_ATTRIBUTES *attr);
+	NTSTATUS NTAPI NtQueryInformationJobObject(HANDLE handle, JOBOBJECTINFOCLASS infoClass, void* info, ULONG len, ULONG* retLen);
+	NTSTATUS NTAPI NtSetInformationJobObject(HANDLE handle, JOBOBJECTINFOCLASS infoClass, void* info, ULONG len);
+	NTSTATUS NTAPI NtAssignProcessToJobObject(HANDLE job, HANDLE process);
+	NTSTATUS NTAPI NtTerminateJobObject(HANDLE handle, NTSTATUS status);
 	NTSTATUS NTAPI NtCreateSection(PHANDLE SectionHandle, ACCESS_MASK DesiredAccess, POBJECT_ATTRIBUTES ObjectAttributes, PLARGE_INTEGER MaximumSize, ULONG SectionPageProtection, ULONG AllocationAttributes, HANDLE FileHandle);
 	NTSTATUS NTAPI NtCreateIoCompletion(PHANDLE IoCompletionHandle, ACCESS_MASK DesiredAccess, POBJECT_ATTRIBUTES ObjectAttributes, ULONG Count);
 	NTSTATUS NTAPI NtQuerySecurityObject(HANDLE Handle, SECURITY_INFORMATION SecurityInformation, PSECURITY_DESCRIPTOR SecurityDescriptor, ULONG Length, PULONG LengthNeeded);
diff --git a/Engine/Source/Programs/UnrealBuildAccelerator/Detours/Private/Windows/UbaDetoursMainWin.cpp b/Engine/Source/Programs/UnrealBuildAccelerator/Detours/Private/Windows/UbaDetoursMainWin.cpp
--- a/Engine/Source/Programs/UnrealBuildAccelerator/Detours/Private/Windows/UbaDetoursMainWin.cpp
+++ b/Engine/Source/Programs/UnrealBuildAccelerator/Detours/Private/Windows/UbaDetoursMainWin.cpp
@@ -8,6 +8,7 @@
 
 #include "UbaProtocol.h"
 #include "UbaProcessStats.h"
+#include "UbaDefinitions.h"
 #include "UbaDetoursPayload.h"
 #include "UbaDetoursShared.h"
 #include "UbaDetoursFunctionsWin.h"
@@ -21,8 +22,8 @@ namespace uba
 	HANDLE g_readEvent;
 	HANDLE g_cancelEvent;
 	HANDLE g_messageMappingHandle;
-	u8* g_readMessageMappingMem;
 	u8* g_writeMessageMappingMem;
+	bool g_useSignalObjectAndWait;
 
 	bool WasCancelled()
 	{
@@ -59,8 +60,8 @@ namespace uba
 			return -5;
 		}
 		
-		g_readMessageMappingMem = mem;
-		g_writeMessageMappingMem = mem;// +CommunicationMemSize / 2;
+		g_useSignalObjectAndWait = !IsRunningWine(); // There seems to be a bug in wine (as of 11.0) for using this function. It can return before it should
+		g_writeMessageMappingMem = mem;
 
 		SetMemoryWorkingSet(CommunicationMemSize*2);
 
@@ -83,29 +84,47 @@ namespace uba
 
 	ANALYSIS_NORETURN void TerminateCurrentProcess(u32 exitCode)
 	{
-		if (True_TerminateProcess)
-			True_TerminateProcess(GetCurrentProcess(), exitCode);
+		if (True_NtTerminateProcess)
+			True_NtTerminateProcess(GetCurrentProcess(), exitCode);
 		else
 			TerminateProcess(GetCurrentProcess(), exitCode);
 	}
 
-	BinaryWriter::BinaryWriter()
+	BinaryWriter::BinaryWriter(ProcessCommunication)
 	{
 		m_begin = g_writeMessageMappingMem;
 		m_pos = m_begin;
 		m_end = m_begin + CommunicationMemSize;// / 2;
+
+		#if UBA_PROTOCOL_GUARD
+		WriteU32(0xdeadbeef);
+		#endif
 	}
 
-	void BinaryWriter::Flush(bool waitOnResponse)
+	BinaryReader BinaryWriter::Flush(bool waitOnResponse)
 	{
+		#if UBA_PROTOCOL_GUARD
+		if (BinaryReader(g_writeMessageMappingMem, 0, CommunicationMemSize).ReadU32() != 0xdeadbeef)
+			ExitProcess(2345);
+		#endif
+
 		if (!waitOnResponse)
 		{
 			SetEvent(g_writeEvent);
-			return;
+			return BinaryReader(nullptr, 0, 0);
 		}
 
 		TimerScope ts(g_stats.waitOnResponse);
-		DWORD res = SignalObjectAndWait(g_writeEvent, g_readEvent, 1000, FALSE);
+		DWORD res;
+		
+		if (g_useSignalObjectAndWait)
+			res = SignalObjectAndWait(g_writeEvent, g_readEvent, 1000, FALSE);
+		else
+		{
+			if (!SetEvent(g_writeEvent))
+				ExitProcess(1339);
+			res = True_WaitForSingleObject(g_readEvent, 1000);
+		}
 
 		do
 		{
@@ -131,13 +150,20 @@ namespace uba
 			res = True_WaitForSingleObject(g_readEvent, 500);
 
 		} while (true);
-	}
 
-	BinaryReader::BinaryReader()
-	{
-		m_begin = g_readMessageMappingMem;
-		m_pos = m_begin;
-		m_end = m_begin + CommunicationMemSize;// / 2;
+		BinaryReader reader(m_begin, 0, CommunicationMemSize);
+	
+		#if UBA_PROTOCOL_GUARD
+		u32 baadfood = reader.ReadU32();
+		if (baadfood != 0xbaadf00d)
+			ExitProcess(2345);
+		BinaryWriter(g_writeMessageMappingMem, 0, CommunicationMemSize).WriteU32(0xcafebabe);
+		#endif
+
+		m_begin = nullptr;
+		m_pos = 0;
+		m_end = m_begin;
+		return reader;
 	}
 }
 
