[UBA]
* Fixed bugs in NtQueryVolumeInformationFile and NtQueryInformationFile
* Added more logging

diff --git a/Engine/Source/Programs/UnrealBuildAccelerator/Core/Public/UbaWinTypes.h b/Engine/Source/Programs/UnrealBuildAccelerator/Core/Public/UbaWinTypes.h
--- a/Engine/Source/Programs/UnrealBuildAccelerator/Core/Public/UbaWinTypes.h
+++ b/Engine/Source/Programs/UnrealBuildAccelerator/Core/Public/UbaWinTypes.h
@@ -64,6 +64,7 @@ struct FILE_FS_DEVICE_INFORMATION
 	DEVICE_TYPE DeviceType;
 	ULONG Characteristics;
 };
+constexpr ULONG FILE_DEVICE_IS_MOUNTED = 0x00000020;
 
 struct FILE_FS_ATTRIBUTE_INFORMATION
 {
@@ -175,8 +176,6 @@ struct FILE_FS_VOLUME_INFORMATION {
 	WCHAR         VolumeLabel[1];
 };
 
-
-
 #define FILE_INFO_CLASSES \
 	FILE_INFO_CLASS(FileBasicInformation, 4) \
 	FILE_INFO_CLASS(FileStandardInformation, 5) \
diff --git a/Engine/Source/Programs/UnrealBuildAccelerator/Detours/Private/Windows/UbaDetoursFunctionsKernelBase.inl b/Engine/Source/Programs/UnrealBuildAccelerator/Detours/Private/Windows/UbaDetoursFunctionsKernelBase.inl
--- a/Engine/Source/Programs/UnrealBuildAccelerator/Detours/Private/Windows/UbaDetoursFunctionsKernelBase.inl
+++ b/Engine/Source/Programs/UnrealBuildAccelerator/Detours/Private/Windows/UbaDetoursFunctionsKernelBase.inl
@@ -21,14 +21,19 @@ DWORD Local_GetLongPathNameW(LPCWSTR lpszShortPath, LPWSTR lpszLongPath, DWORD c
 		u32 len = u32(wcslen(longPath));
 		if (len == 0)
 		{
+			DEBUG_LOG_DETOURED(L"GetLongPathNameW", L"(NOT_FOUND) (%s) -> 0", lpszShortPath);
 			SetLastError(ERROR_FILE_NOT_FOUND);
 			return 0;
 		}
 		SetLastError(ERROR_SUCCESS);
 		if (cchBuffer <= len)
+		{
+			DEBUG_LOG_DETOURED(L"GetLongPathNameW", L"(BUFFER_SMALL) (%s) -> %u", lpszShortPath, len + 1);
 			return len + 1;
-
+		}
 		memcpy(lpszLongPath, longPath, (len + 1) * 2);
+
+		DEBUG_LOG_DETOURED(L"GetLongPathNameW", L"(%s) -> %u", lpszShortPath, len + 1);
 		return len + 1;
 	}
 
@@ -666,6 +671,7 @@ BOOL Detoured_DeviceIoControl(HANDLE hDevice, DWORD dwIoControlCode, LPVOID lpIn
 			if (dwIoControlCode == FSCTL_GET_REPARSE_POINT)
 			{
 				//UBA_ASSERTF(false, TC("FSCTL_GET_REPARSE_POINT"));
+				DEBUG_LOG_DETOURED(L"DeviceIoControl", L"(FSCTL_GET_REPARSE_POINT) %llu (%ls) -> ERROR_NOT_A_REPARSE_POINT", uintptr_t(hDevice), HandleToName(hDevice));
 				return ERROR_NOT_A_REPARSE_POINT;
 			}
 			else if (dwIoControlCode == FSCTL_SET_SPARSE)
@@ -833,7 +839,10 @@ DWORD Detoured_GetFileAttributesW(LPCWSTR lpFileName)
 
 	StringBuffer<MaxPath> fixedPath;
 	if (!FixPath(fixedPath, lpFileName))
+	{
+		DEBUG_LOG_DETOURED(L"GetFileAttributesW", L"(BAD_PATH) (%ls) -> INVALID_FILE_ATTRIBUTES", lpFileName);
 		return INVALID_FILE_ATTRIBUTES;
+	}
 	DevirtualizePath(fixedPath);
 
 	WIN32_FILE_ATTRIBUTE_DATA data;
@@ -972,6 +981,7 @@ DWORD Detoured_GetFullPathNameA(LPCSTR lpFileName, DWORD nBufferLength, LPSTR lp
 	DETOURED_CALL(GetFullPathNameA);
 	if (!lpFileName)
 	{
+		DEBUG_LOG(L"GetFullPathNameA (NULL_NAME)");
 		SetLastError(ERROR_INVALID_NAME);
 		return 0;
 	}
@@ -981,7 +991,10 @@ DWORD Detoured_GetFullPathNameA(LPCSTR lpFileName, DWORD nBufferLength, LPSTR lp
 	FixPath(fullPath, temp.data);
 	u64 requiredSize = fullPath.count + 1;
 	if (nBufferLength < requiredSize)
+	{
+		DEBUG_LOG(L"GetFullPathNameA (BUFFER_SMALL) %S", lpFileName);
 		return DWORD(requiredSize);
+	}
 	if (!fullPath.Parse(lpBuffer, requiredSize))
 	{
 		DEBUG_LOG(L"GetFullPathNameA failed to parse %s to ansi", fullPath.data);
@@ -1845,6 +1858,7 @@ BOOL Detoured_GetFileInformationByHandleEx(HANDLE hFile, FILE_INFO_BY_HANDLE_CLA
 		}
 		else if (fileInformationClass == FileIdBothDirectoryRestartInfo || fileInformationClass == FileIdBothDirectoryInfo)
 		{
+			DEBUG_LOG_DETOURED(L"GetFileInformationByHandleEx", L"(NOT_SUPPORTED) (%u) %llu (%ls) -> FALSE (ERROR_ACCESS_DENIED)", fileInformationClass, uintptr_t(hFile), HandleToName(hFile));
 			SetLastError(ERROR_ACCESS_DENIED);
 			return FALSE;
 		}
diff --git a/Engine/Source/Programs/UnrealBuildAccelerator/Detours/Private/Windows/UbaDetoursFunctionsNtDll.inl b/Engine/Source/Programs/UnrealBuildAccelerator/Detours/Private/Windows/UbaDetoursFunctionsNtDll.inl
--- a/Engine/Source/Programs/UnrealBuildAccelerator/Detours/Private/Windows/UbaDetoursFunctionsNtDll.inl
+++ b/Engine/Source/Programs/UnrealBuildAccelerator/Detours/Private/Windows/UbaDetoursFunctionsNtDll.inl
@@ -71,6 +71,9 @@ StringBuffer<32> ToString(NTSTATUS s)
 
 NTSTATUS Detoured_NtQueryVolumeInformationFile(HANDLE FileHandle, PIO_STATUS_BLOCK IoStatusBlock, PVOID FsInformation, ULONG Length, FS_INFORMATION_CLASS FsInformationClass)
 {
+	IoStatusBlock->Status = STATUS_SUCCESS;
+	IoStatusBlock->Information = 0;
+
 	DETOURED_CALL(NtQueryVolumeInformationFile);
 	HANDLE TrueHandle = FileHandle;
 	if (isDetouredHandle(FileHandle))
@@ -81,8 +84,10 @@ NTSTATUS Detoured_NtQueryVolumeInformationFile(HANDLE FileHandle, PIO_STATUS_BLO
 			if (FsInformationClass == FileFsDeviceInformation)
 			{
 				auto& info = *(FILE_FS_DEVICE_INFORMATION*)FsInformation;
-				info.DeviceType = FILE_DEVICE_FILE_SYSTEM;
-				info.Characteristics = 0;
+				info.DeviceType = FILE_DEVICE_FILE_SYSTEM; //FILE_DEVICE_DISK;
+				info.Characteristics = FILE_DEVICE_IS_MOUNTED;
+				IoStatusBlock->Information = sizeof(FILE_FS_DEVICE_INFORMATION);
+				DEBUG_LOG_DETOURED(L"NtQueryVolumeInformationFile", L"%llu (FileFsDeviceInformation) (%ls) -> %ls", uintptr_t(FileHandle), HandleToName(FileHandle), ToString(STATUS_SUCCESS).data);
 				return STATUS_SUCCESS;
 			}
 		}
@@ -92,7 +97,15 @@ NTSTATUS Detoured_NtQueryVolumeInformationFile(HANDLE FileHandle, PIO_STATUS_BLO
 			if (FsInformationClass == FileFsVolumeInformation)
 			{
 				// TODO This code path is here to handle nodejs queries..
-				
+				ULONG required = offsetof(FILE_FS_VOLUME_INFORMATION, VolumeLabel);
+				if (Length < required)
+				{
+					IoStatusBlock->Status = STATUS_BUFFER_TOO_SMALL;
+					IoStatusBlock->Information = 0;
+					DEBUG_LOG_DETOURED(L"NtQueryVolumeInformationFile", L"%llu (FileFsVolumeInformation) (%ls) -> %ls", uintptr_t(FileHandle), HandleToName(FileHandle), ToString(STATUS_BUFFER_TOO_SMALL).data);
+					return STATUS_BUFFER_TOO_SMALL;
+				}
+
 				auto& info = *(FILE_FS_VOLUME_INFORMATION*)FsInformation;
 				UBA_ASSERT(dh.dirTableOffset != ~0u);
 				DirectoryTable::EntryInformation entryInfo;
@@ -103,13 +116,17 @@ NTSTATUS Detoured_NtQueryVolumeInformationFile(HANDLE FileHandle, PIO_STATUS_BLO
 				info.VolumeLabelLength = 0;
 				info.SupportsObjects = false;
 				info.VolumeLabel[0] = 0;
+				IoStatusBlock->Information = offsetof(FILE_FS_VOLUME_INFORMATION, VolumeLabel) + info.VolumeLabelLength;
+				DEBUG_LOG_DETOURED(L"NtQueryVolumeInformationFile", L"%llu (FileFsVolumeInformation) (%ls) -> %ls", uintptr_t(FileHandle), HandleToName(FileHandle), ToString(STATUS_SUCCESS).data);
 				return STATUS_SUCCESS;
 			}
 			if (FsInformationClass == FileFsDeviceInformation)
 			{
 				auto& info = *(FILE_FS_DEVICE_INFORMATION*)FsInformation;
-				info.DeviceType = FILE_DEVICE_FILE_SYSTEM;
-				info.Characteristics = 0;
+				info.DeviceType = FILE_DEVICE_FILE_SYSTEM; //FILE_DEVICE_DISK;
+				info.Characteristics = FILE_DEVICE_IS_MOUNTED;
+				IoStatusBlock->Information = sizeof(FILE_FS_DEVICE_INFORMATION);
+				DEBUG_LOG_DETOURED(L"NtQueryVolumeInformationFile", L"%llu (FileFsDeviceInformation) (%ls) -> %ls", uintptr_t(FileHandle), HandleToName(FileHandle), ToString(STATUS_SUCCESS).data);
 				return STATUS_SUCCESS;
 			}
 			UBA_ASSERTF(false, L"NtQueryVolumeInformationFile using class %u not handled %ls (%ls)", FsInformationClass, dh.fileObject->fileInfo->name, dh.fileObject->fileInfo->originalName);
@@ -120,12 +137,24 @@ NTSTATUS Detoured_NtQueryVolumeInformationFile(HANDLE FileHandle, PIO_STATUS_BLO
 		if (FsInformationClass == FileFsDeviceInformation)
 		{
 			auto& info = *(FILE_FS_DEVICE_INFORMATION*)FsInformation;
-			info.DeviceType = FILE_DEVICE_FILE_SYSTEM;
-			info.Characteristics = 0;
+			info.DeviceType = FILE_DEVICE_FILE_SYSTEM; //FILE_DEVICE_DISK;
+			info.Characteristics = FILE_DEVICE_IS_MOUNTED;
+			IoStatusBlock->Information = sizeof(FILE_FS_DEVICE_INFORMATION);
+			DEBUG_LOG_DETOURED(L"NtQueryVolumeInformationFile", L"%llu (FileFsDeviceInformation) (%ls) -> %ls", uintptr_t(FileHandle), HandleToName(FileHandle), ToString(STATUS_SUCCESS).data);
 			return STATUS_SUCCESS;
 		}
 		if (FsInformationClass == FileFsVolumeInformation)
 		{
+			// TODO This code path is here to handle nodejs queries..
+			ULONG required = offsetof(FILE_FS_VOLUME_INFORMATION, VolumeLabel);
+			if (Length < required)
+			{
+				IoStatusBlock->Status = STATUS_BUFFER_TOO_SMALL;
+				IoStatusBlock->Information = 0;
+				DEBUG_LOG_DETOURED(L"NtQueryVolumeInformationFile", L"%llu (FileFsVolumeInformation) (%ls) -> %ls", uintptr_t(FileHandle), HandleToName(FileHandle), ToString(STATUS_BUFFER_TOO_SMALL).data);
+				return STATUS_BUFFER_TOO_SMALL;
+			}
+
 			// TODO This code path is here to handle nodejs queries..
 			auto& listHandle = asListDirectoryHandle(FileHandle);
 			u32 entryOffset = listHandle.dir.tableOffset | 0x80000000;
@@ -139,6 +168,8 @@ NTSTATUS Detoured_NtQueryVolumeInformationFile(HANDLE FileHandle, PIO_STATUS_BLO
 			info.VolumeLabelLength = 0;
 			info.SupportsObjects = false;
 			info.VolumeLabel[0] = 0;
+			IoStatusBlock->Information = offsetof(FILE_FS_VOLUME_INFORMATION, VolumeLabel) + info.VolumeLabelLength;
+			DEBUG_LOG_DETOURED(L"NtQueryVolumeInformationFile", L"%llu (FileFsVolumeInformation) (%ls) -> %ls", uintptr_t(FileHandle), HandleToName(FileHandle), ToString(STATUS_SUCCESS).data);
 			return STATUS_SUCCESS;
 		}
 		UBA_ASSERTF(false, L"NtQueryVolumeInformationFile called in ListDirectoryHandle using class %u which is not implemented (%ls)", FsInformationClass, HandleToName(FileHandle));
@@ -269,22 +300,20 @@ NTSTATUS Detoured_NtQueryInformationFile(HANDLE FileHandle, PIO_STATUS_BLOCK IoS
 			return STATUS_SUCCESS;
 		};
 
-	auto setNameInformation = [&](FILE_NAME_INFORMATION& out, FileInfo& fi)
+	auto setNameInformation = [&](FILE_NAME_INFORMATION& out, FileInfo& fi, u32 offset)
 		{
 			const wchar_t* fileName = fi.originalName;
-			auto& info = *(FILE_NAME_INFORMATION*)FileInformation;
-			ULONG header = FIELD_OFFSET(FILE_NAME_INFORMATION, FileName);
 			ULONG nameBytes = (ULONG)(wcslen(fileName) * sizeof(WCHAR));
-			ULONG required = header + nameBytes;
-			if (Length < header)
+			ULONG required = offset + nameBytes;
+			if (Length < offset)
 			{
 				IoStatusBlock->Status = STATUS_INFO_LENGTH_MISMATCH;
 				return STATUS_INFO_LENGTH_MISMATCH;
 			}
-			ULONG copy = Min(nameBytes, Length - header);
-			info.FileNameLength = nameBytes;
+			ULONG copy = Min(nameBytes, Length - offset);
+			out.FileNameLength = nameBytes;
 			if (copy)
-				memcpy(info.FileName, fileName, copy);
+				memcpy(out.FileName, fileName, copy);
 			IoStatusBlock->Information = Min(Length, required);
 			return STATUS_SUCCESS;
 		};
@@ -345,7 +374,7 @@ NTSTATUS Detoured_NtQueryInformationFile(HANDLE FileHandle, PIO_STATUS_BLOCK IoS
 		
 		res = res == STATUS_SUCCESS ? setBasicInformation(info.BasicInformation, entryInfo) : res;
 		res = res == STATUS_SUCCESS ? setStandardInformation(info.StandardInformation, fi, &entryInfo) : res;
-		res = res == STATUS_SUCCESS ? setNameInformation(info.NameInformation, fi) : res;
+		res = res == STATUS_SUCCESS ? setNameInformation(info.NameInformation, fi, FIELD_OFFSET(FILE_ALL_INFORMATION, NameInformation.FileName)) : res;
 		info.InternalInformation.IndexNumber.QuadPart = entryInfo.fileIndex;
 		info.EaInformation.EaSize = 0;
 		info.AccessInformation.AccessFlags = dh.fileObject->desiredAccess;
@@ -401,7 +430,7 @@ NTSTATUS Detoured_NtQueryInformationFile(HANDLE FileHandle, PIO_STATUS_BLOCK IoS
 	{
 		if (dh.trueHandle != INVALID_HANDLE_VALUE)
 			return runTrue();
-		NTSTATUS res = setNameInformation(*(FILE_NAME_INFORMATION*)FileInformation, *dh.fileObject->fileInfo);
+		NTSTATUS res = setNameInformation(*(FILE_NAME_INFORMATION*)FileInformation, *dh.fileObject->fileInfo, FIELD_OFFSET(FILE_NAME_INFORMATION, FileName));
 		DEBUG_LOG_DETOURED(L"NtQueryInformationFile", L"(FileNameInformation) %llu (%s) -> %ls", uintptr_t(FileHandle), HandleToName(FileHandle), ToString(res).data);
 		return res;
 	}
