[UBA]
* Added support for SHGetFolderPath and SHGetFolderPathAndSubDir

diff --git a/Engine/Source/Programs/UnrealBuildAccelerator/Common/Private/UbaSession.cpp b/Engine/Source/Programs/UnrealBuildAccelerator/Common/Private/UbaSession.cpp
--- a/Engine/Source/Programs/UnrealBuildAccelerator/Common/Private/UbaSession.cpp
+++ b/Engine/Source/Programs/UnrealBuildAccelerator/Common/Private/UbaSession.cpp
@@ -3753,8 +3753,26 @@ namespace uba
 
 	bool Session::HostRun(BinaryReader& reader, BinaryWriter& writer)
 	{
-		#if !PLATFORM_WINDOWS
+#if PLATFORM_WINDOWS
+		TString temp;
+		const tchar* subdir = nullptr;
+		int csidl = (int)reader.ReadU32();
+		DWORD flags = reader.ReadU32();
+		if (reader.ReadBool())
+		{
+			temp = reader.ReadString();
+			subdir = temp.data();
+		}
+
+		static HMODULE shellDll = LoadLibrary(TC("Shell32.dll"));
+		using SHGetFolderPathAndSubDirWFunc = HRESULT(HWND hwnd, int csidl, HANDLE hToken, DWORD dwFlags, LPCWSTR pszSubDir, LPWSTR pszPath);
+		static SHGetFolderPathAndSubDirWFunc* SHGetFolderPathAndSubDirW = (SHGetFolderPathAndSubDirWFunc*)GetProcAddress(shellDll, "SHGetFolderPathAndSubDirW");
 
+		tchar path[MAX_PATH];
+		HRESULT res = SHGetFolderPathAndSubDirW(0, csidl, NULL, flags, subdir, path);
+		writer.WriteU32(res);
+		writer.WriteString(path);
+#else
 		Vector<TString> args;
 		while (reader.GetLeft())
 			args.push_back(reader.ReadString());
@@ -3790,7 +3808,7 @@ namespace uba
 			snprintf(result, sizeof(result), "popen failed with command: %s", command.data);
 		writer.WriteBool(success);
 		writer.WriteString(result);
-		#endif
+#endif
 		return true;
 	}
 
diff --git a/Engine/Source/Programs/UnrealBuildAccelerator/Detours/Private/Windows/UbaDetoursFunctionsCombase.inl b/Engine/Source/Programs/UnrealBuildAccelerator/Detours/Private/Windows/UbaDetoursFunctionsCombase.inl
--- a/Engine/Source/Programs/UnrealBuildAccelerator/Detours/Private/Windows/UbaDetoursFunctionsCombase.inl
+++ b/Engine/Source/Programs/UnrealBuildAccelerator/Detours/Private/Windows/UbaDetoursFunctionsCombase.inl
@@ -2,14 +2,14 @@
 
 HRESULT Detoured_CoCreateInstance(REFCLSID rclsid, LPUNKNOWN pUnkOuter, DWORD dwClsContext, REFIID riid, LPVOID *ppv)
 {
-	DEBUG_LOG_DETOURED(L"CoCreateInstance", L"%s", GuidToString((Guid&)rclsid).data);
-
 	// Special handling for com object used to figure out visual studio installation. This is used as a fallback path for lld-link.exe among others
 	// {177F0C4A-1CD3-4DE7-A32C-71DBBB9FA36D}
 	static const CLSID CLSID_SetupConfiguration = { 0x177F0C4A, 0x1CD3, 0x4DE7, { 0xA3, 0x2C, 0x71, 0xDB, 0xBB, 0x9F, 0xA3, 0x6D } };
 
 	if (IsEqualCLSID(rclsid, CLSID_SetupConfiguration))
 	{
+		DEBUG_LOG_DETOURED(L"CoCreateInstance", L"(VisualStudioSetup) %s", GuidToString((Guid&)rclsid).data);
+
 		// TODO: Should probably not be hard coded and instead be transferred from host
 		const wchar_t* dllPath = L"c:\\ProgramData\\Microsoft\\VisualStudio\\Setup\\x64\\Microsoft.VisualStudio.Setup.Configuration.Native.dll";
 		HMODULE hMod = LoadLibraryW(dllPath);
@@ -30,6 +30,6 @@ HRESULT Detoured_CoCreateInstance(REFCLSID rclsid, LPUNKNOWN pUnkOuter, DWORD dw
 		return hr;
 	}
 
-
+	DEBUG_LOG_DETOURED(L"CoCreateInstance", L"%s", GuidToString((Guid&)rclsid).data);
 	return True_CoCreateInstance(rclsid, pUnkOuter, dwClsContext, riid, ppv);
 }
\ No newline at end of file
diff --git a/Engine/Source/Programs/UnrealBuildAccelerator/Detours/Private/Windows/UbaDetoursFunctionsNtDll.inl b/Engine/Source/Programs/UnrealBuildAccelerator/Detours/Private/Windows/UbaDetoursFunctionsNtDll.inl
--- a/Engine/Source/Programs/UnrealBuildAccelerator/Detours/Private/Windows/UbaDetoursFunctionsNtDll.inl
+++ b/Engine/Source/Programs/UnrealBuildAccelerator/Detours/Private/Windows/UbaDetoursFunctionsNtDll.inl
@@ -1803,7 +1803,7 @@ NTSTATUS NTAPI Detoured_NtReadFile(HANDLE FileHandle, HANDLE Event, PIO_APC_ROUT
 			(void)prefix;
 			TimerScope ts(g_kernelStats.readFile);
 			NTSTATUS status = True_NtReadFile(fileHandle, Event, ApcRoutine, ApcContext, IoStatusBlock, Buffer, Length, ByteOffset, Key);
-			DEBUG_LOG_TRUE(L"NtReadFile", L"%s%llu (%ls) -> %s", prefix, uintptr_t(fileHandle), HandleToName(fileHandle), ToString(status).data);
+			DEBUG_LOG_TRUE(L"NtReadFile", L"%s%llu (%ls) (%llu/%llu) -> %s", prefix, uintptr_t(fileHandle), HandleToName(fileHandle), IoStatusBlock->Information, u64(Length), ToString(status).data);
 			return status;
 		};
 
diff --git a/Engine/Source/Programs/UnrealBuildAccelerator/Detours/Private/Windows/UbaDetoursFunctionsShell32.inl b/Engine/Source/Programs/UnrealBuildAccelerator/Detours/Private/Windows/UbaDetoursFunctionsShell32.inl
--- a/Engine/Source/Programs/UnrealBuildAccelerator/Detours/Private/Windows/UbaDetoursFunctionsShell32.inl
+++ b/Engine/Source/Programs/UnrealBuildAccelerator/Detours/Private/Windows/UbaDetoursFunctionsShell32.inl
@@ -1,39 +1,83 @@
 // Copyright Epic Games, Inc. All Rights Reserved.
 
+HRESULT Detoured_SHGetKnownFolderPath(REFKNOWNFOLDERID rfid, DWORD dwFlags, HANDLE hToken, PWSTR* ppszPath)
+{
+	if (!g_runningRemote)
+	{
+		SuppressDetourScope _;
+		HRESULT res = True_SHGetKnownFolderPath(rfid, dwFlags, hToken, ppszPath);
+		DEBUG_LOG_TRUE(L"SHGetKnownFolderPath", L"(%ls) -> %ls", *ppszPath, ToString(res == S_OK));
+		return res;
+	}
 
-//typedef struct _MODLOAD_DATA {} *PMODLOAD_DATA;
-using SHGetKnownFolderPathFunc = HRESULT(REFKNOWNFOLDERID rfid, DWORD dwFlags, HANDLE hToken, PWSTR* ppszPath);
-SHGetKnownFolderPathFunc* True_SHGetKnownFolderPath;
+	UBA_ASSERT(hToken == NULL);
+	RPC_MESSAGE(SHGetKnownFolderPath, getFullFileName)
+	writer.WriteBytes(&rfid, sizeof(KNOWNFOLDERID));
+	writer.WriteU32(dwFlags);
+	BinaryReader reader = writer.Flush();
+	HRESULT res = reader.ReadU32();
+	*ppszPath = NULL;
+	if (res == S_OK)
+	{
+		StringBuffer<> path;
+		reader.ReadString(path);
+		u32 memSize = (path.count+1)*2;
+		void* mem = CoTaskMemAlloc(memSize);
+		UBA_ASSERT(mem);
+		if (!mem)
+			return E_FAIL;
+		memcpy(mem, path.data, memSize);
+		*ppszPath = (PWSTR)mem;
+	}
+	DEBUG_LOG_DETOURED(L"SHGetKnownFolderPath", L"(%ls) -> %ls", *ppszPath, ToString(res == S_OK));
+	return res;
+}
 
-HRESULT Detoured_SHGetKnownFolderPath(REFKNOWNFOLDERID rfid, DWORD dwFlags, HANDLE hToken, PWSTR* ppszPath)
+
+HRESULT Detoured_SHGetFolderPathAndSubDirW(HWND hwnd, int csidl, HANDLE hToken, DWORD dwFlags, LPCWSTR pszSubDir, LPWSTR pszPath)
 {
-	if (g_runningRemote)
+	if (!g_runningRemote)
 	{
-		UBA_ASSERT(hToken == NULL);
-		RPC_MESSAGE(SHGetKnownFolderPath, getFullFileName)
-		writer.WriteBytes(&rfid, sizeof(KNOWNFOLDERID));
-		writer.WriteU32(dwFlags);
-		BinaryReader reader = writer.Flush();
-		HRESULT res = reader.ReadU32();
-		*ppszPath = NULL;
-		if (res == S_OK)
-		{
-			StringBuffer<> path;
-			reader.ReadString(path);
-			u32 memSize = (path.count+1)*2;
-			void* mem = CoTaskMemAlloc(memSize);
-			UBA_ASSERT(mem);
-			if (!mem)
-				return E_FAIL;
-			memcpy(mem, path.data, memSize);
-			*ppszPath = (PWSTR)mem;
-		}
-		DEBUG_LOG_DETOURED(L"SHGetKnownFolderPath", L"(%ls) -> %ls", *ppszPath, ToString(res == S_OK));
+		SuppressDetourScope _;
+		HRESULT res = True_SHGetFolderPathAndSubDirW(hwnd, csidl, hToken, dwFlags, pszSubDir, pszPath);
+		DEBUG_LOG_TRUE(L"SHGetFolderPathAndSubDirW", L"(%s) -> %s", pszPath, ToString(res == S_OK));
 		return res;
 	}
 
+	UBA_ASSERT(hwnd == NULL);
+	UBA_ASSERT(hToken == NULL);
+
+	RPC_MESSAGE(HostRun, getFullFileName);
+	u16& size = *(u16*)writer.AllocWrite(2);
+	u64 pos = writer.GetPosition();
+	writer.WriteU32(csidl);
+	writer.WriteU32(dwFlags);
+	writer.WriteBool(pszSubDir != nullptr);
+	if (pszSubDir)
+		writer.WriteString(pszSubDir);
+	size = u16(writer.GetPosition() - pos);
+	BinaryReader reader = writer.Flush();
+	u32 res = reader.ReadU32();
+	reader.ReadString(pszPath, MAX_PATH);
+
+	DEBUG_LOG_DETOURED(L"SHGetFolderPathAndSubDirW", L"%s", pszPath);
+	return res;
+}
+
+#if UBA_DEBUG
+HRESULT Detoured_SHGetFolderPathW(HWND hwnd, int csidl, HANDLE hToken, DWORD dwFlags, LPWSTR pszPath)
+{
 	SuppressDetourScope _;
-	HRESULT res = True_SHGetKnownFolderPath(rfid, dwFlags, hToken, ppszPath);
-	DEBUG_LOG_TRUE(L"SHGetKnownFolderPath", L"(%ls) -> %ls", *ppszPath, ToString(res == S_OK));
+	HRESULT res = True_SHGetFolderPathW(hwnd, csidl, hToken, dwFlags, pszPath);
+	DEBUG_LOG_TRUE(L"SHGetFolderPathW", L"(%s) -> %s", pszPath, ToString(res == S_OK));
 	return res;
 }
+
+HRESULT Detoured_SHGetKnownFolderIDList(REFKNOWNFOLDERID rfid, DWORD flags, HANDLE token, PIDLIST_ABSOLUTE *pidl)
+{
+	HRESULT res = True_SHGetKnownFolderIDList(rfid, flags, token, pidl);
+	DEBUG_LOG_TRUE(L"SHGetKnownFolderIDList", L"-> %s", ToString(res == S_OK));
+	return res;
+
+}
+#endif
\ No newline at end of file
diff --git a/Engine/Source/Programs/UnrealBuildAccelerator/Detours/Private/Windows/UbaDetoursFunctionsWin.cpp b/Engine/Source/Programs/UnrealBuildAccelerator/Detours/Private/Windows/UbaDetoursFunctionsWin.cpp
--- a/Engine/Source/Programs/UnrealBuildAccelerator/Detours/Private/Windows/UbaDetoursFunctionsWin.cpp
+++ b/Engine/Source/Programs/UnrealBuildAccelerator/Detours/Private/Windows/UbaDetoursFunctionsWin.cpp
@@ -728,6 +728,11 @@ int DetourAttachFunctions(bool runningRemote)
 		DETOURED_FUNCTIONS_SHLWAPI
 	}
 
+	if (HMODULE moduleHandle = GetModuleHandleW(L"shell32.dll"))
+	{
+		DETOURED_FUNCTIONS_SHELL32
+	}
+
 	#if UBA_SUPPORT_MSPDBSRV
 	if (HMODULE moduleHandle = GetModuleHandleW(L"rpcrt4.dll"))
 	{
@@ -796,10 +801,13 @@ void OnModuleLoaded(HMODULE moduleHandle, const StringView& name)
 	// SHGetKnownFolderPath is used by Metal.exe and must always execute on host
 	if (!True_SHGetKnownFolderPath && name.Contains(L"shell32.dll"))
 	{
-		True_SHGetKnownFolderPath = (SHGetKnownFolderPathFunc*)GetProcAddress(moduleHandle, "SHGetKnownFolderPath");
-		UBA_ASSERT(True_SHGetKnownFolderPath);
 		DetourTransactionBegin();
-		DetourAttachFunction((PVOID*)&True_SHGetKnownFolderPath, Detoured_SHGetKnownFolderPath, "SHGetKnownFolderPath");
+		#define DETOURED_FUNCTION(Func) True_##Func = (decltype(True_##Func))GetProcAddress(moduleHandle, #Func);
+		DETOURED_FUNCTIONS_SHELL32
+		#undef DETOURED_FUNCTION
+		#define DETOURED_FUNCTION(Func) DetourAttachFunction((PVOID*)&True_##Func, Detoured_##Func, #Func);
+		DETOURED_FUNCTIONS_SHELL32
+		#undef DETOURED_FUNCTION
 		DetourTransactionCommit();
 	}
 }
diff --git a/Engine/Source/Programs/UnrealBuildAccelerator/Detours/Private/Windows/UbaDetoursFunctionsWin.h b/Engine/Source/Programs/UnrealBuildAccelerator/Detours/Private/Windows/UbaDetoursFunctionsWin.h
--- a/Engine/Source/Programs/UnrealBuildAccelerator/Detours/Private/Windows/UbaDetoursFunctionsWin.h
+++ b/Engine/Source/Programs/UnrealBuildAccelerator/Detours/Private/Windows/UbaDetoursFunctionsWin.h
@@ -29,6 +29,7 @@
 	DETOURED_FUNCTIONS_KERNEL32 \
 	DETOURED_FUNCTIONS_NTDLL \
 	DETOURED_FUNCTIONS_SHLWAPI \
+	DETOURED_FUNCTIONS_SHELL32 \
 	DETOURED_FUNCTIONS_UCRTBASE \
 	DETOURED_FUNCTIONS_RPCRT4 \
 	DETOURED_FUNCTIONS_COMBASE \
@@ -160,6 +161,11 @@
 #define DETOURED_FUNCTIONS_SHLWAPI \
 	DETOURED_FUNCTIONS_SHLWAPI_DEBUG
 
+#define DETOURED_FUNCTIONS_SHELL32 \
+	DETOURED_FUNCTION(SHGetFolderPathAndSubDirW) \
+	DETOURED_FUNCTION(SHGetKnownFolderPath) \
+	DETOURED_FUNCTIONS_SHELL32_DEBUG
+
 #if !defined(__clang__)
 #define DETOURED_WSPLITPATH DETOURED_FUNCTION(_wsplitpath_s)
 #else
@@ -386,6 +392,9 @@
 	DETOURED_FUNCTION(SHCreateStreamOnFileW) \
 	DETOURED_FUNCTION(PathFileExistsW) \
 
+#define DETOURED_FUNCTIONS_SHELL32_DEBUG \
+	DETOURED_FUNCTION(SHGetKnownFolderIDList) \
+	DETOURED_FUNCTION(SHGetFolderPathW) \
 
 #define DETOURED_FUNCTIONS_UCRTBASE_DEBUG \
 	DETOURED_FUNCTION(_wcsnicoll_l) \
@@ -422,6 +431,7 @@ DETOURED_FUNCTION(OpenEncryptedFileRawA) \
 #define DETOURED_FUNCTIONS_KERNEL32_DEBUG
 #define DETOURED_FUNCTIONS_NTDLL_DEBUG
 #define DETOURED_FUNCTIONS_SHLWAPI_DEBUG
+#define DETOURED_FUNCTIONS_SHELL32_DEBUG
 #define DETOURED_FUNCTIONS_UCRTBASE_DEBUG
 #define DETOURED_FUNCTIONS_MEMORY_DEBUG
 #endif
@@ -483,6 +493,10 @@ extern "C" {
 	SIZE_T NTAPI RtlSizeHeap(HANDLE HeapPtr, ULONG Flags, PVOID Ptr);
 	void NTAPI BaseThreadInitThunk(ULONG Unknown,LPTHREAD_START_ROUTINE StartAddress,PVOID ThreadParameter);
 	void* _expand_base(void* memblock, size_t size);
+	HRESULT SHGetFolderPathAndSubDirW(HWND hwnd, int csidl, HANDLE hToken, DWORD dwFlags, LPCWSTR pszSubDir, LPWSTR pszPath);
+	HRESULT SHGetKnownFolderPath(REFKNOWNFOLDERID rfid, DWORD dwFlags, HANDLE hToken, PWSTR* ppszPath);
+	HRESULT SHGetKnownFolderIDList(REFKNOWNFOLDERID rfid, DWORD flags, HANDLE token, PIDLIST_ABSOLUTE *pidl);
+	HRESULT SHGetFolderPathW(HWND hwnd, int csidl, HANDLE hToken, DWORD dwFlags, LPWSTR pszPath);
 }
 
 
