[UBA]
* Fixed FileFsDeviceInformation returning wrong values in FILE_FS_DEVICE_INFORMATION

diff --git a/Engine/Source/Programs/UnrealBuildAccelerator/Core/Public/UbaPathUtils.h b/Engine/Source/Programs/UnrealBuildAccelerator/Core/Public/UbaPathUtils.h
--- a/Engine/Source/Programs/UnrealBuildAccelerator/Core/Public/UbaPathUtils.h
+++ b/Engine/Source/Programs/UnrealBuildAccelerator/Core/Public/UbaPathUtils.h
@@ -121,7 +121,11 @@ namespace uba
 			else if (c == '.' && lastChar == '.')
 				hasDotDot = true;
 			else if (c == ':')
+			{
+				if (isAbsolute)
+					return false;
 				isAbsolute = true;
+			}
 			if (c == '\\')
 			{
 				if (lastChar == '.')
diff --git a/Engine/Source/Programs/UnrealBuildAccelerator/Core/Public/UbaWinTypes.h b/Engine/Source/Programs/UnrealBuildAccelerator/Core/Public/UbaWinTypes.h
--- a/Engine/Source/Programs/UnrealBuildAccelerator/Core/Public/UbaWinTypes.h
+++ b/Engine/Source/Programs/UnrealBuildAccelerator/Core/Public/UbaWinTypes.h
@@ -64,8 +64,10 @@ struct FILE_FS_DEVICE_INFORMATION
 	DEVICE_TYPE DeviceType;
 	ULONG Characteristics;
 };
+constexpr ULONG FILE_REMOTE_DEVICE = 0x00000010;
 constexpr ULONG FILE_DEVICE_IS_MOUNTED = 0x00000020;
 
+
 struct FILE_FS_ATTRIBUTE_INFORMATION
 {
 	ULONG FileSystemAttributes;
diff --git a/Engine/Source/Programs/UnrealBuildAccelerator/Detours/Private/Windows/UbaDetoursFunctionsKernelBase.inl b/Engine/Source/Programs/UnrealBuildAccelerator/Detours/Private/Windows/UbaDetoursFunctionsKernelBase.inl
--- a/Engine/Source/Programs/UnrealBuildAccelerator/Detours/Private/Windows/UbaDetoursFunctionsKernelBase.inl
+++ b/Engine/Source/Programs/UnrealBuildAccelerator/Detours/Private/Windows/UbaDetoursFunctionsKernelBase.inl
@@ -841,6 +841,7 @@ DWORD Detoured_GetFileAttributesW(LPCWSTR lpFileName)
 	if (!FixPath(fixedPath, lpFileName))
 	{
 		DEBUG_LOG_DETOURED(L"GetFileAttributesW", L"(BAD_PATH) (%ls) -> INVALID_FILE_ATTRIBUTES", lpFileName);
+		SetLastError(ERROR_INVALID_NAME);
 		return INVALID_FILE_ATTRIBUTES;
 	}
 	DevirtualizePath(fixedPath);
diff --git a/Engine/Source/Programs/UnrealBuildAccelerator/Detours/Private/Windows/UbaDetoursFunctionsNtDll.inl b/Engine/Source/Programs/UnrealBuildAccelerator/Detours/Private/Windows/UbaDetoursFunctionsNtDll.inl
--- a/Engine/Source/Programs/UnrealBuildAccelerator/Detours/Private/Windows/UbaDetoursFunctionsNtDll.inl
+++ b/Engine/Source/Programs/UnrealBuildAccelerator/Detours/Private/Windows/UbaDetoursFunctionsNtDll.inl
@@ -74,6 +74,16 @@ NTSTATUS Detoured_NtQueryVolumeInformationFile(HANDLE FileHandle, PIO_STATUS_BLO
 	IoStatusBlock->Status = STATUS_SUCCESS;
 	IoStatusBlock->Information = 0;
 
+	auto runDefaultFsDeviceInfo = [&]()
+		{
+			auto& info = *(FILE_FS_DEVICE_INFORMATION*)FsInformation;
+			info.DeviceType = FILE_DEVICE_DISK;
+			info.Characteristics = FILE_DEVICE_IS_MOUNTED; // | FILE_REMOTE_DEVICE; // Use remote device to turn off some optimizations
+			IoStatusBlock->Information = sizeof(FILE_FS_DEVICE_INFORMATION);
+			DEBUG_LOG_DETOURED(L"NtQueryVolumeInformationFile", L"%llu (FileFsDeviceInformation) (%ls) -> %ls", uintptr_t(FileHandle), HandleToName(FileHandle), ToString(STATUS_SUCCESS).data);
+			return STATUS_SUCCESS;
+		};
+
 	DETOURED_CALL(NtQueryVolumeInformationFile);
 	HANDLE TrueHandle = FileHandle;
 	if (isDetouredHandle(FileHandle))
@@ -82,14 +92,7 @@ NTSTATUS Detoured_NtQueryVolumeInformationFile(HANDLE FileHandle, PIO_STATUS_BLO
 		if (dh.fileObject->fileInfo->memoryFile)
 		{
 			if (FsInformationClass == FileFsDeviceInformation)
-			{
-				auto& info = *(FILE_FS_DEVICE_INFORMATION*)FsInformation;
-				info.DeviceType = FILE_DEVICE_FILE_SYSTEM; //FILE_DEVICE_DISK;
-				info.Characteristics = FILE_DEVICE_IS_MOUNTED;
-				IoStatusBlock->Information = sizeof(FILE_FS_DEVICE_INFORMATION);
-				DEBUG_LOG_DETOURED(L"NtQueryVolumeInformationFile", L"%llu (FileFsDeviceInformation) (%ls) -> %ls", uintptr_t(FileHandle), HandleToName(FileHandle), ToString(STATUS_SUCCESS).data);
-				return STATUS_SUCCESS;
-			}
+				return runDefaultFsDeviceInfo();
 		}
 		TrueHandle = dh.trueHandle;
 		if (TrueHandle == INVALID_HANDLE_VALUE)
@@ -121,28 +124,15 @@ NTSTATUS Detoured_NtQueryVolumeInformationFile(HANDLE FileHandle, PIO_STATUS_BLO
 				return STATUS_SUCCESS;
 			}
 			if (FsInformationClass == FileFsDeviceInformation)
-			{
-				auto& info = *(FILE_FS_DEVICE_INFORMATION*)FsInformation;
-				info.DeviceType = FILE_DEVICE_FILE_SYSTEM; //FILE_DEVICE_DISK;
-				info.Characteristics = FILE_DEVICE_IS_MOUNTED;
-				IoStatusBlock->Information = sizeof(FILE_FS_DEVICE_INFORMATION);
-				DEBUG_LOG_DETOURED(L"NtQueryVolumeInformationFile", L"%llu (FileFsDeviceInformation) (%ls) -> %ls", uintptr_t(FileHandle), HandleToName(FileHandle), ToString(STATUS_SUCCESS).data);
-				return STATUS_SUCCESS;
-			}
+				return runDefaultFsDeviceInfo();
 			UBA_ASSERTF(false, L"NtQueryVolumeInformationFile using class %u not handled %ls (%ls)", FsInformationClass, dh.fileObject->fileInfo->name, dh.fileObject->fileInfo->originalName);
 		}
 	}
 	else if (isListDirectoryHandle(FileHandle))
 	{
 		if (FsInformationClass == FileFsDeviceInformation)
-		{
-			auto& info = *(FILE_FS_DEVICE_INFORMATION*)FsInformation;
-			info.DeviceType = FILE_DEVICE_FILE_SYSTEM; //FILE_DEVICE_DISK;
-			info.Characteristics = FILE_DEVICE_IS_MOUNTED;
-			IoStatusBlock->Information = sizeof(FILE_FS_DEVICE_INFORMATION);
-			DEBUG_LOG_DETOURED(L"NtQueryVolumeInformationFile", L"%llu (FileFsDeviceInformation) (%ls) -> %ls", uintptr_t(FileHandle), HandleToName(FileHandle), ToString(STATUS_SUCCESS).data);
-			return STATUS_SUCCESS;
-		}
+			return runDefaultFsDeviceInfo();
+
 		if (FsInformationClass == FileFsVolumeInformation)
 		{
 			// TODO This code path is here to handle nodejs queries..
@@ -498,7 +488,7 @@ NTSTATUS NTAPI Detoured_NtQueryDirectoryFile(HANDLE FileHandle, HANDLE Event, PI
 			if (entryInfo.attributes == 0) // File was deleted
 				continue;
 
-			if (FileName && wcsncmp(FileName->Buffer, fileName, FileName->Length / 2) != 0)
+			if (FileName && !StringView(FileName->Buffer, FileName->Length / 2).Equals(fileName))
 				continue;
 
 			u32 fileNameBytes = u32(wcslen(fileName) * 2);
