[UBA]
* Fixed a bug when creating FILE_NAME_INFORMATION in NtQueryInformationFile
* Added detouring of RtlGetVersion, GetFileInformationByName and GetProcAddress (only for logging)
* Improved lots of logging info

diff --git a/Engine/Source/Programs/UnrealBuildAccelerator/Core/Public/UbaWinTypes.h b/Engine/Source/Programs/UnrealBuildAccelerator/Core/Public/UbaWinTypes.h
--- a/Engine/Source/Programs/UnrealBuildAccelerator/Core/Public/UbaWinTypes.h
+++ b/Engine/Source/Programs/UnrealBuildAccelerator/Core/Public/UbaWinTypes.h
@@ -205,6 +205,7 @@ enum FS_INFORMATION_CLASS {};
 
 constexpr FS_INFORMATION_CLASS FileFsVolumeInformation = (FS_INFORMATION_CLASS)1;
 constexpr FS_INFORMATION_CLASS FileFsDeviceInformation = (FS_INFORMATION_CLASS)4;
+constexpr FS_INFORMATION_CLASS FileFsAttributeInformation = (FS_INFORMATION_CLASS)5;
 
 enum RTL_PATH_TYPE
 {
diff --git a/Engine/Source/Programs/UnrealBuildAccelerator/Detours/Private/Windows/UbaDetoursFunctionsKernelBase.inl b/Engine/Source/Programs/UnrealBuildAccelerator/Detours/Private/Windows/UbaDetoursFunctionsKernelBase.inl
--- a/Engine/Source/Programs/UnrealBuildAccelerator/Detours/Private/Windows/UbaDetoursFunctionsKernelBase.inl
+++ b/Engine/Source/Programs/UnrealBuildAccelerator/Detours/Private/Windows/UbaDetoursFunctionsKernelBase.inl
@@ -547,8 +547,9 @@ BOOL Detoured_GetVolumeInformationByHandleW(HANDLE hFile, LPWSTR lpVolumeNameBuf
 		DEBUG_LOG_DETOURED(L"GetVolumeInformationByHandleW", L"%llu (Serial: %u) (%ls) -> Success", uintptr_t(hFile), entryInfo.volumeSerial, HandleToName(hFile));
 		return true;
 	}
-	bool res = True_GetVolumeInformationByHandleW(trueHandle, lpVolumeNameBuffer, nVolumeNameSize, lpVolumeSerialNumber, lpMaximumComponentLength, lpFileSystemFlags, lpFileSystemNameBuffer, nFileSystemNameSize);
-	return res;
+
+	DEBUG_LOG_TRUE(L"GetVolumeInformationByHandleW", L"%llu (%ls)", u64(hFile), HandleToName(hFile));
+	return True_GetVolumeInformationByHandleW(trueHandle, lpVolumeNameBuffer, nVolumeNameSize, lpVolumeSerialNumber, lpMaximumComponentLength, lpFileSystemFlags, lpFileSystemNameBuffer, nFileSystemNameSize);
 }
 
 BOOL Detoured_GetVolumeInformationW(LPCWSTR lpRootPathName, LPWSTR lpVolumeNameBuffer, DWORD nVolumeNameSize, LPDWORD lpVolumeSerialNumber, LPDWORD lpMaximumComponentLength, LPDWORD lpFileSystemFlags, LPWSTR lpFileSystemNameBuffer, DWORD nFileSystemNameSize)
@@ -572,6 +573,8 @@ BOOL Detoured_GetVolumeInformationW(LPCWSTR lpRootPathName, LPWSTR lpVolumeNameB
 		DEBUG_LOG_DETOURED(L"GetVolumeInformationW", L"%ls", lpRootPathName);
 		return true;
 	}
+
+	DEBUG_LOG_TRUE(L"GetVolumeInformationW", L"%ls", lpRootPathName);
 	SuppressCreateFileDetourScope s;
 	auto res = True_GetVolumeInformationW(lpRootPathName, lpVolumeNameBuffer, nVolumeNameSize, lpVolumeSerialNumber, lpMaximumComponentLength, lpFileSystemFlags, lpFileSystemNameBuffer, nFileSystemNameSize);
 	return res;
@@ -3690,6 +3693,66 @@ BOOL Detoured_GetDiskFreeSpaceExA(LPCSTR lpDirectoryName, PULARGE_INTEGER lpFree
 	return True_GetDiskFreeSpaceExA(lpDirectoryName, lpFreeBytesAvailableToCaller, lpTotalNumberOfBytes, lpTotalNumberOfFreeBytes);
 }
 
+BOOL Detoured_GetFileInformationByName(PCWSTR FileName, FILE_INFO_BY_NAME_CLASS FileInformationClass, PVOID FileInfoBuffer, ULONG FileInfoBufferSize)
+{
+	DETOURED_CALL(GetFileInformationByName);
+
+#if 0
+	enum FILE_INFO_BY_NAME_CLASS
+	{
+		FileStatByNameInfo,
+		FileStatLxByNameInfo,
+		FileCaseSensitiveByNameInfo,
+		FileStatBasicByNameInfo,
+		MaximumFileInfoByNameClass
+	};
+	
+	if (FileInformationClass == FileStatBasicByNameInfo)
+	{
+		struct FILE_STAT_BASIC_INFORMATION
+		{
+			LARGE_INTEGER FileId;
+			LARGE_INTEGER CreationTime;
+			LARGE_INTEGER LastAccessTime;
+			LARGE_INTEGER LastWriteTime;
+			LARGE_INTEGER ChangeTime;
+			LARGE_INTEGER AllocationSize;
+			LARGE_INTEGER EndOfFile;
+			ULONG FileAttributes;
+			ULONG ReparseTag;
+			ULONG NumberOfLinks;
+			ACCESS_MASK EffectiveAccess;
+		};
+		FILE_STAT_BASIC_INFORMATION
+	}
+#endif
+	if (g_runningRemote)
+	{
+		SetLastError(ERROR_CALL_NOT_IMPLEMENTED);
+		DEBUG_LOG_DETOURED(L"GetFileInformationByName", L"(%ls) (class %u) -> %s (ERROR_CALL_NOT_IMPLEMENTED)", FileName, FileInformationClass, ToString(false));
+		return false;
+	}
+
+	StringBuffer<> pathName;
+	FixPath(pathName, FileName);
+	DevirtualizePath(pathName);
+	StringKey sourceKey = ToStringKeyLower(pathName);
+
+	SCOPED_WRITE_LOCK(g_mappedFileTable.m_lookupLock, lock);
+	auto it = g_mappedFileTable.m_lookup.find(sourceKey);
+	if (it != g_mappedFileTable.m_lookup.end())
+	{
+		FileInfo& fi = it->second;
+		UBA_ASSERT(!fi.memoryFile);
+	}
+	lock.Leave();
+
+
+	BOOL res = True_GetFileInformationByName(pathName.data, FileInformationClass, FileInfoBuffer, FileInfoBufferSize);
+	DEBUG_LOG_TRUE(L"GetFileInformationByName", L"(%ls) (class %u) -> %s", pathName.data, FileInformationClass, ToString(res));
+	return res;
+}
+
 DWORD Detoured_GetLongPathNameA(LPCSTR lpszShortPath, LPSTR lpszLongPath, DWORD cchBuffer)
 {
 	DETOURED_CALL(GetLongPathNameA);
@@ -4413,6 +4476,13 @@ LPWSTR* Detoured_CommandLineToArgvW(LPCWSTR lpCmdLine, int *pNumArgs)
 	return res;
 }
 
+void* Detoured_GetProcAddress(HMODULE hModule, LPCSTR lpProcName)
+{
+	void* res = True_GetProcAddress(hModule, lpProcName);
+	DEBUG_LOG_TRUE(L"GetProcAddress", L"%S -> 0x%p", lpProcName, res);
+	return res;
+}
+
 BOOL Detoured_FreeLibrary(HMODULE hModule)
 {
 	DETOURED_CALL(FreeLibrary);
diff --git a/Engine/Source/Programs/UnrealBuildAccelerator/Detours/Private/Windows/UbaDetoursFunctionsNtDll.inl b/Engine/Source/Programs/UnrealBuildAccelerator/Detours/Private/Windows/UbaDetoursFunctionsNtDll.inl
--- a/Engine/Source/Programs/UnrealBuildAccelerator/Detours/Private/Windows/UbaDetoursFunctionsNtDll.inl
+++ b/Engine/Source/Programs/UnrealBuildAccelerator/Detours/Private/Windows/UbaDetoursFunctionsNtDll.inl
@@ -1,10 +1,11 @@
 // Copyright Epic Games, Inc. All Rights Reserved.
 
+#if UBA_DEBUG_LOG_ENABLED
 StringBuffer<64> ToString(FILE_INFORMATION_CLASS c)
 {
 	switch (u32(c))
 	{
-		#define FILE_INFO_CLASS(name, value) case value: return StringBuffer<64>(TC(#name));
+		#define FILE_INFO_CLASS(name, value) case value: return StringBuffer<64>(TCV(#name));
 		FILE_INFO_CLASSES
 		#undef FILE_INFO_CLASS
 	default:
@@ -12,6 +13,18 @@ StringBuffer<64> ToString(FILE_INFORMATION_CLASS c)
 	}
 }
 
+StringBuffer<64> ToString(FS_INFORMATION_CLASS c)
+{
+	switch (u32(c))
+	{
+	case FileFsVolumeInformation: return StringBuffer<64>(TCV("FileFsVolumeInformation"));
+	case FileFsDeviceInformation: return StringBuffer<64>(TCV("FileFsDeviceInformation"));
+	case FileFsAttributeInformation: return StringBuffer<64>(TCV("FileFsAttributeInformation"));
+	default:
+		return StringBuffer<64>(TCV("class ")).AppendValue(u32(c));
+	}
+}
+#endif
 
 bool IsContentWrite(u32 desiredAccess, u32 createDisposition)
 {
@@ -84,6 +97,32 @@ NTSTATUS Detoured_NtQueryVolumeInformationFile(HANDLE FileHandle, PIO_STATUS_BLO
 			return STATUS_SUCCESS;
 		};
 
+	auto runDefaultFsVolumeInfo = [&](u32 entryOffset)
+		{
+			ULONG required = offsetof(FILE_FS_VOLUME_INFORMATION, VolumeLabel);
+			if (Length < required)
+			{
+				IoStatusBlock->Status = STATUS_BUFFER_TOO_SMALL;
+				IoStatusBlock->Information = 0;
+				DEBUG_LOG_DETOURED(L"NtQueryVolumeInformationFile", L"%llu (FileFsVolumeInformation) (%ls) -> %ls", uintptr_t(FileHandle), HandleToName(FileHandle), ToString(STATUS_BUFFER_TOO_SMALL).data);
+				return STATUS_BUFFER_TOO_SMALL;
+			}
+
+			auto& info = *(FILE_FS_VOLUME_INFORMATION*)FsInformation;
+			UBA_ASSERT(entryOffset != ~0u);
+			DirectoryTable::EntryInformation entryInfo;
+			g_directoryTable.GetEntryInformation(entryInfo, entryOffset);
+			UBA_ASSERT(entryInfo.attributes != 0);
+			info.VolumeCreationTime.QuadPart = 123;
+			info.VolumeSerialNumber = entryInfo.volumeSerial;
+			info.VolumeLabelLength = 0;
+			info.SupportsObjects = true;
+			info.VolumeLabel[0] = 0;
+			IoStatusBlock->Information = offsetof(FILE_FS_VOLUME_INFORMATION, VolumeLabel) + info.VolumeLabelLength;
+			DEBUG_LOG_DETOURED(L"NtQueryVolumeInformationFile", L"%llu (FileFsVolumeInformation) (%ls) -> %ls", uintptr_t(FileHandle), HandleToName(FileHandle), ToString(STATUS_SUCCESS).data);
+			return STATUS_SUCCESS;
+		};
+
 	DETOURED_CALL(NtQueryVolumeInformationFile);
 	HANDLE TrueHandle = FileHandle;
 	if (isDetouredHandle(FileHandle))
@@ -98,31 +137,7 @@ NTSTATUS Detoured_NtQueryVolumeInformationFile(HANDLE FileHandle, PIO_STATUS_BLO
 		if (TrueHandle == INVALID_HANDLE_VALUE)
 		{
 			if (FsInformationClass == FileFsVolumeInformation)
-			{
-				// TODO This code path is here to handle nodejs queries..
-				ULONG required = offsetof(FILE_FS_VOLUME_INFORMATION, VolumeLabel);
-				if (Length < required)
-				{
-					IoStatusBlock->Status = STATUS_BUFFER_TOO_SMALL;
-					IoStatusBlock->Information = 0;
-					DEBUG_LOG_DETOURED(L"NtQueryVolumeInformationFile", L"%llu (FileFsVolumeInformation) (%ls) -> %ls", uintptr_t(FileHandle), HandleToName(FileHandle), ToString(STATUS_BUFFER_TOO_SMALL).data);
-					return STATUS_BUFFER_TOO_SMALL;
-				}
-
-				auto& info = *(FILE_FS_VOLUME_INFORMATION*)FsInformation;
-				UBA_ASSERT(dh.dirTableOffset != ~0u);
-				DirectoryTable::EntryInformation entryInfo;
-				g_directoryTable.GetEntryInformation(entryInfo, dh.dirTableOffset);
-				UBA_ASSERT(entryInfo.attributes != 0);
-				info.VolumeCreationTime.QuadPart = 0;
-				info.VolumeSerialNumber = entryInfo.volumeSerial;
-				info.VolumeLabelLength = 0;
-				info.SupportsObjects = false;
-				info.VolumeLabel[0] = 0;
-				IoStatusBlock->Information = offsetof(FILE_FS_VOLUME_INFORMATION, VolumeLabel) + info.VolumeLabelLength;
-				DEBUG_LOG_DETOURED(L"NtQueryVolumeInformationFile", L"%llu (FileFsVolumeInformation) (%ls) -> %ls", uintptr_t(FileHandle), HandleToName(FileHandle), ToString(STATUS_SUCCESS).data);
-				return STATUS_SUCCESS;
-			}
+				return runDefaultFsVolumeInfo(dh.dirTableOffset);
 			if (FsInformationClass == FileFsDeviceInformation)
 				return runDefaultFsDeviceInfo();
 			UBA_ASSERTF(false, L"NtQueryVolumeInformationFile using class %u not handled %ls (%ls)", FsInformationClass, dh.fileObject->fileInfo->name, dh.fileObject->fileInfo->originalName);
@@ -130,42 +145,16 @@ NTSTATUS Detoured_NtQueryVolumeInformationFile(HANDLE FileHandle, PIO_STATUS_BLO
 	}
 	else if (isListDirectoryHandle(FileHandle))
 	{
+		auto& listHandle = asListDirectoryHandle(FileHandle);
+
 		if (FsInformationClass == FileFsDeviceInformation)
 			return runDefaultFsDeviceInfo();
-
 		if (FsInformationClass == FileFsVolumeInformation)
-		{
-			// TODO This code path is here to handle nodejs queries..
-			ULONG required = offsetof(FILE_FS_VOLUME_INFORMATION, VolumeLabel);
-			if (Length < required)
-			{
-				IoStatusBlock->Status = STATUS_BUFFER_TOO_SMALL;
-				IoStatusBlock->Information = 0;
-				DEBUG_LOG_DETOURED(L"NtQueryVolumeInformationFile", L"%llu (FileFsVolumeInformation) (%ls) -> %ls", uintptr_t(FileHandle), HandleToName(FileHandle), ToString(STATUS_BUFFER_TOO_SMALL).data);
-				return STATUS_BUFFER_TOO_SMALL;
-			}
-
-			// TODO This code path is here to handle nodejs queries..
-			auto& listHandle = asListDirectoryHandle(FileHandle);
-			u32 entryOffset = listHandle.dir.tableOffset | 0x80000000;
-
-			auto& info = *(FILE_FS_VOLUME_INFORMATION*)FsInformation;
-			DirectoryTable::EntryInformation entryInfo;
-			g_directoryTable.GetEntryInformation(entryInfo, entryOffset);
-			UBA_ASSERT(entryInfo.attributes != 0);
-			info.VolumeCreationTime.QuadPart = 0;
-			info.VolumeSerialNumber = entryInfo.volumeSerial;
-			info.VolumeLabelLength = 0;
-			info.SupportsObjects = false;
-			info.VolumeLabel[0] = 0;
-			IoStatusBlock->Information = offsetof(FILE_FS_VOLUME_INFORMATION, VolumeLabel) + info.VolumeLabelLength;
-			DEBUG_LOG_DETOURED(L"NtQueryVolumeInformationFile", L"%llu (FileFsVolumeInformation) (%ls) -> %ls", uintptr_t(FileHandle), HandleToName(FileHandle), ToString(STATUS_SUCCESS).data);
-			return STATUS_SUCCESS;
-		}
+			return runDefaultFsVolumeInfo(listHandle.dir.tableOffset | 0x80000000);
 		UBA_ASSERTF(false, L"NtQueryVolumeInformationFile called in ListDirectoryHandle using class %u which is not implemented (%ls)", FsInformationClass, HandleToName(FileHandle));
 	}
 	auto res = True_NtQueryVolumeInformationFile(TrueHandle, IoStatusBlock, FsInformation, Length, FsInformationClass);
-	DEBUG_LOG_TRUE(L"NtQueryVolumeInformationFile", L"%llu (class %u) (%ls) -> %ls", uintptr_t(FileHandle), FsInformationClass, HandleToName(FileHandle), ToString(res).data);
+	DEBUG_LOG_TRUE(L"NtQueryVolumeInformationFile", L"%llu (%s) (%ls) -> %ls", uintptr_t(FileHandle), ToString(FsInformationClass).data, HandleToName(FileHandle), ToString(res).data);
 	return res;
 }
 
@@ -292,7 +281,8 @@ NTSTATUS Detoured_NtQueryInformationFile(HANDLE FileHandle, PIO_STATUS_BLOCK IoS
 
 	auto setNameInformation = [&](FILE_NAME_INFORMATION& out, FileInfo& fi, u32 offset)
 		{
-			const wchar_t* fileName = fi.originalName;
+			UBA_ASSERT(fi.originalName[1] == ':');
+			const wchar_t* fileName = fi.originalName + 2; // skip drive letter and colon
 			ULONG nameBytes = (ULONG)(wcslen(fileName) * sizeof(WCHAR));
 			ULONG required = offset + nameBytes;
 			if (Length < offset)
@@ -1744,15 +1734,16 @@ NTSTATUS NTAPI Shared_NtCreateFile(bool IsCreateFunc, PHANDLE hFileHandle, ACCES
 			info.lastDesiredAccess = lastDesiredAccess;
 			Rpc_UpdateCloseHandle(L"", closeId, false, L"", {}, 0, false);
 		}
-		DEBUG_LOG_TRUE(funcName, L"%ls (%ls) (%ls) -> %ls", isWriteStr, lpFileName, (fileName.data != lpFileName ? fileName.data : L""), ToString(res).data);
+		DEBUG_LOG_TRUE(funcName, L"%ls(%ls) (%ls) -> %ls", isWriteStr, lpFileName, (fileName.data != lpFileName ? fileName.data : L""), ToString(res).data);
 		return res;
 	}
 
 	TrackFileInput();
 
+	HANDLE trueHandle = *hFileHandle;
 	UBA_ASSERT(info.originalName);
-	*hFileHandle = CreateFileHandle(*hFileHandle);
-	DEBUG_LOG_TRUE(funcName, L"%ls %llu (%ls)%s (%u %u) -> %ls", isWriteStr, uintptr_t(*hFileHandle), tempFileName.data, isDeleteOnClose ? TC(" DeleteOnClose") : TC(""), CreateDisposition, CreateOptions, ToString(res).data);
+	*hFileHandle = CreateFileHandle(trueHandle);
+	DEBUG_LOG_TRUE(funcName, L"%ls%llu/%llu (%ls)%s (%u %u) -> %ls", isWriteStr, uintptr_t(*hFileHandle), u64(trueHandle), tempFileName.data, isDeleteOnClose ? TC(" DeleteOnClose") : TC(""), CreateDisposition, CreateOptions, ToString(res).data);
 	return res;
 }
 
@@ -1793,7 +1784,7 @@ NTSTATUS NTAPI Detoured_NtReadFile(HANDLE FileHandle, HANDLE Event, PIO_APC_ROUT
 			(void)prefix;
 			TimerScope ts(g_kernelStats.readFile);
 			NTSTATUS status = True_NtReadFile(fileHandle, Event, ApcRoutine, ApcContext, IoStatusBlock, Buffer, Length, ByteOffset, Key);
-			DEBUG_LOG_TRUE(L"NtReadFile", L"%s%llu (%ls) (%llu/%llu) -> %s", prefix, uintptr_t(fileHandle), HandleToName(fileHandle), IoStatusBlock->Information, u64(Length), ToString(status).data);
+			DEBUG_LOG_TRUE(L"NtReadFile", L"%s%llu (%ls) (%llu/%llu) -> %s", prefix, uintptr_t(FileHandle), HandleToName(FileHandle), IoStatusBlock->Information, u64(Length), ToString(status).data);
 			return status;
 		};
 
@@ -2517,6 +2508,14 @@ NTSTATUS NTAPI Detoured_ZwSetInformationFile(HANDLE FileHandle, PIO_STATUS_BLOCK
 	return True_ZwSetInformationFile(FileHandle, IoStatusBlock, FileInformation, Length, FileInformationClass);
 }
 
+NTSTATUS Detoured_RtlGetVersion(PRTL_OSVERSIONINFOW lpVersionInformation)
+{
+	NTSTATUS res = True_RtlGetVersion(lpVersionInformation);
+	DEBUG_LOG_TRUE(L"RtlGetVersion", L"%u.%u.%u -> %s", lpVersionInformation->dwMajorVersion, lpVersionInformation->dwMinorVersion, lpVersionInformation->dwBuildNumber, ToString(res).data);
+	//lpVersionInformation->dwBuildNumber = 26052;
+	return res;
+}
+
 PVOID Detoured_RtlAllocateHeap(PVOID HeapHandle, ULONG Flags, SIZE_T Size)
 {
 	//if (Flags & HEAP_ZERO_MEMORY)
diff --git a/Engine/Source/Programs/UnrealBuildAccelerator/Detours/Private/Windows/UbaDetoursFunctionsWin.cpp b/Engine/Source/Programs/UnrealBuildAccelerator/Detours/Private/Windows/UbaDetoursFunctionsWin.cpp
--- a/Engine/Source/Programs/UnrealBuildAccelerator/Detours/Private/Windows/UbaDetoursFunctionsWin.cpp
+++ b/Engine/Source/Programs/UnrealBuildAccelerator/Detours/Private/Windows/UbaDetoursFunctionsWin.cpp
@@ -27,16 +27,17 @@
 NTSTATUS NtCopyFileChunk(HANDLE Source, HANDLE Dest, HANDLE Event, PIO_STATUS_BLOCK IoStatusBlock, ULONG Length, PULONG SourceOffset, PULONG DestOffset, PULONG SourceKey, PULONG DestKey, ULONG Flags) { return 0; }
 #define NtCurrentProcess() ((HANDLE)-1)
 
+#undef DETOURED_FUNCTION_UNKNOWN
+#define DETOURED_FUNCTION_UNKNOWN(Func) decltype(Func)* True_##Func;
+
 #define DETOURED_FUNCTION(Func) decltype(Func)* True_##Func = ::Func; 
 DETOURED_FUNCTIONS
 DETOURED_FUNCTIONS_MEMORY
 DETOURED_FUNCTIONS_MEMORY_NON_WINE
 #undef DETOURED_FUNCTION
 
-#define DETOURED_FUNCTION(Func) void* Local_##Func = ::Func;
-DETOURED_FUNCTIONS_MEMORY
-DETOURED_FUNCTIONS_MEMORY_NON_WINE
-#undef DETOURED_FUNCTION
+#undef DETOURED_FUNCTION_UNKNOWN
+#define DETOURED_FUNCTION_UNKNOWN(Func) DETOURED_FUNCTION(Func)
 
 #if !DETOURED_INCLUDE_DEBUG
 #define True_GetEnvironmentVariableW GetEnvironmentVariableW
@@ -699,10 +700,22 @@ int DetourAttachFunctions(bool runningRemote)
 	{
 		DETOURED_FUNCTIONS_KERNEL32
 
-		// Old wine versions do not have this implemented. Call to another function that shows a message
 		if (IsRunningWine())
+		{
+			// Old wine versions do not have this implemented. Call to another function that shows a message
 			if (!GetProcAddress(moduleHandle, "InitializeSynchronizationBarrier"))
 				PatchIAT(GetModuleHandleW(nullptr), "KERNEL32.dll", "InitializeSynchronizationBarrier", (void*)&Detoured_InitializeSynchronizationBarrier);
+
+			// DeviceIoControl and GetProcAddress exists as a separately implementation on wine
+
+			void* ptr1 = GetProcAddress(moduleHandle, "DeviceIoControl");
+			DetourAttachFunction(&ptr1, Detoured_DeviceIoControl, "DeviceIoControl");
+
+			#if DETOURED_INCLUDE_DEBUG
+			void* ptr2 = GetProcAddress(moduleHandle, "GetProcAddress");
+			DetourAttachFunction(&ptr2, Detoured_GetProcAddress, "GetProcAddress");
+			#endif
+		}
 	}
 
 	if (HMODULE moduleHandle = GetModuleHandleW(L"ntdll.dll"))
diff --git a/Engine/Source/Programs/UnrealBuildAccelerator/Detours/Private/Windows/UbaDetoursFunctionsWin.h b/Engine/Source/Programs/UnrealBuildAccelerator/Detours/Private/Windows/UbaDetoursFunctionsWin.h
--- a/Engine/Source/Programs/UnrealBuildAccelerator/Detours/Private/Windows/UbaDetoursFunctionsWin.h
+++ b/Engine/Source/Programs/UnrealBuildAccelerator/Detours/Private/Windows/UbaDetoursFunctionsWin.h
@@ -243,6 +243,7 @@
 	DETOURED_FUNCTION(ExitProcess) \
 	DETOURED_FUNCTION(GetCommandLineA) \
 	DETOURED_FUNCTION(CommandLineToArgvW) \
+	DETOURED_FUNCTION(GetProcAddress) \
 	DETOURED_FUNCTION(FreeLibrary) \
 	DETOURED_FUNCTION(RegOpenKeyW) \
 	DETOURED_FUNCTION(RegOpenKeyExW) \
@@ -275,6 +276,7 @@
 	DETOURED_FUNCTION(GetLocaleInfoEx) \
 	DETOURED_FUNCTION(GetUserDefaultLocaleName) \
 	DETOURED_FUNCTION(GetDiskFreeSpaceExA) \
+	DETOURED_FUNCTION_UNKNOWN(GetFileInformationByName) \
 	DETOURED_FUNCTION(GetLongPathNameA) \
 	DETOURED_FUNCTION(GetVolumePathNameA) \
 	DETOURED_FUNCTION(GetFileSize) \
@@ -364,6 +366,7 @@
 	DETOURED_FUNCTION(GetShortPathNameA) \
 
 #define DETOURED_FUNCTIONS_NTDLL_DEBUG \
+	DETOURED_FUNCTION(RtlGetVersion) \
 	DETOURED_FUNCTION(RtlAllocateHeap) \
 	DETOURED_FUNCTION(RtlReAllocateHeap) \
 	DETOURED_FUNCTION(RtlValidateHeap) \
@@ -479,6 +482,7 @@ extern "C" {
 	NTSTATUS NTAPI LdrAddRefDll(ULONG Flags, PVOID BaseAddress);
 	PVOID WINAPI ResolveDelayLoadedAPI(PVOID ParentModuleBase, PCIMAGE_DELAYLOAD_DESCRIPTOR DelayloadDescriptor, void* FailureDllHook, void* FailureSystemHook, PIMAGE_THUNK_DATA ThunkAddress, ULONG Flags);
 	void WINAPI RtlExitUserThread(ULONG);
+	NTSTATUS NTAPI RtlGetVersion(PRTL_OSVERSIONINFOW lpVersionInformation);
 	BOOLEAN NTAPI RtlFreeHeap(PVOID HeapHandle, ULONG Flags, PVOID HeapBase);
 	BOOLEAN NTAPI RtlValidateHeap(HANDLE HeapPtr, ULONG Flags, PVOID Block);
 	NTSTATUS NTAPI RtlDosPathNameToNtPathName_U_WithStatus(PCWSTR dos_path, PUNICODE_STRING ntpath, PWSTR* file_part, VOID* reserved);
@@ -497,8 +501,12 @@ extern "C" {
 	HRESULT SHGetKnownFolderPath(REFKNOWNFOLDERID rfid, DWORD dwFlags, HANDLE hToken, PWSTR* ppszPath);
 	HRESULT SHGetKnownFolderIDList(REFKNOWNFOLDERID rfid, DWORD flags, HANDLE token, PIDLIST_ABSOLUTE *pidl);
 	HRESULT SHGetFolderPathW(HWND hwnd, int csidl, HANDLE hToken, DWORD dwFlags, LPWSTR pszPath);
+	enum FILE_INFO_BY_NAME_CLASS{};
+	BOOL WINAPI GetFileInformationByName(PCWSTR FileName, FILE_INFO_BY_NAME_CLASS FileInformationClass, PVOID FileInfoBuffer, ULONG FileInfoBufferSize);
+
 }
 
+#define DETOURED_FUNCTION_UNKNOWN(Func) DETOURED_FUNCTION(Func)
 
 #define DETOURED_FUNCTION(Func) extern decltype(Func)* True_##Func;
 DETOURED_FUNCTIONS
