[UBA]
* Better logging formatting

diff --git a/Engine/Source/Programs/UnrealBuildAccelerator/Detours/Private/Windows/UbaDetoursFunctionsKernelBase.inl b/Engine/Source/Programs/UnrealBuildAccelerator/Detours/Private/Windows/UbaDetoursFunctionsKernelBase.inl
--- a/Engine/Source/Programs/UnrealBuildAccelerator/Detours/Private/Windows/UbaDetoursFunctionsKernelBase.inl
+++ b/Engine/Source/Programs/UnrealBuildAccelerator/Detours/Private/Windows/UbaDetoursFunctionsKernelBase.inl
@@ -394,8 +394,9 @@ BOOL Detoured_UnlockFile(HANDLE hFile, DWORD dwFileOffsetLow, DWORD dwFileOffset
 		trueHandle = dh.trueHandle;
 		UBA_ASSERT(trueHandle != INVALID_HANDLE_VALUE);
 	}
-	DEBUG_LOG_TRUE(L"UnlockFile", L"%llu (%ls)", uintptr_t(hFile), HandleToName(hFile));
-	return True_UnlockFile(trueHandle, dwFileOffsetLow, dwFileOffsetHigh, nNumberOfBytesToUnlockLow, nNumberOfBytesToUnlockHigh);
+	BOOL res = True_UnlockFile(trueHandle, dwFileOffsetLow, dwFileOffsetHigh, nNumberOfBytesToUnlockLow, nNumberOfBytesToUnlockHigh);
+	DEBUG_LOG_TRUE(L"UnlockFile", L"%llu (%ls) -> %s", uintptr_t(hFile), HandleToName(hFile), ToString(res));
+	return res;
 }
 
 BOOL Detoured_UnlockFileEx(HANDLE hFile, DWORD dwReserved, DWORD nNumberOfBytesToUnlockLow, DWORD nNumberOfBytesToUnlockHigh, LPOVERLAPPED lpOverlapped)
@@ -408,8 +409,9 @@ BOOL Detoured_UnlockFileEx(HANDLE hFile, DWORD dwReserved, DWORD nNumberOfBytesT
 		trueHandle = dh.trueHandle;
 		UBA_ASSERT(trueHandle != INVALID_HANDLE_VALUE);
 	}
-	DEBUG_LOG_TRUE(L"UnlockFile", L"%llu (%ls)", uintptr_t(hFile), HandleToName(hFile));
-	return True_UnlockFileEx(trueHandle, dwReserved, nNumberOfBytesToUnlockLow, nNumberOfBytesToUnlockHigh, lpOverlapped);
+	BOOL res = True_UnlockFileEx(trueHandle, dwReserved, nNumberOfBytesToUnlockLow, nNumberOfBytesToUnlockHigh, lpOverlapped);
+	DEBUG_LOG_TRUE(L"UnlockFileEx", L"%llu (%ls) -> %s", uintptr_t(hFile), HandleToName(hFile), ToString(res));
+	return res;
 }
 
 BOOL Detoured_WriteConsoleA(HANDLE hConsoleOutput, const VOID* lpBuffer, DWORD nNumberOfCharsToWrite, LPDWORD lpNumberOfCharsWritten, LPVOID lpReserved)
@@ -4479,7 +4481,13 @@ LPWSTR* Detoured_CommandLineToArgvW(LPCWSTR lpCmdLine, int *pNumArgs)
 void* Detoured_GetProcAddress(HMODULE hModule, LPCSTR lpProcName)
 {
 	void* res = True_GetProcAddress(hModule, lpProcName);
-	DEBUG_LOG_TRUE(L"GetProcAddress", L"%S -> 0x%p", lpProcName, res);
+
+	StringBuffer<> func;
+	if ((ULONG_PTR)lpProcName >> 16)
+		func.Append(lpProcName);
+	else
+		func.Append(TCV("Ordinal:")).AppendValue(LOWORD(lpProcName));
+	DEBUG_LOG_TRUE(L"GetProcAddress", L"%s -> 0x%p", func.data, res);
 	return res;
 }
 
diff --git a/Engine/Source/Programs/UnrealBuildAccelerator/Detours/Private/Windows/UbaDetoursFunctionsNtDll.inl b/Engine/Source/Programs/UnrealBuildAccelerator/Detours/Private/Windows/UbaDetoursFunctionsNtDll.inl
--- a/Engine/Source/Programs/UnrealBuildAccelerator/Detours/Private/Windows/UbaDetoursFunctionsNtDll.inl
+++ b/Engine/Source/Programs/UnrealBuildAccelerator/Detours/Private/Windows/UbaDetoursFunctionsNtDll.inl
@@ -1150,7 +1150,7 @@ NTSTATUS NTAPI Shared_NtCreateFile(bool IsCreateFunc, PHANDLE hFileHandle, ACCES
 	keepInMemory = keepInMemory || ((isWrite || isDeleteOnClose) && g_rules->IsOutputFile(fileName, g_systemTemp) && g_allowKeepFilesInMemory) || (isWrite && isThrowAway);
 
 #if UBA_DEBUG_LOG_ENABLED
-	const wchar_t* isWriteStr = isWrite ? L" WRITE" : L""; (void)isWriteStr;
+	const wchar_t* isWriteStr = isWrite ? L"WRITE " : L""; (void)isWriteStr;
 #endif
 
 	bool isSystemFile = fileName.StartsWith(g_systemRoot.data);
@@ -1170,7 +1170,7 @@ NTSTATUS NTAPI Shared_NtCreateFile(bool IsCreateFunc, PHANDLE hFileHandle, ACCES
 		NTSTATUS res = Local_NtCreateFile(IsCreateFunc, hFileHandle, DesiredAccess, ObjectAttributes, IoStatusBlock, AllocationSize, FileAttributes, ShareAccess, CreateDisposition, CreateOptions, EaBuffer, EaLength);
 		if (NT_ERROR(res))
 			*hFileHandle = INVALID_HANDLE_VALUE;
-		DEBUG_LOG_TRUE(funcName, L"(NODETOUR)%ls %llu (%.*ls) -> %ls", isWriteStr, uintptr_t(*hFileHandle), ObjectAttributes->ObjectName->Length / 2, ObjectAttributes->ObjectName->Buffer, ToString(res).data);
+		DEBUG_LOG_TRUE(funcName, L"(NODETOUR) %s%llu (%.*ls) -> %ls", isWriteStr, uintptr_t(*hFileHandle), ObjectAttributes->ObjectName->Length / 2, ObjectAttributes->ObjectName->Buffer, ToString(res).data);
 		if (NT_ERROR(res))
 			return res;
 
@@ -1523,7 +1523,7 @@ NTSTATUS NTAPI Shared_NtCreateFile(bool IsCreateFunc, PHANDLE hFileHandle, ACCES
 			*hFileHandle = CreateFileHandle();
 			TrackFileInput();
 
-			DEBUG_LOG_DETOURED(funcName, L"(SHAREDMEMFILE)%ls %llu (%ls) (%ls) -> Success", isWriteStr, uintptr_t(*hFileHandle), lpFileName, (fileName.data != lpFileName ? fileName.data : L""));
+			DEBUG_LOG_DETOURED(funcName, L"(SHAREDMEMFILE) %ls%llu (%ls) (%ls) -> Success", isWriteStr, uintptr_t(*hFileHandle), lpFileName, (fileName.data != lpFileName ? fileName.data : L""));
 			return STATUS_SUCCESS;
 		}
 
@@ -1548,7 +1548,7 @@ NTSTATUS NTAPI Shared_NtCreateFile(bool IsCreateFunc, PHANDLE hFileHandle, ACCES
 		*hFileHandle = CreateFileHandle();
 		TrackFileInput();
 
-		DEBUG_LOG_DETOURED(funcName, L"(SHAREDMEMFILE)%ls %llu (%ls) (%ls) -> Success", isWriteStr, uintptr_t(*hFileHandle), lpFileName, (fileName.data != lpFileName ? fileName.data : L""));
+		DEBUG_LOG_DETOURED(funcName, L"(SHAREDMEMFILE) %ls%llu (%ls) (%ls) -> Success", isWriteStr, uintptr_t(*hFileHandle), lpFileName, (fileName.data != lpFileName ? fileName.data : L""));
 		return STATUS_SUCCESS;
 	}
 
@@ -1584,7 +1584,7 @@ NTSTATUS NTAPI Shared_NtCreateFile(bool IsCreateFunc, PHANDLE hFileHandle, ACCES
 		lookupLock.Leave();
 		*hFileHandle = CreateFileHandle();
 		TrackFileInput();
-		DEBUG_LOG_DETOURED(funcName, L"(WRITTENFILE)%ls %llu (%ls) (%ls) -> Success", isWriteStr, uintptr_t(*hFileHandle), lpFileName, (fileName.data != lpFileName ? fileName.data : L""));
+		DEBUG_LOG_DETOURED(funcName, L"(WRITTENFILE) %ls%llu (%ls) (%ls) -> Success", isWriteStr, uintptr_t(*hFileHandle), lpFileName, (fileName.data != lpFileName ? fileName.data : L""));
 		return STATUS_SUCCESS;
 	}
 	
@@ -1694,7 +1694,7 @@ NTSTATUS NTAPI Shared_NtCreateFile(bool IsCreateFunc, PHANDLE hFileHandle, ACCES
 
 		TrackFileInput();
 
-		DEBUG_LOG_DETOURED(funcName, L"(%s)%ls %llu (%ls) (%ls) -> Success", memoryType, isWriteStr, uintptr_t(*hFileHandle), lpFileName, (fileName.data != lpFileName ? fileName.data : L""));
+		DEBUG_LOG_DETOURED(funcName, L"(%s) %ls%llu (%ls) (%ls) -> Success", memoryType, isWriteStr, uintptr_t(*hFileHandle), lpFileName, (fileName.data != lpFileName ? fileName.data : L""));
 		return STATUS_SUCCESS;
 	}
 
