[UBA]
* Fixed so NtQueryDirectory RestartScan checks if directory table is up-to-date and if not refresh it

diff --git a/Engine/Source/Programs/UnrealBuildAccelerator/Core/Public/UbaDirectoryTable.h b/Engine/Source/Programs/UnrealBuildAccelerator/Core/Public/UbaDirectoryTable.h
--- a/Engine/Source/Programs/UnrealBuildAccelerator/Core/Public/UbaDirectoryTable.h
+++ b/Engine/Source/Programs/UnrealBuildAccelerator/Core/Public/UbaDirectoryTable.h
@@ -112,6 +112,11 @@ namespace uba
 			PopulateDirectoryNoLock(hasher, dir);
 		}
 
+		bool IsDirectoryUpToDate(Directory& dir)
+		{
+			return dir.parseOffset == dir.tableOffset;
+		}
+
 		void PopulateDirectoryNoLock(const StringKeyHasher& hasher, Directory& dir)
 		{
 			if (dir.parseOffset == dir.tableOffset)
diff --git a/Engine/Source/Programs/UnrealBuildAccelerator/Detours/Private/Windows/UbaDetoursFunctionsNtDll.inl b/Engine/Source/Programs/UnrealBuildAccelerator/Detours/Private/Windows/UbaDetoursFunctionsNtDll.inl
--- a/Engine/Source/Programs/UnrealBuildAccelerator/Detours/Private/Windows/UbaDetoursFunctionsNtDll.inl
+++ b/Engine/Source/Programs/UnrealBuildAccelerator/Detours/Private/Windows/UbaDetoursFunctionsNtDll.inl
@@ -426,6 +426,8 @@ NTSTATUS NTAPI Detoured_NtQueryDirectoryFile(HANDLE FileHandle, HANDLE Event, PI
 
 	if (isListDirectoryHandle(FileHandle))
 	{
+		DEBUG_LOG_DETOURED(L"NtQueryDirectoryFile", L"%llu %s", u64(FileHandle), HandleToName(FileHandle));
+
 		IoStatusBlock->Information = 0;
 
 		auto& listHandle = asListDirectoryHandle(FileHandle);
@@ -434,8 +436,24 @@ NTSTATUS NTAPI Detoured_NtQueryDirectoryFile(HANDLE FileHandle, HANDLE Event, PI
 		UBA_ASSERT(Event == 0 && ApcRoutine == nullptr && ApcContext == nullptr);
 
 		if (RestartScan)
+		{
+			// Check so directory has not changed since this handle was opened.. RestartScan means that it needs to be up-to-date
+			auto& dir = listHandle.dir;
+			if (!g_directoryTable.IsDirectoryUpToDate(dir))
+			{
+				StringBuffer<> fileNameLower(listHandle.originalName);
+				fileNameLower.MakeLower();
+				DirHash hash(fileNameLower);
+				SCOPED_READ_LOCK(dir.lock, lock);
+				g_directoryTable.PopulateDirectoryNoLock(hash.open, dir);
+				listHandle.fileTableOffsets.resize(dir.files.size());
+				u32 it = 0;
+				for (auto& pair : dir.files)
+					listHandle.fileTableOffsets[it++] = pair.second;
+				lock.Leave();
+			}
 			listHandle.it = 0;
-
+		}
 		u8* prevInformation = nullptr;
 		u8* it = (u8*)FileInformation;
 		u8* bufferEnd = it + Length;
@@ -514,7 +532,8 @@ NTSTATUS NTAPI Detoured_NtQueryDirectoryFile(HANDLE FileHandle, HANDLE Event, PI
 			prevInformation = it;
 			it = (u8*)fileNamePos + info.FileNameLength + 2;
 
-			DEBUG_LOG_DETOURED(L"NtQueryDirectoryFile", L"%llu %.*s", u64(FileHandle), int(fileNameBytes/sizeof(tchar)), fileNamePos);
+			// Very spammy
+			//DEBUG_LOG_DETOURED(L"NtQueryDirectoryFile", L"%llu %.*s", u64(FileHandle), int(fileNameBytes/sizeof(tchar)), fileNamePos);
 
 			res = STATUS_SUCCESS;
 
@@ -584,6 +603,8 @@ NTSTATUS NTAPI Detoured_NtQueryDirectoryFile(HANDLE FileHandle, HANDLE Event, PI
 		//DEBUG_LOG_TRUE(L"NtQueryDirectoryFile", L"(%u) %llu (%ls) -> %ls", FileInformationClass, uintptr_t(FileHandle), HandleToName(FileHandle), ToString(res));
 	}
 #endif
+
+	DEBUG_LOG_TRUE(L"NtQueryDirectoryFile", L"%llu -> %s", u64(FileHandle), ToString(res).data);
 	return res;
 }
 
