[UBA]
* Added response of GetEnvironmentStringsW to logging

diff --git a/Engine/Source/Programs/UnrealBuildAccelerator/Detours/Private/Windows/UbaDetoursFunctionsKernelBase.inl b/Engine/Source/Programs/UnrealBuildAccelerator/Detours/Private/Windows/UbaDetoursFunctionsKernelBase.inl
--- a/Engine/Source/Programs/UnrealBuildAccelerator/Detours/Private/Windows/UbaDetoursFunctionsKernelBase.inl
+++ b/Engine/Source/Programs/UnrealBuildAccelerator/Detours/Private/Windows/UbaDetoursFunctionsKernelBase.inl
@@ -4084,8 +4084,8 @@ DWORD Detoured_GetEnvironmentVariableA(LPCSTR lpName, LPSTR lpBuffer, DWORD nSiz
 LPWCH Detoured_GetEnvironmentStringsW()
 {
 	DETOURED_CALL(GetEnvironmentStringsW);
-	DEBUG_LOG_TRUE(L"GetEnvironmentStringsW", L"");
 	auto res = True_GetEnvironmentStringsW();
+	DEBUG_LOG_TRUE(L"GetEnvironmentStringsW", L"-> %s", res);
 
 	#if 0 // Enable to print out environment variables in the log
 	auto it = res;
