[UBA]
* Fixed bug in FixPath logic when string is not terminated but starts with \??\

diff --git a/Engine/Source/Programs/UnrealBuildAccelerator/Core/Public/UbaPathUtils.h b/Engine/Source/Programs/UnrealBuildAccelerator/Core/Public/UbaPathUtils.h
--- a/Engine/Source/Programs/UnrealBuildAccelerator/Core/Public/UbaPathUtils.h
+++ b/Engine/Source/Programs/UnrealBuildAccelerator/Core/Public/UbaPathUtils.h
@@ -74,6 +74,7 @@ namespace uba
 
 		if (read[0] == '\\')
 		{
+			const tchar* prevRead = read;
 			if (read[1] == '\\')
 			{
 				if (read[2] == '?' && read[3] == '\\')
@@ -100,6 +101,7 @@ namespace uba
 					read += 3;
 				}
 			}
+			fileNameLen -= u32(read - prevRead);
 		}
 
 		bool isAbsolute = false;
diff --git a/Engine/Source/Programs/UnrealBuildAccelerator/Test/Private/UbaTestBasics.cpp b/Engine/Source/Programs/UnrealBuildAccelerator/Test/Private/UbaTestBasics.cpp
--- a/Engine/Source/Programs/UnrealBuildAccelerator/Test/Private/UbaTestBasics.cpp
+++ b/Engine/Source/Programs/UnrealBuildAccelerator/Test/Private/UbaTestBasics.cpp
@@ -184,6 +184,8 @@ namespace uba
 #if PLATFORM_WINDOWS
 		if (!FixPath2(TC("\"e:\\temp\""), ~0u, workingDir, TStrlen(workingDir), buffer, sizeof_array(buffer), &lengthResult))
 			return logger.Error(TC("FixPath2 (1) failed"));
+		if (!FixPath2(TC("\\??\\c:\\windows:::"), 14, workingDir, TStrlen(workingDir), buffer, sizeof_array(buffer), &lengthResult) || lengthResult != 10)
+			return logger.Error(TC("FixPath2 (2) failed"));
 #else
 
 		if (!TestPath(TC("/..")))
