[UBA]
* Fixed shipping config compile error

diff --git a/Engine/Source/Programs/UnrealBuildAccelerator/Detours/Private/UbaDetoursShared.cpp b/Engine/Source/Programs/UnrealBuildAccelerator/Detours/Private/UbaDetoursShared.cpp
--- a/Engine/Source/Programs/UnrealBuildAccelerator/Detours/Private/UbaDetoursShared.cpp
+++ b/Engine/Source/Programs/UnrealBuildAccelerator/Detours/Private/UbaDetoursShared.cpp
@@ -426,7 +426,9 @@ namespace uba
 
 		if (!TakeLockForRpc())
 		{
+			#if UBA_DEBUG_LOG_ENABLED
 			FlushDebug();
+			#endif
 			return;
 		}
 		RPC_MESSAGE_NO_LOCK(Log, log)
