[UBA]
* Fixed bug where an output file written directly to disk remotely would use the wrong size when getting file size in a parent process.
* More logging format improvements

diff --git a/Engine/Source/Programs/UnrealBuildAccelerator/Detours/Private/Windows/UbaDetoursFunctionsKernelBase.inl b/Engine/Source/Programs/UnrealBuildAccelerator/Detours/Private/Windows/UbaDetoursFunctionsKernelBase.inl
--- a/Engine/Source/Programs/UnrealBuildAccelerator/Detours/Private/Windows/UbaDetoursFunctionsKernelBase.inl
+++ b/Engine/Source/Programs/UnrealBuildAccelerator/Detours/Private/Windows/UbaDetoursFunctionsKernelBase.inl
@@ -4085,7 +4085,7 @@ LPWCH Detoured_GetEnvironmentStringsW()
 {
 	DETOURED_CALL(GetEnvironmentStringsW);
 	auto res = True_GetEnvironmentStringsW();
-	DEBUG_LOG_TRUE(L"GetEnvironmentStringsW", L"-> %s", res);
+	DEBUG_LOG_TRUE(L"GetEnvironmentStringsW", L"-> 0x%llu", u64(res));
 
 	#if 0 // Enable to print out environment variables in the log
 	auto it = res;
diff --git a/Engine/Source/Programs/UnrealBuildAccelerator/Detours/Private/Windows/UbaDetoursFunctionsNtDll.inl b/Engine/Source/Programs/UnrealBuildAccelerator/Detours/Private/Windows/UbaDetoursFunctionsNtDll.inl
--- a/Engine/Source/Programs/UnrealBuildAccelerator/Detours/Private/Windows/UbaDetoursFunctionsNtDll.inl
+++ b/Engine/Source/Programs/UnrealBuildAccelerator/Detours/Private/Windows/UbaDetoursFunctionsNtDll.inl
@@ -302,6 +302,8 @@ NTSTATUS Detoured_NtQueryInformationFile(HANDLE FileHandle, PIO_STATUS_BLOCK IoS
 	{
 		UBA_ASSERT(dh.type == HandleType_File);
 		FileInfo& fi = *dh.fileObject->fileInfo;
+		if (!fi.memoryFile && !fi.isFileMap && fi.created) // fi.created means that it was created in process hierarchy and other fields might be outdated when it comes to size
+			return runTrue();
 		if (setStandardInformation(*(FILE_STANDARD_INFORMATION*)FileInformation, fi, nullptr) != STATUS_SUCCESS)
 			return runTrue();
 		DEBUG_LOG_DETOURED(L"NtQueryInformationFile", L"(FileStandardInformation) %llu (%s) -> %ls", uintptr_t(FileHandle), HandleToName(FileHandle), ToString(STATUS_SUCCESS).data);
@@ -1924,7 +1926,7 @@ NTSTATUS NTAPI Detoured_NtWriteFile(HANDLE FileHandle, HANDLE Event, PIO_APC_ROU
 			(void)prefix;
 			TimerScope ts(g_kernelStats.writeFile);
 			NTSTATUS status = True_NtWriteFile(fileHandle, Event, ApcRoutine, ApcContext, IoStatusBlock, Buffer, Length, ByteOffset, Key);
-			DEBUG_LOG_TRUE(L"NtWriteFile", L"%s%llu (%ls) -> %s", prefix, uintptr_t(fileHandle), HandleToName(fileHandle), ToString(status).data);
+			DEBUG_LOG_TRUE(L"NtWriteFile", L"%s%llu (%llu/%llu) (%ls) -> %s", prefix, uintptr_t(FileHandle), u64(IoStatusBlock->Information), u64(Length), HandleToName(FileHandle), ToString(status).data);
 			g_kernelStats.writeFile.bytes += IoStatusBlock->Information;
 			return status;
 		};
@@ -1934,7 +1936,7 @@ NTSTATUS NTAPI Detoured_NtWriteFile(HANDLE FileHandle, HANDLE Event, PIO_APC_ROU
 	if (FileHandle == PseudoHandle)
 		return trueRun(FileHandle, TC("(PseudoHandle)"));
 
-	auto successRun = [&]()
+	auto successRun = [&](const tchar* prefix)
 		{
 			IoStatusBlock->Information = Length;
 			IoStatusBlock->Status = STATUS_SUCCESS;
@@ -1942,6 +1944,7 @@ NTSTATUS NTAPI Detoured_NtWriteFile(HANDLE FileHandle, HANDLE Event, PIO_APC_ROU
 				SetEvent(Event);
 			if (ApcRoutine != nullptr)
 				ApcRoutine(ApcContext, IoStatusBlock, 0);
+			DEBUG_LOG_DETOURED(L"NtWriteFile", L"%s%llu (%llu/%llu) (%ls) -> %s", prefix, uintptr_t(FileHandle), u64(IoStatusBlock->Information), u64(Length), HandleToName(FileHandle), ToString(STATUS_SUCCESS).data);
 			return STATUS_SUCCESS;
 		};
 
@@ -1949,7 +1952,7 @@ NTSTATUS NTAPI Detoured_NtWriteFile(HANDLE FileHandle, HANDLE Event, PIO_APC_ROU
 	{
 		//DEBUG_LOG_DETOURED(L"WriteStdFile2", L"%llu", uintptr_t(hFile));
 		WriteStdFile(Buffer, Length, FileHandle == g_stdHandle[0]);
-		return successRun();
+		return successRun(TC("(StdOut1)"));
 	}
 
 	if (!isDetouredHandle(FileHandle))
@@ -1964,7 +1967,7 @@ NTSTATUS NTAPI Detoured_NtWriteFile(HANDLE FileHandle, HANDLE Event, PIO_APC_ROU
 			//DEBUG_LOG_DETOURED(L"WriteStdFile1", L"%llu", uintptr_t(hFile));
 			WriteStdFile(Buffer, Length, dh.type == HandleType_StdErr);
 		}
-		return successRun();
+		return successRun(TC("(StdOut2)"));
 	}
 
 	auto& fo = *dh.fileObject;
@@ -1987,14 +1990,14 @@ NTSTATUS NTAPI Detoured_NtWriteFile(HANDLE FileHandle, HANDLE Event, PIO_APC_ROU
 	if (mf.isThrowAway)
 	{
 		mf.writtenSize = newSize;
-		return successRun();
+		return successRun(TC("(ThrowAway)"));
 	}
 	mf.EnsureCommitted(dh, newSize);
 	mf.writtenSize = newSize;
 	mf.isReported = false;
 	lock.Leave();
 	memcpy(mf.baseAddress + offset, Buffer, Length);
-	return successRun();
+	return successRun(TC("(Memory)"));
 }
 
 NTSTATUS NTAPI Detoured_NtFsControlFile(HANDLE FileHandle, HANDLE Event, PIO_APC_ROUTINE ApcRoutine, PVOID ApcContext, PIO_STATUS_BLOCK IoStatusBlock, ULONG FsControlCode, PVOID InputBuffer, ULONG InputBufferLength, PVOID OutputBuffer, ULONG OutputBufferLength)
