[UBA]
* Moved GetFileInformationByName to be non-debug.. since it should not let requests through when running remotely

diff --git a/Engine/Source/Programs/UnrealBuildAccelerator/Detours/Private/Windows/UbaDetoursFunctionsKernelBase.inl b/Engine/Source/Programs/UnrealBuildAccelerator/Detours/Private/Windows/UbaDetoursFunctionsKernelBase.inl
--- a/Engine/Source/Programs/UnrealBuildAccelerator/Detours/Private/Windows/UbaDetoursFunctionsKernelBase.inl
+++ b/Engine/Source/Programs/UnrealBuildAccelerator/Detours/Private/Windows/UbaDetoursFunctionsKernelBase.inl
@@ -2154,6 +2154,68 @@
 	return res;
 }
 
+BOOL Detoured_GetFileInformationByName(PCWSTR FileName, FILE_INFO_BY_NAME_CLASS FileInformationClass, PVOID FileInfoBuffer, ULONG FileInfoBufferSize)
+{
+	DETOURED_CALL(GetFileInformationByName);
+
+#if 0
+	enum FILE_INFO_BY_NAME_CLASS
+	{
+		FileStatByNameInfo,
+		FileStatLxByNameInfo,
+		FileCaseSensitiveByNameInfo,
+		FileStatBasicByNameInfo,
+		MaximumFileInfoByNameClass
+	};
+	
+	if (FileInformationClass == FileStatBasicByNameInfo)
+	{
+		struct FILE_STAT_BASIC_INFORMATION
+		{
+			LARGE_INTEGER FileId;
+			LARGE_INTEGER CreationTime;
+			LARGE_INTEGER LastAccessTime;
+			LARGE_INTEGER LastWriteTime;
+			LARGE_INTEGER ChangeTime;
+			LARGE_INTEGER AllocationSize;
+			LARGE_INTEGER EndOfFile;
+			ULONG FileAttributes;
+			ULONG ReparseTag;
+			ULONG NumberOfLinks;
+			ACCESS_MASK EffectiveAccess;
+		};
+		FILE_STAT_BASIC_INFORMATION
+	}
+#endif
+	if (g_runningRemote)
+	{
+		SetLastError(ERROR_CALL_NOT_IMPLEMENTED);
+		DEBUG_LOG_DETOURED(L"GetFileInformationByName", L"(%ls) (class %u) -> %s (ERROR_CALL_NOT_IMPLEMENTED)", FileName, FileInformationClass, ToString(false));
+		return false;
+	}
+
+	StringBuffer<> pathName;
+	FixPath(pathName, FileName);
+	DevirtualizePath(pathName);
+	StringKey sourceKey = ToStringKeyLower(pathName);
+
+#if UBA_DEBUG
+	SCOPED_WRITE_LOCK(g_mappedFileTable.m_lookupLock, lock);
+	auto it = g_mappedFileTable.m_lookup.find(sourceKey);
+	if (it != g_mappedFileTable.m_lookup.end())
+	{
+		FileInfo& fi = it->second;
+		UBA_ASSERT(!fi.memoryFile);
+	}
+	lock.Leave();
+#endif
+
+	BOOL res = True_GetFileInformationByName(pathName.data, FileInformationClass, FileInfoBuffer, FileInfoBufferSize);
+	DEBUG_LOG_TRUE(L"GetFileInformationByName", L"(%ls) (class %u) -> %s", pathName.data, FileInformationClass, ToString(res));
+	return res;
+}
+
+
 BOOL Detoured_SetFileInformationByHandle(HANDLE hFile, FILE_INFO_BY_HANDLE_CLASS FileInformationClass, LPVOID lpFileInformation, DWORD dwBufferSize)
 {
 	DETOURED_CALL(SetFileInformationByHandle);
@@ -3695,66 +3757,6 @@
 	return True_GetDiskFreeSpaceExA(lpDirectoryName, lpFreeBytesAvailableToCaller, lpTotalNumberOfBytes, lpTotalNumberOfFreeBytes);
 }
 
-BOOL Detoured_GetFileInformationByName(PCWSTR FileName, FILE_INFO_BY_NAME_CLASS FileInformationClass, PVOID FileInfoBuffer, ULONG FileInfoBufferSize)
-{
-	DETOURED_CALL(GetFileInformationByName);
-
-#if 0
-	enum FILE_INFO_BY_NAME_CLASS
-	{
-		FileStatByNameInfo,
-		FileStatLxByNameInfo,
-		FileCaseSensitiveByNameInfo,
-		FileStatBasicByNameInfo,
-		MaximumFileInfoByNameClass
-	};
-	
-	if (FileInformationClass == FileStatBasicByNameInfo)
-	{
-		struct FILE_STAT_BASIC_INFORMATION
-		{
-			LARGE_INTEGER FileId;
-			LARGE_INTEGER CreationTime;
-			LARGE_INTEGER LastAccessTime;
-			LARGE_INTEGER LastWriteTime;
-			LARGE_INTEGER ChangeTime;
-			LARGE_INTEGER AllocationSize;
-			LARGE_INTEGER EndOfFile;
-			ULONG FileAttributes;
-			ULONG ReparseTag;
-			ULONG NumberOfLinks;
-			ACCESS_MASK EffectiveAccess;
-		};
-		FILE_STAT_BASIC_INFORMATION
-	}
-#endif
-	if (g_runningRemote)
-	{
-		SetLastError(ERROR_CALL_NOT_IMPLEMENTED);
-		DEBUG_LOG_DETOURED(L"GetFileInformationByName", L"(%ls) (class %u) -> %s (ERROR_CALL_NOT_IMPLEMENTED)", FileName, FileInformationClass, ToString(false));
-		return false;
-	}
-
-	StringBuffer<> pathName;
-	FixPath(pathName, FileName);
-	DevirtualizePath(pathName);
-	StringKey sourceKey = ToStringKeyLower(pathName);
-
-	SCOPED_WRITE_LOCK(g_mappedFileTable.m_lookupLock, lock);
-	auto it = g_mappedFileTable.m_lookup.find(sourceKey);
-	if (it != g_mappedFileTable.m_lookup.end())
-	{
-		FileInfo& fi = it->second;
-		UBA_ASSERT(!fi.memoryFile);
-	}
-	lock.Leave();
-
-
-	BOOL res = True_GetFileInformationByName(pathName.data, FileInformationClass, FileInfoBuffer, FileInfoBufferSize);
-	DEBUG_LOG_TRUE(L"GetFileInformationByName", L"(%ls) (class %u) -> %s", pathName.data, FileInformationClass, ToString(res));
-	return res;
-}
-
 DWORD Detoured_GetLongPathNameA(LPCSTR lpszShortPath, LPSTR lpszLongPath, DWORD cchBuffer)
 {
 	DETOURED_CALL(GetLongPathNameA);
