[UBA]
* Added detouring for job objects to be able to handle termination of child processes in a proper way

diff --git a/Engine/Source/Programs/UnrealBuildAccelerator/Common/Private/UbaProcess.cpp b/Engine/Source/Programs/UnrealBuildAccelerator/Common/Private/UbaProcess.cpp
--- a/Engine/Source/Programs/UnrealBuildAccelerator/Common/Private/UbaProcess.cpp
+++ b/Engine/Source/Programs/UnrealBuildAccelerator/Common/Private/UbaProcess.cpp
@@ -409,9 +409,6 @@ namespace uba
 					err.Appendf(TC(" (GetExitCodeProcess returned 0x%x - %s Size: %llu, CasKey: %s)"), exitCode, m_realApplication.c_str(), imageInfo.fileSize, CasKeyString(imageInfo.key).str);
 			}
 
-				if (m_parentReportedExit)
-					return false;
-					
 			if (err.IsEmpty())
 			{
 				// Ok, this is a bit hacky putting this here.. but theory is that process could have exited and exit message is still in the shared memory
@@ -530,9 +527,6 @@ namespace uba
 				return true;
 			}
 
-			if (m_parentReportedExit)
-				return false;
-
 			StringBuffer<> err;
 			err.Appendf(TC("ERROR: Process %u (%s) not active but did not get exit message. Received %u messages. Signal code: %i. Exit value or signal: %i. Execution time: %s."), m_nativeProcessId, m_startInfo.GetDescription(), m_messageCount, signalInfo.si_code, signalInfo.si_status, TimeToText(GetTime() - m_startTime).str);
 			LogLine(false, err.data, LogEntryType_Error);
@@ -807,10 +801,13 @@ namespace uba
 			m_session.EmergencyShutdown();
 		}
 
-		if (IsCancelled())
-			m_exitCode = ProcessCancelExitCode;
-		else
-			m_exitCode = exitCode;
+		if (m_exitCode == ~0u)
+		{
+			if (IsCancelled())
+				m_exitCode = ProcessCancelExitCode;
+			else
+				m_exitCode = exitCode;
+		}
 	}
 
 	void ProcessImpl::ThreadExit()
@@ -1496,8 +1493,9 @@ namespace uba
 			auto& process = *(ProcessImpl*)child.m_process;
 			if (process.m_nativeProcessId != nativeProcessId)
 				continue;
-			m_session.m_logger.Info(TC("CHILD TERMINATED. Exit code %u"), exitCode);
-			process.m_parentReportedExit = true;
+			//m_session.m_logger.Info(TC("CHILD TERMINATED. Exit code %u"), exitCode);
+			process.Cancel();
+			process.m_exitCode = exitCode;
 			return true;
 		}
 		return true;
@@ -1560,15 +1558,6 @@ namespace uba
 		AtomicMax(m_processStats.detoursMemory, stats.detoursMemory.load());
 		AtomicMax(m_processStats.peakMemory, stats.peakMemory.load());
 
-		// Traverse children that has been tagged to be exited. This is a for-now solution until we found use cases where parent process
-		// kills job object before it exits itself
-		for (auto& child : m_childProcesses)
-		{
-			auto& process = *(ProcessImpl*)child.m_process;
-			if (process.m_parentReportedExit)
-				process.Cancel();
-		}
-
 		return false;
 	}
 
@@ -2179,7 +2168,7 @@ namespace uba
 
 			if (!DetourCopyPayloadToProcess((HANDLE)m_nativeProcessHandle, DetoursPayloadGuid, &payload, sizeof(payload)))
 			{
-				if (!m_cancelled && !m_parentReportedExit)
+				if (!m_cancelled)
 					logger.Error(TC("Process with id %u failed to copy payload to %sprocess %llu - %s (%s)"), m_id, m_parentProcess?TC("child "):TC(""), m_nativeProcessHandle, m_realApplication.c_str(), LastErrorToText().data);
 				return UBA_EXIT_CODE(9);
 			}
@@ -2206,7 +2195,7 @@ namespace uba
 
 		if (!m_startInfo.startSuspended && ResumeThread(m_nativeThreadHandle) == -1)
 		{
-			LogEntryType logType = (m_cancelled || m_parentReportedExit) ? LogEntryType_Info : LogEntryType_Error;
+			LogEntryType logType = m_cancelled ? LogEntryType_Info : LogEntryType_Error;
 			logger.Logf(logType, TC("Failed to resume thread %llu for processid %u - %s. (%s)"), m_nativeThreadHandle, m_id, m_realApplication.c_str(), LastErrorToText().data);
 			return UBA_EXIT_CODE(11);
 		}
@@ -2537,8 +2526,8 @@ namespace uba
 		if (cancel)
 		{
 			if (!TerminateProcess(handle, ProcessCancelExitCode))
-				if (!m_parentReportedExit || GetLastError() != ERROR_ACCESS_DENIED)
-					logger.Info(TC("TerminateProcess for process %llu (%s) failed when cancelling (%s)%s"), u64(handle), m_startInfo.GetDescription(), LastErrorToText().data, (m_parentReportedExit ? TC("(ParentReportedExit)") : TC("")));
+				if (GetLastError() != ERROR_ACCESS_DENIED)
+					logger.Info(TC("TerminateProcess for process %llu (%s) failed when cancelling (%s)"), u64(handle), m_startInfo.GetDescription(), LastErrorToText().data);
 			DWORD res = WaitForSingleObject(handle, 1000); // On windows this should return in milliseconds. On wine it might take a little while. We need to wait to make sure process is reaped
 			if (res == WAIT_FAILED)
 				logger.Info(TC("WaitForSingleObject for process %llu (%s) failed when cancelling (%s)"), u64(handle), m_startInfo.GetDescription(), LastErrorToText().data);
@@ -2561,7 +2550,7 @@ namespace uba
 						const tchar* gotMessage = m_gotExitMessage ? TC("Got") : TC("Did not get");
 						logger.Info(TC("WaitForSingleObject timed out after 120 seconds waiting for process %llu (%s) to exit (Exit code %u, %s ExitMessage and wrote %u files. Runtime: %s). Will terminate and wait again"), u64(handle), m_startInfo.GetDescription(), m_nativeProcessExitCode, gotMessage, u32(m_shared.writtenFiles.size()), TimeToText(GetTime() - m_startTime).str);
 						if (!TerminateProcess(handle, m_nativeProcessExitCode))
-							logger.Info(TC("TerminateProcess for process %llu (%s) failed when waiting timed out (%s)%s"), u64(handle), m_startInfo.GetDescription(), LastErrorToText().data, (m_parentReportedExit ? TC("(ParentReportedExit)") : TC("")));
+							logger.Info(TC("TerminateProcess for process %llu (%s) failed when waiting timed out (%s)"), u64(handle), m_startInfo.GetDescription(), LastErrorToText().data);
 						continue;
 					}
 					logger.Error(TC("WaitForSingleObject failed while waiting for process %llu (%s) to exit even after terminating it (%s)"), u64(handle), m_startInfo.GetDescription(), LastErrorToText().data);
diff --git a/Engine/Source/Programs/UnrealBuildAccelerator/Detours/Private/Windows/UbaDetoursFunctionsKernelBase.inl b/Engine/Source/Programs/UnrealBuildAccelerator/Detours/Private/Windows/UbaDetoursFunctionsKernelBase.inl
--- a/Engine/Source/Programs/UnrealBuildAccelerator/Detours/Private/Windows/UbaDetoursFunctionsKernelBase.inl
+++ b/Engine/Source/Programs/UnrealBuildAccelerator/Detours/Private/Windows/UbaDetoursFunctionsKernelBase.inl
@@ -3458,84 +3458,6 @@ BOOL Detoured_CreateProcessA(LPCSTR lpApplicationName, LPSTR lpCommandLine, LPSE
 	return Detoured_CreateProcessW(lpApplicationNameW, lpCommandLineW, lpProcessAttributes, lpThreadAttributes, bInheritHandles, dwCreationFlags, lpEnvironment, lpCurrentDirectoryW, &lpStartupInfoW, lpProcessInformation);
 }
 
-BOOL Detoured_AssignProcessToJobObject(HANDLE hJob, HANDLE hProcess)
-{
-	HANDLE trueHandle = hProcess;
-	if (isDetouredHandle(hProcess))
-		trueHandle = asDetouredHandle(hProcess).trueHandle;
-	BOOL res = True_AssignProcessToJobObject(hJob, trueHandle);
-	DEBUG_LOG_DETOURED(L"AssignProcessToJobObject", L"Process %llu assigned to %llu -> %ls", uintptr_t(trueHandle), u64(hJob), ToString(res));
-
-	if (res)
-	{
-		JOBOBJECT_EXTENDED_LIMIT_INFORMATION info = {};
-		if (QueryInformationJobObject(hJob, JobObjectExtendedLimitInformation, &info, sizeof(info), nullptr))
-		{
-			if (info.BasicLimitInformation.LimitFlags & JOB_OBJECT_LIMIT_KILL_ON_JOB_CLOSE)
-			{
-				DWORD pid = GetProcessId(hProcess);
-				DEBUG_LOG(L"PROCESS %u WILL BE TERMINATED from job %llu", pid, u64(hJob));
-				RPC_MESSAGE(ExitChildProcess, detach) // Let it be part of detach timer for now
-				writer.WriteU32(pid);
-				writer.WriteU32(0xC000010A);
-				writer.Flush();
-			}
-		}
-		else
-		{
-			UBA_ASSERT(false);
-		}
-	}
-
-	return res;
-}
-
-BOOL Detoured_SetInformationJobObject(HANDLE hJob, JOBOBJECTINFOCLASS JobObjectInformationClass, LPVOID lpJobObjectInformation, DWORD cbJobObjectInformationLength)
-{
-	DEBUG_LOG_TRUE(L"SetInformationJobObject", L"%u", JobObjectInformationClass);
-	bool res = True_SetInformationJobObject(hJob, JobObjectInformationClass, lpJobObjectInformation, cbJobObjectInformationLength);
-
-	if (!res)
-		return res;
-
-	bool killOnJobClose = false;
-	if (JobObjectInformationClass == JobObjectBasicLimitInformation)
-	{
-		auto info = *(JOBOBJECT_BASIC_LIMIT_INFORMATION*)lpJobObjectInformation;
-		killOnJobClose = (info.LimitFlags & JOB_OBJECT_LIMIT_KILL_ON_JOB_CLOSE) != 0;
-	}
-	if (JobObjectInformationClass == JobObjectExtendedLimitInformation)
-	{
-		auto info = *(JOBOBJECT_EXTENDED_LIMIT_INFORMATION*)lpJobObjectInformation;
-		killOnJobClose = (info.BasicLimitInformation.LimitFlags & JOB_OBJECT_LIMIT_KILL_ON_JOB_CLOSE) != 0;
-	}
-
-	if (killOnJobClose)
-	{
-		DEBUG_LOG(L"PROCESSES WILL BE TERMINATED ON CLOSE OF %llu", u64(hJob));
-
-		u8 buffer[sizeof(JOBOBJECT_BASIC_PROCESS_ID_LIST) + 1024 * sizeof(ULONG_PTR)];
-		auto idList = (JOBOBJECT_BASIC_PROCESS_ID_LIST*)buffer;
-		if (QueryInformationJobObject(hJob, JobObjectBasicProcessIdList, buffer, sizeof(buffer), NULL))
-		{
-			for (DWORD i = 0; i < idList->NumberOfProcessIdsInList; ++i)
-			{
-				DWORD pid = (DWORD)idList->ProcessIdList[i];
-				DEBUG_LOG(L"PROCESS %u WILL BE TERMINATED from job %llu", pid, u64(hJob));
-				// TODO: This should be tracked and sent when hJob handle closes.. for now we just tell all processes that they might go away at any time
-				RPC_MESSAGE(ExitChildProcess, detach) // Let it be part of detach timer for now
-				writer.WriteU32(pid);
-				writer.Flush();
-			}
-		}
-		else
-		{
-			UBA_ASSERT(false);
-		}
-	}
-	return res;
-}
-
 BOOL Detoured_RegisterWaitForSingleObject(PHANDLE phNewWaitObject, HANDLE hObject, WAITORTIMERCALLBACK Callback, PVOID Context, ULONG dwMilliseconds, ULONG dwFlags)
 {
 	HANDLE trueHandle = hObject;
diff --git a/Engine/Source/Programs/UnrealBuildAccelerator/Detours/Private/Windows/UbaDetoursFunctionsNtDll.inl b/Engine/Source/Programs/UnrealBuildAccelerator/Detours/Private/Windows/UbaDetoursFunctionsNtDll.inl
--- a/Engine/Source/Programs/UnrealBuildAccelerator/Detours/Private/Windows/UbaDetoursFunctionsNtDll.inl
+++ b/Engine/Source/Programs/UnrealBuildAccelerator/Detours/Private/Windows/UbaDetoursFunctionsNtDll.inl
@@ -813,6 +813,39 @@ NTSTATUS Detoured_NtSetInformationObject(HANDLE ObjectHandle, OBJECT_INFORMATION
 	return res;
 }
 
+struct JobInfo { Vector<HANDLE> processes; bool shouldKill = false; };
+auto& g_jobsLock = *new ReaderWriterLock();
+auto& g_jobs = *new UnorderedMap<FileObject*, JobInfo>();
+
+void TerminateProcessesInJob(JobInfo&& info, u32 exitCode)
+{
+	for (auto& process : info.processes)
+	{
+		DWORD pid = GetProcessId(process);
+		DEBUG_LOG(L"PROCESS %u WILL BE TERMINATED by job", pid);
+		RPC_MESSAGE(ExitChildProcess, detach) // Let it be part of detach timer for now
+			writer.WriteU32(pid);
+		writer.WriteU32(exitCode);
+		writer.Flush();
+	}
+}
+
+void CloseProcess(u32 exitStatus)
+{
+	{
+		SCOPED_WRITE_LOCK(g_jobsLock, lock);
+		for (auto& kv : g_jobs)
+			if (kv.second.shouldKill)
+				TerminateProcessesInJob(std::move(kv.second), exitStatus);
+	}
+	CloseCaches();
+	SendExitMessage(exitStatus, GetTime());
+	PostDeinit();
+	#if UBA_DEBUG_LOG_ENABLED
+	FlushDebugLog();
+	#endif
+}
+
 NTSTATUS NTAPI Detoured_NtTerminateProcess(HANDLE ProcessHandle, NTSTATUS ExitStatus)
 {
 	// NtTerminateProcess can be reached in three ways.. one is through ExitProcess, TerminateProcess and NtTerminateProcess
@@ -820,12 +853,8 @@ NTSTATUS NTAPI Detoured_NtTerminateProcess(HANDLE ProcessHandle, NTSTATUS ExitSt
 	DETOURED_CALL(TerminateProcess);
 	DEBUG_LOG_DETOURED(L"NtTerminateProcess", L"%llu (%ls) ExitCode: %u", u64(ProcessHandle), HandleToName(ProcessHandle), ExitStatus);
 
-	if (ProcessHandle == NULL)
-		return True_NtTerminateProcess(ProcessHandle, ExitStatus);
-
 	HANDLE trueProcess = ProcessHandle;
-
-	bool isSelf = trueProcess == GetCurrentProcess();
+	bool isSelf = trueProcess == NtCurrentProcess() || trueProcess == NULL;
 	if (!isSelf)
 		if (isDetouredHandle(trueProcess))
 			trueProcess = asDetouredHandle(trueProcess).trueHandle;
@@ -838,19 +867,9 @@ NTSTATUS NTAPI Detoured_NtTerminateProcess(HANDLE ProcessHandle, NTSTATUS ExitSt
 	{
 		// Not called through ExitProcess
 		if (!g_exitMessageSent)
-		{
-			CloseCaches();
-			SendExitMessage((DWORD)ExitStatus, GetTime());
-			PostDeinit();
-			#if UBA_DEBUG_LOG_ENABLED
-			FlushDebugLog();
-			#endif
-		}
-		return True_NtTerminateProcess(trueProcess, ExitStatus);
+			CloseProcess((DWORD)ExitStatus);
 	}
-
-	pid = GetProcessId(trueProcess);
-	if (pid != GetCurrentProcessId())
+	else
 	{
 		DEBUG_LOG(L"PROCESS %u WILL BE TERMINATED", pid);
 		RPC_MESSAGE(ExitChildProcess, detach) // Let it be part of detach timer for now
@@ -859,9 +878,88 @@ NTSTATUS NTAPI Detoured_NtTerminateProcess(HANDLE ProcessHandle, NTSTATUS ExitSt
 		writer.Flush();
 	}
 
-	return True_NtTerminateProcess(ProcessHandle, ExitStatus);
+	return True_NtTerminateProcess(trueProcess, ExitStatus);
+}
+
+NTSTATUS Detoured_NtCreateJobObject(HANDLE* handle, ACCESS_MASK access, const OBJECT_ATTRIBUTES* attr)
+{
+	NTSTATUS res = True_NtCreateJobObject(handle, access, attr);
+	auto g = MakeGuard([&]() { DEBUG_LOG_TRUE(L"NtCreateJobObject", L"%llu -> %s", u64(*handle), ToString(res).data); });
+	if (res != STATUS_SUCCESS)
+		return res;
+
+	auto fo = new FileObject();
+	auto dh = new DetouredHandle(HandleType_Job, *handle);
+	dh->fileObject = fo;
+	*handle = makeDetouredHandle(dh);
+	return STATUS_SUCCESS;
+}
+
+NTSTATUS Detoured_NtSetInformationJobObject(HANDLE handle, JOBOBJECTINFOCLASS infoClass, void* info, ULONG len)
+{
+	UBA_ASSERT(isDetouredHandle(handle));
+	DetouredHandle& dh = asDetouredHandle(handle);
+	HANDLE trueHandle = dh.trueHandle;
+
+	NTSTATUS res = True_NtSetInformationJobObject(trueHandle, infoClass, info, len);
+	
+	const tchar* prefix = TC("");(void)prefix;
+	auto g = MakeGuard([&]() { DEBUG_LOG_TRUE(L"NtSetInformationJobObject", L"%s(%u) -> %s", prefix, infoClass, ToString(res).data); });
+	
+	if (res != STATUS_SUCCESS)
+		return res;
+
+	bool killOnJobClose = false;
+	if (infoClass == JobObjectBasicLimitInformation)
+	{
+		auto info2 = *(JOBOBJECT_BASIC_LIMIT_INFORMATION*)info;
+		killOnJobClose = (info2.LimitFlags & JOB_OBJECT_LIMIT_KILL_ON_JOB_CLOSE) != 0;
+	}
+	if (infoClass == JobObjectExtendedLimitInformation)
+	{
+		auto info2 = *(JOBOBJECT_EXTENDED_LIMIT_INFORMATION*)info;
+		killOnJobClose = (info2.BasicLimitInformation.LimitFlags & JOB_OBJECT_LIMIT_KILL_ON_JOB_CLOSE) != 0;
+	}
+	if (!killOnJobClose)
+		return res;
+
+	SCOPED_WRITE_LOCK(g_jobsLock, lock);
+	JobInfo& jobInfo = g_jobs[dh.fileObject];
+	jobInfo.shouldKill = true;
+	prefix = TC("(KillOnClose) ");
+	return res;
+}
+
+NTSTATUS Detoured_NtAssignProcessToJobObject(HANDLE job, HANDLE process)
+{
+	UBA_ASSERT(isDetouredHandle(job));
+	DetouredHandle& dh = asDetouredHandle(job);
+	HANDLE trueJobHandle = dh.trueHandle;
+
+	HANDLE processHandle = process;
+	if (isDetouredHandle(process))
+		processHandle = asDetouredHandle(process).trueHandle;
+
+	NTSTATUS res = True_NtAssignProcessToJobObject(trueJobHandle, processHandle);
+	DEBUG_LOG_DETOURED(L"NtAssignProcessToJobObject", L"Process %llu assigned to %llu -> %s", u64(process), u64(job), ToString(res).data);
+	if (res != STATUS_SUCCESS)
+		return res;
+
+	SCOPED_WRITE_LOCK(g_jobsLock, lock);
+	JobInfo& jobInfo = g_jobs[dh.fileObject];
+	jobInfo.processes.push_back(process);
+	return res;
+}
+
+NTSTATUS Detoured_NtTerminateJobObject(HANDLE handle, NTSTATUS status)
+{
+	UBA_ASSERT(!isDetouredHandle(handle));
+	auto res = True_NtTerminateJobObject(handle, status);
+	DEBUG_LOG_DETOURED(L"NtTerminateJobObject", L"%llu -> %s", u64(handle), ToString(res).data);
+	return res;
 }
 
+
 NTSTATUS NTAPI Detoured_NtCreateSection(PHANDLE SectionHandle, ACCESS_MASK DesiredAccess, POBJECT_ATTRIBUTES ObjectAttributes, PLARGE_INTEGER MaximumSize, ULONG SectionPageProtection, ULONG AllocationAttributes, HANDLE FileHandle)
 {
 	DETOURED_CALL(NtCreateSection);
@@ -880,13 +978,7 @@ VOID NTAPI Detoured_RtlExitUserProcess(NTSTATUS ExitStatus)
 	DETOURED_CALL(ExitProcess);
 	//DEBUG_LOG_TRUE(L"RtlExitUserProcess", L"(%u)", uExitCode);
 
-	CloseCaches();
-	SendExitMessage((DWORD)ExitStatus, GetTime());
-	PostDeinit();
-	#if UBA_DEBUG_LOG_ENABLED
-	FlushDebugLog();
-	#endif
-
+	CloseProcess(u32(ExitStatus));
 	True_RtlExitUserProcess(ExitStatus);
 }
 
@@ -2077,21 +2169,16 @@ NTSTATUS NTAPI Detoured_NtClose(HANDLE handle)
 
 	DetouredHandle& dh = asDetouredHandle(handle);
 
-	NTSTATUS res = STATUS_SUCCESS;
-
-	if (dh.trueHandle != INVALID_HANDLE_VALUE)
-	{
-		TimerScope ts(g_kernelStats.closeFile);
-		res = True_NtClose(dh.trueHandle);
-	}
+	HANDLE trueHandle = dh.trueHandle;
+	auto closeTrueHandle = [&]() { if (trueHandle == INVALID_HANDLE_VALUE) return STATUS_SUCCESS; TimerScope ts(g_kernelStats.closeFile); return True_NtClose(dh.trueHandle); };
 
 	FileObject* fo = dh.fileObject;
 	if (!fo)
 	{
-		if (dh.type >= HandleType_StdErr) // TODO: This might leak if handle is duplicated.. but ignore for now
-			return res;
+		if (dh.type < HandleType_StdErr) // TODO: We leak std handles if duplicated.. but ignore for now
+			delete& dh;
+		auto res = closeTrueHandle();
 		DEBUG_LOG_TRUE(L"NtClose", L"%llu (%ls) -> %ls", uintptr_t(handle), HandleToName(handle), ToString(res).data);
-		delete& dh;
 		return res;
 	}
 
@@ -2099,11 +2186,29 @@ NTSTATUS NTAPI Detoured_NtClose(HANDLE handle)
 	UBA_ASSERT(foRefCount != ~u64(0));
 	if (foRefCount)
 	{
-		DEBUG_LOG_TRUE(L"NtClose", L"%llu (%ls) -> %ls", uintptr_t(handle), HandleToName(handle), ToString(res).data);
 		delete& dh;
+		auto res = closeTrueHandle();
+		DEBUG_LOG_TRUE(L"NtClose", L"%llu (%ls) -> %ls", uintptr_t(handle), HandleToName(handle), ToString(res).data);
 		return res;
 	}
 
+	if (dh.type == HandleType_Job)
+	{
+		SCOPED_WRITE_LOCK(g_jobsLock, lock);
+		auto findIt = g_jobs.find(dh.fileObject);
+		if (findIt != g_jobs.end())
+		{
+			JobInfo jobInfo = std::move(findIt->second);
+			g_jobs.erase(findIt);
+			lock.Leave();
+			TerminateProcessesInJob(std::move(jobInfo), 0xC000010A);
+		}
+		DEBUG_LOG_TRUE(L"NtClose", L"%llu (Job) (%ls)", uintptr_t(handle), HandleToName(handle));
+		return closeTrueHandle();
+	}
+
+	NTSTATUS res = closeTrueHandle();
+
 	SharedMemoryHandle memoryHandle;
 	u64 memoryWritten = 0;
 	FileInfo& fi = *fo->fileInfo;
@@ -2452,7 +2557,19 @@ NTSTATUS Detoured_NtQuerySecurityObject(HANDLE Handle, SECURITY_INFORMATION Secu
 
 #if DETOURED_INCLUDE_DEBUG
 
-NTSTATUS NTAPI Detoured_NtCreateIoCompletion(PHANDLE IoCompletionHandle, ACCESS_MASK DesiredAccess, POBJECT_ATTRIBUTES ObjectAttributes, ULONG Count)
+NTSTATUS Detoured_NtOpenJobObject(HANDLE* handle, ACCESS_MASK access, const OBJECT_ATTRIBUTES *attr)
+{
+	UBA_ASSERT(!isDetouredHandle(handle));
+	return True_NtOpenJobObject(handle, access, attr);
+}
+
+NTSTATUS Detoured_NtQueryInformationJobObject(HANDLE handle, JOBOBJECTINFOCLASS infoClass, void* info, ULONG len, ULONG* retLen)
+{
+	UBA_ASSERT(!isDetouredHandle(handle));
+	return True_NtQueryInformationJobObject(handle, infoClass, info, len, retLen);
+}
+
+NTSTATUS Detoured_NtCreateIoCompletion(PHANDLE IoCompletionHandle, ACCESS_MASK DesiredAccess, POBJECT_ATTRIBUTES ObjectAttributes, ULONG Count)
 {
 	UBA_ASSERT(!isDetouredHandle(*IoCompletionHandle));
 	return True_NtCreateIoCompletion(IoCompletionHandle, DesiredAccess, ObjectAttributes, Count);
diff --git a/Engine/Source/Programs/UnrealBuildAccelerator/Detours/Private/Windows/UbaDetoursFunctionsWin.h b/Engine/Source/Programs/UnrealBuildAccelerator/Detours/Private/Windows/UbaDetoursFunctionsWin.h
--- a/Engine/Source/Programs/UnrealBuildAccelerator/Detours/Private/Windows/UbaDetoursFunctionsWin.h
+++ b/Engine/Source/Programs/UnrealBuildAccelerator/Detours/Private/Windows/UbaDetoursFunctionsWin.h
@@ -122,8 +122,6 @@
 	DETOURED_FUNCTIONS_KERNELBASE_DEBUG \
 
 #define DETOURED_FUNCTIONS_KERNEL32 \
-	DETOURED_FUNCTION(AssignProcessToJobObject) \
-	DETOURED_FUNCTION(SetInformationJobObject) \
 	DETOURED_FUNCTION(RegisterWaitForSingleObject) \
 	DETOURED_FUNCTION(CreateFileMappingA) \
 	DETOURED_FUNCTION(GetExitCodeProcess) \
@@ -150,6 +148,10 @@
 	DETOURED_FUNCTION(NtSetInformationFile) \
 	DETOURED_FUNCTION(NtSetInformationObject) \
 	DETOURED_FUNCTION(NtTerminateProcess) \
+	DETOURED_FUNCTION(NtCreateJobObject) \
+	DETOURED_FUNCTION(NtSetInformationJobObject) \
+	DETOURED_FUNCTION(NtAssignProcessToJobObject) \
+	DETOURED_FUNCTION(NtTerminateJobObject) \
 	DETOURED_FUNCTION(NtCreateSection) \
 	DETOURED_FUNCTION(RtlExitUserProcess) \
 	DETOURED_FUNCTION(RtlSizeHeap) \
@@ -375,6 +377,8 @@
 	DETOURED_FUNCTION(RtlGetFullPathName_U) \
 	DETOURED_FUNCTION(RtlGetFullPathName_UEx) \
 	DETOURED_FUNCTION(RtlDosPathNameToRelativeNtPathName_U) \
+	DETOURED_FUNCTION(NtOpenJobObject) \
+	DETOURED_FUNCTION(NtQueryInformationJobObject) \
 	DETOURED_FUNCTION(NtCreateIoCompletion) \
 	DETOURED_FUNCTION(NtFlushBuffersFileEx) \
 	DETOURED_FUNCTION(NtAlpcCreatePort) \
diff --git a/Engine/Source/Programs/UnrealBuildAccelerator/Detours/Private/Windows/UbaDetoursUtilsWin.h b/Engine/Source/Programs/UnrealBuildAccelerator/Detours/Private/Windows/UbaDetoursUtilsWin.h
--- a/Engine/Source/Programs/UnrealBuildAccelerator/Detours/Private/Windows/UbaDetoursUtilsWin.h
+++ b/Engine/Source/Programs/UnrealBuildAccelerator/Detours/Private/Windows/UbaDetoursUtilsWin.h
@@ -36,6 +36,7 @@ namespace uba
 		HandleType_File,
 		HandleType_FileMapping,
 		HandleType_Process,
+		HandleType_Job,
 		HandleType_StdErr,
 		HandleType_StdOut,
 		HandleType_StdIn,
