[UBA]
* Fixed bug in NtClose where detoured handle was deleted before logging of it

diff --git a/Engine/Source/Programs/UnrealBuildAccelerator/Detours/Private/Windows/UbaDetoursFunctionsNtDll.inl b/Engine/Source/Programs/UnrealBuildAccelerator/Detours/Private/Windows/UbaDetoursFunctionsNtDll.inl
--- a/Engine/Source/Programs/UnrealBuildAccelerator/Detours/Private/Windows/UbaDetoursFunctionsNtDll.inl
+++ b/Engine/Source/Programs/UnrealBuildAccelerator/Detours/Private/Windows/UbaDetoursFunctionsNtDll.inl
@@ -2170,15 +2170,15 @@ NTSTATUS NTAPI Detoured_NtClose(HANDLE handle)
 	DetouredHandle& dh = asDetouredHandle(handle);
 
 	HANDLE trueHandle = dh.trueHandle;
-	auto closeTrueHandle = [&]() { if (trueHandle == INVALID_HANDLE_VALUE) return STATUS_SUCCESS; TimerScope ts(g_kernelStats.closeFile); return True_NtClose(dh.trueHandle); };
+	auto closeTrueHandle = [&]() { if (trueHandle == INVALID_HANDLE_VALUE) return STATUS_SUCCESS; TimerScope ts(g_kernelStats.closeFile); return True_NtClose(trueHandle); };
 
 	FileObject* fo = dh.fileObject;
 	if (!fo)
 	{
-		if (dh.type < HandleType_StdErr) // TODO: We leak std handles if duplicated.. but ignore for now
-			delete& dh;
 		auto res = closeTrueHandle();
 		DEBUG_LOG_TRUE(L"NtClose", L"%llu (%ls) -> %ls", uintptr_t(handle), HandleToName(handle), ToString(res).data);
+		if (dh.type < HandleType_StdErr) // TODO: We leak std handles if duplicated.. but ignore for now
+			delete& dh;
 		return res;
 	}
 
@@ -2186,9 +2186,9 @@ NTSTATUS NTAPI Detoured_NtClose(HANDLE handle)
 	UBA_ASSERT(foRefCount != ~u64(0));
 	if (foRefCount)
 	{
-		delete& dh;
 		auto res = closeTrueHandle();
 		DEBUG_LOG_TRUE(L"NtClose", L"%llu (%ls) -> %ls", uintptr_t(handle), HandleToName(handle), ToString(res).data);
+		delete& dh;
 		return res;
 	}
 
