[UBA]
* Fixed bug in test code where string was null terminated in the wrong place

diff --git a/Engine/Source/Programs/UnrealBuildAccelerator/TestApp/Private/Windows/UbaTestAppWin.cpp b/Engine/Source/Programs/UnrealBuildAccelerator/TestApp/Private/Windows/UbaTestAppWin.cpp
--- a/Engine/Source/Programs/UnrealBuildAccelerator/TestApp/Private/Windows/UbaTestAppWin.cpp
+++ b/Engine/Source/Programs/UnrealBuildAccelerator/TestApp/Private/Windows/UbaTestAppWin.cpp
@@ -405,7 +405,7 @@ int wmain(int argc, wchar_t* argv[])
 		int read = ReadTestFile(buf, 32*2, L"SpecialFile1");
 		if (read != 2)
 			return LogError(L"Expected to read two bytes but got %i", read);
-		buf[read*2] = 0;
+		buf[read] = 0;
 
 		int readIndex = _wtoi(buf);
 		if (callIndex != readIndex)
