[UBA]
* Added traffic.dll to known dlls (must use wine side dll to prevent crashes)
* Fixed so UbaCli.exe -listimports <badfilename> outputs error failing to open file

diff --git a/Engine/Source/Programs/UnrealBuildAccelerator/Cli/Private/UbaCli.cpp b/Engine/Source/Programs/UnrealBuildAccelerator/Cli/Private/UbaCli.cpp
--- a/Engine/Source/Programs/UnrealBuildAccelerator/Cli/Private/UbaCli.cpp
+++ b/Engine/Source/Programs/UnrealBuildAccelerator/Cli/Private/UbaCli.cpp
@@ -1231,7 +1231,7 @@ namespace uba
 				}
 				logger.Info(TC("  %s"), import);
 
-			}, error))
+			}, error) || error.count)
 				return logger.Error(TC("%s"), error.data);
 			return true;
 		}
diff --git a/Engine/Source/Programs/UnrealBuildAccelerator/Core/Public/UbaWinBinDependencyParser.h b/Engine/Source/Programs/UnrealBuildAccelerator/Core/Public/UbaWinBinDependencyParser.h
--- a/Engine/Source/Programs/UnrealBuildAccelerator/Core/Public/UbaWinBinDependencyParser.h
+++ b/Engine/Source/Programs/UnrealBuildAccelerator/Core/Public/UbaWinBinDependencyParser.h
@@ -71,6 +71,7 @@ namespace uba
 		L"shlwapi.dll",
 		L"sspicli.dll",
 		L"tbs.dll",
+		L"traffic.dll",
 		L"ucrtbase.dll",
 		L"umpdc.dll",
 		L"umppc17706.dll",
@@ -110,7 +111,7 @@ namespace uba
 
 	DWORD UbaExceptionHandlerAllowAccessViolation(DWORD code, _EXCEPTION_POINTERS* info, const tchar* fileName);
 
-	inline bool FindImportsInMem(const tchar* fileName, void* mem, u64 memSize, const Function<void(const tchar* import, bool isKnown, const tchar* const* importLoaderPaths)>& func)
+	inline bool FindImportsInMem(StringView fileName, void* mem, u64 memSize, const Function<void(const tchar* import, bool isKnown, const tchar* const* importLoaderPaths)>& func, StringBufferBase& outError)
 	{
 		auto hdrs = (PIMAGE_NT_HEADERS)(PCHAR(mem) + PIMAGE_DOS_HEADER(mem)->e_lfanew);
 		auto pSech=IMAGE_FIRST_SECTION(hdrs);
@@ -125,19 +126,19 @@ namespace uba
 			while(true)
 			{
 				if (u64((u8*)importDesc - (u8*)mem) + sizeof(IMAGE_IMPORT_DESCRIPTOR) > memSize)
-					return false;
+					return outError.Appendf(TC("Import parsing error1: %.*s"), fileName.count, fileName.data).ToFalse();
 				if (importDesc->Name == NULL)
 					break;
 				char* name = (PCHAR)((DWORD_PTR)mem + GetOffset(importDesc->Name, pSech, hdrs));
 				if (u64((u8*)name - (u8*)mem) + 2 > memSize)
-					return false;
+					return outError.Appendf(TC("Import parsing error2: %.*s"), fileName.count, fileName.data).ToFalse();
 				StringBuffer<256> wname;
 				wname.Append(name);
 				func(wname.data, IsKnownSystemFile(wname.data), nullptr);
 				++importDesc;
 			}
 		}
-		__except(UbaExceptionHandlerAllowAccessViolation(GetExceptionCode(), GetExceptionInformation(), fileName))
+		__except(UbaExceptionHandlerAllowAccessViolation(GetExceptionCode(), GetExceptionInformation(), fileName.data))
 		{
 			return false;
 		}
@@ -147,23 +148,30 @@ namespace uba
 	inline bool ParseBinary(const StringView& filePath, const StringView& originalPath, BinaryInfo& outInfo, const Function<void(const tchar* import, bool isKnown, const tchar* const* importLoaderPaths)>& func, StringBufferBase& outError)
 	{
 		if (filePath.EndsWith(TCV(".bat")))
+		{
+			outError.Appendf(TC("Batch file has no dependencies: %.*s"), filePath.count, filePath.data);
 			return true;
+		}
 		HANDLE fileHandle = True_CreateFileW(filePath.data, GENERIC_READ, FILE_SHARE_READ, 0, OPEN_EXISTING, FILE_ATTRIBUTE_NORMAL | FILE_FLAG_SEQUENTIAL_SCAN, 0);
 		if (fileHandle == INVALID_HANDLE_VALUE)
+		{
+			outError.Appendf(TC("File not found: %.*s"), filePath.count, filePath.data);
 			return true;
+		}
+
 		auto closeFileHandle = MakeGuard([&]() { CloseHandle(fileHandle); });
 		HANDLE fileMapping = True_CreateFileMappingW(fileHandle, NULL, PAGE_READONLY, 0, 0, NULL);
 		if (!fileMapping)
-			return false;
+			return outError.Appendf(TC("Failed to create mapping: %.*s"), filePath.count, filePath.data).ToFalse();
 		auto closeMappingHandle = MakeGuard([&]() { CloseHandle(fileMapping); });
 		void* mem = True_MapViewOfFile(fileMapping, FILE_MAP_READ, 0, 0, 0);
 		if (!mem)
-			return false;
+			return outError.Appendf(TC("Failed to map view of mapping: %.*s"), filePath.count, filePath.data).ToFalse();
 		auto unmap = MakeGuard([&]() { True_UnmapViewOfFile(mem); });
 		LARGE_INTEGER fileSize;
 		if (!::GetFileSizeEx(fileHandle, &fileSize))
-			return false;
+			return outError.Appendf(TC("Failed to get size of file: %.*s"), filePath.count, filePath.data).ToFalse();
 
-		return FindImportsInMem(filePath.data, mem, fileSize.QuadPart, func);
+		return FindImportsInMem(filePath, mem, fileSize.QuadPart, func, outError);
 	}
 }
