[UBA]
* Added so UbObjTool is hooking up allocator for oodle

diff --git a/Engine/Source/Programs/UnrealBuildAccelerator/ObjTool/Private/UbaObjTool.cpp b/Engine/Source/Programs/UnrealBuildAccelerator/ObjTool/Private/UbaObjTool.cpp
--- a/Engine/Source/Programs/UnrealBuildAccelerator/ObjTool/Private/UbaObjTool.cpp
+++ b/Engine/Source/Programs/UnrealBuildAccelerator/ObjTool/Private/UbaObjTool.cpp
@@ -58,6 +58,11 @@ namespace uba
 	{
 		using namespace uba;
 
+		//#if !PLATFORM_LINUX
+		//AddExceptionHandler();
+		//#endif
+		InitMemory();
+
 		TString objFile;
 		bool printSymbols = false;
 		bool writeImpLib = false;
