[UBA]
* Fixed bug in FixPath2 where it used strchr on a non-nullterminated string to see if there was a ~ in the path. This obviously gave false positives sometimes and likely could also cause crashes

diff --git a/Engine/Source/Programs/UnrealBuildAccelerator/Core/Public/UbaPathUtils.h b/Engine/Source/Programs/UnrealBuildAccelerator/Core/Public/UbaPathUtils.h
--- a/Engine/Source/Programs/UnrealBuildAccelerator/Core/Public/UbaPathUtils.h
+++ b/Engine/Source/Programs/UnrealBuildAccelerator/Core/Public/UbaPathUtils.h
@@ -21,6 +21,9 @@ namespace uba
 
 #if PLATFORM_WINDOWS
 
+		if (fileNameLen == ~0u)
+			fileNameLen = TStrlen(fileName);
+
 		const tchar* read = fileName;
 		tchar* write = buffer;
 
@@ -28,48 +31,47 @@ namespace uba
 		tchar lastChar = 0;
 		bool hasDotBeforeSlash = false;
 		bool hasDotDot = false;
-		//bool hasSlash = false;
 		bool seenNonBackslash = false;
 
 		if (fileName[0] == '\"')
+		{
 			++read;
+			--fileNameLen;
+		}
 
-		tchar fullName[1024];
+		StringBuffer<1024> fullName;
 		while (true)
 		{
 			tchar temp[1024];
-			if (read == fullName)
+			if (read == fullName.data)
 			{
 				TStrcpy_s(temp, sizeof_array(temp), read);
 				read = temp;
 			}
 
-			const tchar* tilde = TStrchr(read, '~');
-			if (!tilde)
+			StringView sv(read, fileNameLen);
+			u32 tildePos;
+			if (!sv.Contains('~', &tildePos))
 				break;
-
-			if (fileNameLen == ~0u)
-				fileNameLen = TStrlen(read);
-			if (fileNameLen >= 1024)
-				return false;
-
+			
 			// Since this might be a in memory file we can't use the actual full name, so let's find first slash after '~'
-			const tchar* backslash = TStrchr(tilde, '\\');
-			memcpy(fullName, read, fileNameLen*sizeof(tchar));
-			fullName[fileNameLen] = 0;
-			if (backslash)
-				fullName[backslash - read] = 0;
-			u32 len = Local_GetLongPathNameW(fullName, fullName, sizeof_array(fullName));
+			StringView tilde = sv.Skip(tildePos);
+			u32 backslashAfterTilde = tilde.count;
+			bool hasBackslashAfterTilde = tilde.Contains('\\', &backslashAfterTilde);
+
+			fullName.Clear().Append(sv.data, tildePos + backslashAfterTilde);
+			u32 len = Local_GetLongPathNameW(fullName.data, fullName.data, sizeof_array(fullName));
 			if (!len)
 			{
 				if (outBufferCharLen)
 					*outBufferCharLen = 0;
 				return false;
 			}
-			if (backslash)
-				TStrcpy_s(fullName + len, sizeof_array(fullName) - len, backslash);
-			read = fullName;
-			fileNameLen = TStrlen(read);
+			fullName.count = len;
+			if (hasBackslashAfterTilde)
+				fullName.Append(tilde.Skip(backslashAfterTilde));
+			read = fullName.data;
+			fileNameLen = fullName.count;
 		}
 
 		if (read[0] == '\\')
