[UBA]
* Fixed so NtQueryInformationFile FileNetworkOpenInformation is more correct
* Fixed so NtClose is showing proper D/T in detouring log
* Changed so .hpp are put in shared memory when loaded
* Introduced DEBUG_LOG_PROXY for entries where we know we call a detoured layer under
* Increased max callstack depth from 16 to 24 in GetCallstack
* Removed some detouring logging spam (_get_osfhandle)

diff --git a/Engine/Source/Programs/UnrealBuildAccelerator/Common/Private/UbaFile.cpp b/Engine/Source/Programs/UnrealBuildAccelerator/Common/Private/UbaFile.cpp
--- a/Engine/Source/Programs/UnrealBuildAccelerator/Common/Private/UbaFile.cpp
+++ b/Engine/Source/Programs/UnrealBuildAccelerator/Common/Private/UbaFile.cpp
@@ -210,7 +210,7 @@ namespace uba
 			return GetFileBasicInformation(out, logger, fileName, errorOnFail);
 		IO_STATUS_BLOCK b;
 		FILE_NETWORK_OPEN_INFORMATION info;
-		NTSTATUS res = NtQueryInformationFile(asHANDLE(hFile), &b, &info, sizeof(info), (FILE_INFORMATION_CLASS)34); // FileNetworkOpenInformation
+		NTSTATUS res = NtQueryInformationFile(asHANDLE(hFile), &b, &info, sizeof(info), FileNetworkOpenInformation);
 		if (res != STATUS_SUCCESS)
 			return errorOnFail ? logger.Error(TC("GetFileBasicInformationByHandle (NtQueryInformationFile) failed on %s (0x%x)"), fileName, res) : false;
 		out.attributes = info.FileAttributes;
diff --git a/Engine/Source/Programs/UnrealBuildAccelerator/Common/Private/UbaSession.cpp b/Engine/Source/Programs/UnrealBuildAccelerator/Common/Private/UbaSession.cpp
--- a/Engine/Source/Programs/UnrealBuildAccelerator/Common/Private/UbaSession.cpp
+++ b/Engine/Source/Programs/UnrealBuildAccelerator/Common/Private/UbaSession.cpp
@@ -2991,7 +2991,7 @@ namespace uba
 		// Reason is because there is additional overhead to use memory mappings.
 		// Upside is that all things that are memory mapped can be stored compressed in cas storage so it saves space.
 
-		if (fileName.EndsWith(TCV(".h")) || fileName.EndsWith(TCV(".inl")) || fileName.EndsWith(TCV(".gch")))
+		if (fileName.EndsWith(TCV(".h")) || fileName.EndsWith(TCV(".inl")) || fileName.EndsWith(TCV(".gch")) || fileName.EndsWith(TCV(".hpp")))
 			return 4 * 1024; // clang seems to need 4k alignment? Is it a coincidence it works or what is happening inside the code? (msvc works with alignment 1byte here)
 		if (fileName.EndsWith(TCV(".lib")))
 			return 4 * 1024;
diff --git a/Engine/Source/Programs/UnrealBuildAccelerator/Core/Private/UbaPlatform.cpp b/Engine/Source/Programs/UnrealBuildAccelerator/Core/Private/UbaPlatform.cpp
--- a/Engine/Source/Programs/UnrealBuildAccelerator/Core/Private/UbaPlatform.cpp
+++ b/Engine/Source/Programs/UnrealBuildAccelerator/Core/Private/UbaPlatform.cpp
@@ -273,7 +273,7 @@ namespace uba
 	UBA_NOINLINE u32 GetCallstack(void** outCallstack, u32 callstackCapacity, u32 skipCallstack, void* contextPtr)
 	{
 #if PLATFORM_WINDOWS
-		constexpr int MaxDepth = 16;// 62;
+		constexpr int MaxDepth = 24;// 62;
 		void* callers[MaxDepth] = {0};
 		u32 count = 0;//RtlCaptureStackBackTrace(0, MaxDepth, callers, NULL); // For some reason this does not work on wine
 		if (!count)
diff --git a/Engine/Source/Programs/UnrealBuildAccelerator/Detours/Private/UbaDetoursShared.h b/Engine/Source/Programs/UnrealBuildAccelerator/Detours/Private/UbaDetoursShared.h
--- a/Engine/Source/Programs/UnrealBuildAccelerator/Detours/Private/UbaDetoursShared.h
+++ b/Engine/Source/Programs/UnrealBuildAccelerator/Detours/Private/UbaDetoursShared.h
@@ -37,13 +37,10 @@ namespace uba
 			UBA_FMT_CHECK(fmt, ##__VA_ARGS__); \
 			if (isLogging()) WriteDebugLogWithPrefix(#Prefix, STRING_JOIN(ls, __LINE__), Command, fmt, ##__VA_ARGS__); \
 
-		//#define DEBUG_LOG_DETOURED(Command, ...) 
-		#define DEBUG_LOG_DETOURED(Command, ...) DEBUG_LOG_PREFIX(D, Command, __VA_ARGS__)
-		//#define DEBUG_LOG_TRUE(...)
 		#define DEBUG_LOG_TRUE(Command, ...) DEBUG_LOG_PREFIX(T, Command, __VA_ARGS__)
-		//#define DEBUG_LOG_PIPE(Command, ...) ts.leave(); DEBUG_LOG_PREFIX(P, Command, __VA_ARGS__)
+		#define DEBUG_LOG_DETOURED(Command, ...) DEBUG_LOG_PREFIX(D, Command, __VA_ARGS__)
+		#define DEBUG_LOG_PROXY(Command, ...) DEBUG_LOG_PREFIX(P, Command, __VA_ARGS__)
 		#define DEBUG_LOG_PIPE(Command, ...) ts.Leave();
-		//#define DEBUG_LOG(...)
 		#define DEBUG_LOG(fmt, ...) { UBA_FMT_CHECK(fmt, ##__VA_ARGS__); if (isLogging()) WriteDebugLog(fmt, ##__VA_ARGS__); }
 	#else
 		#define DEBUG_LOG(...) {}
diff --git a/Engine/Source/Programs/UnrealBuildAccelerator/Detours/Private/Windows/UbaDetoursFunctionsKernelBase.inl b/Engine/Source/Programs/UnrealBuildAccelerator/Detours/Private/Windows/UbaDetoursFunctionsKernelBase.inl
--- a/Engine/Source/Programs/UnrealBuildAccelerator/Detours/Private/Windows/UbaDetoursFunctionsKernelBase.inl
+++ b/Engine/Source/Programs/UnrealBuildAccelerator/Detours/Private/Windows/UbaDetoursFunctionsKernelBase.inl
@@ -202,7 +202,7 @@ BOOL Detoured_DuplicateHandle(HANDLE hSourceProcessHandle, HANDLE hSourceHandle,
 HANDLE Detoured_CreateFileW(LPCWSTR lpFileName, DWORD dwDesiredAccess, DWORD dwShareMode, LPSECURITY_ATTRIBUTES lpSecurityAttributes, DWORD dwCreationDisposition, DWORD dwFlagsAndAttributes, HANDLE hTemplateFile)
 {
 	DETOURED_CALL(CreateFileW);
-	DEBUG_LOG_DETOURED(L"CreateFileW", L"%ls", lpFileName);
+	DEBUG_LOG_PROXY(L"CreateFileW", L"%ls", lpFileName);
 	u32 disallowDetour = !lpFileName || Equals(lpFileName, L"nul");
 	t_disallowDetour += disallowDetour;
 
@@ -1036,6 +1036,10 @@ BOOL Detoured_GetVolumePathNameW(LPCWSTR lpszFileName, LPWSTR lpszVolumePathName
 
 	DEBUG_LOG_TRUE(L"GetVolumePathNameW", L"(%ls)", lpszFileName);
 	SuppressCreateFileDetourScope cfs;
+
+
+	// TODO: This is causing TONS of kernel calls
+
 	auto res = True_GetVolumePathNameW(lpszFileName, lpszVolumePathName, cchBufferLength);
 	return res;
 }
@@ -3697,14 +3701,14 @@ BOOL Detoured_GetVolumePathNameA(LPCSTR lpszFileName, LPSTR lpszVolumePathName,
 DWORD Detoured_GetFileSize(HANDLE hFile, LPDWORD lpFileSizeHigh)
 {
 	DETOURED_CALL(GetFileSize);
-	DEBUG_LOG_TRUE(L"GetFileSize", L"%llu (%ls)", uintptr_t(hFile), HandleToName(hFile));
+	DEBUG_LOG_PROXY(L"GetFileSize", L"%llu (%ls)", uintptr_t(hFile), HandleToName(hFile));
 	return True_GetFileSize(hFile, lpFileSizeHigh); // Calls NtQueryInformationFile
 }
 
 DWORD Detoured_GetFileSizeEx(HANDLE hFile, PLARGE_INTEGER lpFileSize)
 {
 	DETOURED_CALL(GetFileSizeEx);
-	DEBUG_LOG_TRUE(L"GetFileSizeEx", L"%llu (%ls)", uintptr_t(hFile), HandleToName(hFile));
+	DEBUG_LOG_PROXY(L"GetFileSizeEx", L"%llu (%ls)", uintptr_t(hFile), HandleToName(hFile));
 	return True_GetFileSizeEx(hFile, lpFileSize); // Calls NtQueryInformationFile
 }
 
@@ -3725,8 +3729,8 @@ DWORD Detoured_SetFilePointer(HANDLE hFile, LONG lDistanceToMove, PLONG lpDistan
 DWORD Detoured_SetFilePointerEx(HANDLE hFile, LARGE_INTEGER liDistanceToMove, PLARGE_INTEGER lpNewFilePointer, DWORD dwMoveMethod)
 {
 	DETOURED_CALL(SetFilePointerEx);
-	DEBUG_LOG_TRUE(L"SetFilePointerEx", L"%llu %lli %u (%ls)", uintptr_t(hFile), liDistanceToMove.QuadPart, dwMoveMethod, HandleToName(hFile));
-	return True_SetFilePointerEx(hFile, liDistanceToMove, lpNewFilePointer, dwMoveMethod); // This ends up in NtSetInformationFile
+	DEBUG_LOG_PROXY(L"SetFilePointerEx", L"%llu %lli %u (%ls)", uintptr_t(hFile), liDistanceToMove.QuadPart, dwMoveMethod, HandleToName(hFile));
+	return True_SetFilePointerEx(hFile, liDistanceToMove, lpNewFilePointer, dwMoveMethod); // This ends up in NtQueryInformationFile and NtSetInformationFile
 }
 
 BOOL Detoured_SetEndOfFile(HANDLE hFile)
diff --git a/Engine/Source/Programs/UnrealBuildAccelerator/Detours/Private/Windows/UbaDetoursFunctionsNtDll.inl b/Engine/Source/Programs/UnrealBuildAccelerator/Detours/Private/Windows/UbaDetoursFunctionsNtDll.inl
--- a/Engine/Source/Programs/UnrealBuildAccelerator/Detours/Private/Windows/UbaDetoursFunctionsNtDll.inl
+++ b/Engine/Source/Programs/UnrealBuildAccelerator/Detours/Private/Windows/UbaDetoursFunctionsNtDll.inl
@@ -245,37 +245,44 @@ NTSTATUS Detoured_NtQueryInformationFile(HANDLE FileHandle, PIO_STATUS_BLOCK IoS
 			return STATUS_SUCCESS;
 		};
 
-	auto setStandardInformation = [&](FILE_STANDARD_INFORMATION& out, FileInfo& fi, DirectoryTable::EntryInformation* entryInfo)
+	auto getFileSizeAndAllocation = [&](u64& outSize, u64& outAllocationSize, FileInfo& fi, DirectoryTable::EntryInformation* entryInfo)
 		{
-			u64 size = 0;
-			u64 allocationSize = 0;
 			if (fi.memoryFile)
 			{
-				size = fi.memoryFile->writtenSize;
-				allocationSize = fi.memoryFile->committedSize;
+				outSize = fi.memoryFile->writtenSize;
+				outAllocationSize = fi.memoryFile->committedSize;
 			}
 			else if (fi.fileMapMem)
 			{
-				size = fi.fileMapMemSize;
-				allocationSize = size;
+				outSize = fi.fileMapMemSize;
+				outAllocationSize = outSize;
 			}
 			else if (fi.size != InvalidValue)
 			{
-				size = fi.size;
-				allocationSize = size;
+				outSize = fi.size;
+				outAllocationSize = outSize;
 			}
 			else if (entryInfo)
 			{
-				size = entryInfo->size;
-				allocationSize = size;
+				outSize = entryInfo->size;
+				outAllocationSize = outSize;
 			}
 			else
+				return false;
+			return true;
+		};
+
+	auto setStandardInformation = [&](FILE_STANDARD_INFORMATION& out, FileInfo& fi, DirectoryTable::EntryInformation* entryInfo)
+		{
+			u64 size = 0;
+			u64 allocationSize = 0;
+			if (!getFileSizeAndAllocation(size, allocationSize, fi, entryInfo))
 				return STATUS_INVALID_HANDLE;
+			out.EndOfFile.QuadPart = size;
+			out.AllocationSize.QuadPart = allocationSize;
 			out.DeletePending = FALSE;
 			out.NumberOfLinks = 0;
 			out.Directory = FALSE;
-			out.AllocationSize.QuadPart = allocationSize;
-			out.EndOfFile.QuadPart = size;
 			return STATUS_SUCCESS;
 		};
 
@@ -373,20 +380,21 @@ NTSTATUS Detoured_NtQueryInformationFile(HANDLE FileHandle, PIO_STATUS_BLOCK IoS
 		DirectoryTable::EntryInformation entryInfo;
 		g_directoryTable.GetEntryInformation(entryInfo, dh.dirTableOffset);
 		UBA_ASSERT(entryInfo.attributes != 0);
+		FileInfo& fi = *dh.fileObject->fileInfo;
+
+		u64 size = 0;
+		u64 allocationSize = 0;
+		getFileSizeAndAllocation(size, allocationSize, fi, &entryInfo);
 
-		// TODO This code path is here to handle nodejs queries.. Is not properly implemented and miss things
 		auto& info = *(FILE_NETWORK_OPEN_INFORMATION*)FileInformation;
 		info.CreationTime.QuadPart = entryInfo.lastWrite;
 		info.LastAccessTime.QuadPart = entryInfo.lastWrite;
 		info.LastWriteTime.QuadPart = entryInfo.lastWrite;
 		info.ChangeTime.QuadPart = entryInfo.lastWrite;
-		info.AllocationSize.QuadPart = entryInfo.size;
-		u64 fileSize = dh.fileObject->fileInfo->size;
-		if (fileSize == InvalidValue)
-			fileSize = entryInfo.size;
-		info.EndOfFile.QuadPart = fileSize;
+		info.EndOfFile.QuadPart = size;
+		info.AllocationSize.QuadPart = allocationSize;
 		info.FileAttributes = entryInfo.attributes;
-		DEBUG_LOG_DETOURED(L"NtQueryInformationFile", L"NOT IMPLEMENTED (FileNetworkOpenInformation) %llu (%s) -> %ls", uintptr_t(FileHandle), HandleToName(FileHandle), ToString(STATUS_SUCCESS).data);
+		DEBUG_LOG_DETOURED(L"NtQueryInformationFile", L"(FileNetworkOpenInformation) %llu (%s) -> %ls", uintptr_t(FileHandle), HandleToName(FileHandle), ToString(STATUS_SUCCESS).data);
 		return STATUS_SUCCESS;
 	}
 	if (FileInformationClass == FileNormalizedNameInformation)
@@ -2170,13 +2178,28 @@ NTSTATUS NTAPI Detoured_NtClose(HANDLE handle)
 	DetouredHandle& dh = asDetouredHandle(handle);
 
 	HANDLE trueHandle = dh.trueHandle;
-	auto closeTrueHandle = [&]() { if (trueHandle == INVALID_HANDLE_VALUE) return STATUS_SUCCESS; TimerScope ts(g_kernelStats.closeFile); return True_NtClose(trueHandle); };
+	auto closeTrueHandle = [&](bool log)
+		{
+			if (trueHandle == INVALID_HANDLE_VALUE)
+			{
+				if (!log)
+					return STATUS_SUCCESS;
+				DEBUG_LOG_DETOURED(L"NtClose", L"%llu (%ls) -> %ls", uintptr_t(handle), HandleToName(handle), ToString(STATUS_SUCCESS).data);
+				return STATUS_SUCCESS;
+			}
+			TimerScope ts(g_kernelStats.closeFile);
+			auto res = True_NtClose(trueHandle);
+			ts.Leave();
+			if (!log)
+				return res;
+			DEBUG_LOG_TRUE(L"NtClose", L"%llu (%ls) -> %ls", uintptr_t(handle), HandleToName(handle), ToString(res).data);
+			return res;
+		};
 
 	FileObject* fo = dh.fileObject;
 	if (!fo)
 	{
-		auto res = closeTrueHandle();
-		DEBUG_LOG_TRUE(L"NtClose", L"%llu (%ls) -> %ls", uintptr_t(handle), HandleToName(handle), ToString(res).data);
+		auto res = closeTrueHandle(true);
 		if (dh.type < HandleType_StdErr) // TODO: We leak std handles if duplicated.. but ignore for now
 			delete& dh;
 		return res;
@@ -2186,8 +2209,7 @@ NTSTATUS NTAPI Detoured_NtClose(HANDLE handle)
 	UBA_ASSERT(foRefCount != ~u64(0));
 	if (foRefCount)
 	{
-		auto res = closeTrueHandle();
-		DEBUG_LOG_TRUE(L"NtClose", L"%llu (%ls) -> %ls", uintptr_t(handle), HandleToName(handle), ToString(res).data);
+		auto res = closeTrueHandle(true);
 		delete& dh;
 		return res;
 	}
@@ -2204,10 +2226,10 @@ NTSTATUS NTAPI Detoured_NtClose(HANDLE handle)
 			TerminateProcessesInJob(std::move(jobInfo), 0xC000010A);
 		}
 		DEBUG_LOG_TRUE(L"NtClose", L"%llu (Job) (%ls)", uintptr_t(handle), HandleToName(handle));
-		return closeTrueHandle();
+		return closeTrueHandle(false);
 	}
 
-	NTSTATUS res = closeTrueHandle();
+	NTSTATUS res = closeTrueHandle(true);
 
 	SharedMemoryHandle memoryHandle;
 	u64 memoryWritten = 0;
diff --git a/Engine/Source/Programs/UnrealBuildAccelerator/Detours/Private/Windows/UbaDetoursFunctionsUcrtBase.inl b/Engine/Source/Programs/UnrealBuildAccelerator/Detours/Private/Windows/UbaDetoursFunctionsUcrtBase.inl
--- a/Engine/Source/Programs/UnrealBuildAccelerator/Detours/Private/Windows/UbaDetoursFunctionsUcrtBase.inl
+++ b/Engine/Source/Programs/UnrealBuildAccelerator/Detours/Private/Windows/UbaDetoursFunctionsUcrtBase.inl
@@ -216,7 +216,7 @@ intptr_t Detoured__get_osfhandle(int fd)
 	}
 
 	auto res = True__get_osfhandle(fd);
-	DEBUG_LOG_TRUE(L"_get_osfhandle", L"(%i) -> %llu", fd, u64(res));
+	//DEBUG_LOG_TRUE(L"_get_osfhandle", L"(%i) -> %llu", fd, u64(res)); // Spam
 	return res;
 }
 
