[UBA]
* Added so ispc application rule supports process racing

diff --git a/Engine/Source/Programs/UnrealBuildAccelerator/Core/Private/UbaApplicationRules.cpp b/Engine/Source/Programs/UnrealBuildAccelerator/Core/Private/UbaApplicationRules.cpp
--- a/Engine/Source/Programs/UnrealBuildAccelerator/Core/Private/UbaApplicationRules.cpp
+++ b/Engine/Source/Programs/UnrealBuildAccelerator/Core/Private/UbaApplicationRules.cpp
@@ -650,6 +650,11 @@ namespace uba
 		{
 			return fileName.EndsWith(TCV(".ispc.txt"));
 		}
+
+		virtual bool SupportsRacing() const override
+		{
+			return true;
+		}
 	};
 
 	class ApplicationRulesUBTDll : public ApplicationRules
