[UBA]
* Removed UnmapViewOfFile during shutdown of process since it is now using the shared memory where it can't unmap parts of it in the same way
* Changed more kernel detouring functions to log as proxy
* Added test (commented out) for FixPath handling short->long paths

diff --git a/Engine/Source/Programs/UnrealBuildAccelerator/Detours/Private/UbaDetoursShared.h b/Engine/Source/Programs/UnrealBuildAccelerator/Detours/Private/UbaDetoursShared.h
--- a/Engine/Source/Programs/UnrealBuildAccelerator/Detours/Private/UbaDetoursShared.h
+++ b/Engine/Source/Programs/UnrealBuildAccelerator/Detours/Private/UbaDetoursShared.h
@@ -46,6 +46,7 @@ namespace uba
 		#define DEBUG_LOG(...) {}
 		#define DEBUG_LOG_DETOURED(Command, ...) {}
 		#define DEBUG_LOG_TRUE(Command, ...) {}
+		#define DEBUG_LOG_PROXY(Command, ...) {}
 		#define DEBUG_LOG_PIPE(...) ts.Leave();
 	#endif
 
diff --git a/Engine/Source/Programs/UnrealBuildAccelerator/Detours/Private/Windows/UbaDetoursFunctionsKernelBase.inl b/Engine/Source/Programs/UnrealBuildAccelerator/Detours/Private/Windows/UbaDetoursFunctionsKernelBase.inl
--- a/Engine/Source/Programs/UnrealBuildAccelerator/Detours/Private/Windows/UbaDetoursFunctionsKernelBase.inl
+++ b/Engine/Source/Programs/UnrealBuildAccelerator/Detours/Private/Windows/UbaDetoursFunctionsKernelBase.inl
@@ -3722,7 +3722,7 @@ BOOL Detoured_GetFileTime(HANDLE hFile, LPFILETIME lpCreationTime, LPFILETIME lp
 DWORD Detoured_SetFilePointer(HANDLE hFile, LONG lDistanceToMove, PLONG lpDistanceToMoveHigh, DWORD dwMoveMethod)
 {
 	DETOURED_CALL(SetFilePointer);
-	DEBUG_LOG_TRUE(L"SetFilePointer", L"%llu %lli %u (%ls)", uintptr_t(hFile), ToLargeInteger(lpDistanceToMoveHigh ? *lpDistanceToMoveHigh : 0ll, lDistanceToMove).QuadPart, dwMoveMethod, HandleToName(hFile));
+	DEBUG_LOG_PROXY(L"SetFilePointer", L"%llu %lli %u (%ls)", uintptr_t(hFile), ToLargeInteger(lpDistanceToMoveHigh ? *lpDistanceToMoveHigh : 0ll, lDistanceToMove).QuadPart, dwMoveMethod, HandleToName(hFile));
 	return True_SetFilePointer(hFile, lDistanceToMove, lpDistanceToMoveHigh, dwMoveMethod);
 }
 
@@ -3736,7 +3736,7 @@ DWORD Detoured_SetFilePointerEx(HANDLE hFile, LARGE_INTEGER liDistanceToMove, PL
 BOOL Detoured_SetEndOfFile(HANDLE hFile)
 {
 	DETOURED_CALL(SetEndOfFile);
-	DEBUG_LOG_TRUE(L"SetEndOfFile", L"%llu (%ls)", uintptr_t(hFile), HandleToName(hFile));
+	DEBUG_LOG_PROXY(L"SetEndOfFile", L"%llu (%ls)", uintptr_t(hFile), HandleToName(hFile));
 	return True_SetEndOfFile(hFile);
 }
 
diff --git a/Engine/Source/Programs/UnrealBuildAccelerator/Detours/Private/Windows/UbaDetoursFunctionsWin.cpp b/Engine/Source/Programs/UnrealBuildAccelerator/Detours/Private/Windows/UbaDetoursFunctionsWin.cpp
--- a/Engine/Source/Programs/UnrealBuildAccelerator/Detours/Private/Windows/UbaDetoursFunctionsWin.cpp
+++ b/Engine/Source/Programs/UnrealBuildAccelerator/Detours/Private/Windows/UbaDetoursFunctionsWin.cpp
@@ -426,6 +426,8 @@ LastErrorToText::LastErrorToText(u32 lastError)
 		Resize(size - 2);
 }
 
+const wchar_t* ToString(BOOL b) { return b ? L"Success" : L"Error"; }
+
 const wchar_t* HandleToName(HANDLE handle)
 {
 	if (handle == INVALID_HANDLE_VALUE)
@@ -502,14 +504,14 @@ void CloseCaches()
 		FileInfo& info = it.second;
 		if (info.fileMapMem)
 		{
-			DEBUG_LOG_TRUE(L"INTERNAL UnmapViewOfFile", L"0x%llx (%ls) (%ls)", uintptr_t(info.fileMapMem), info.name, info.originalName);
-			True_UnmapViewOfFile(info.fileMapMem);
+			// TODO: Most files here are in the big shared memory block and can't be unmapped this way
+			//BOOL res = True_UnmapViewOfFile(info.fileMapMem);
+			//DEBUG_LOG_TRUE(L"INTERNAL UnmapViewOfFile", L"0x%llx (%ls) (%ls) -> %s", uintptr_t(info.fileMapMem), info.name, info.originalName, ToString(res));
 		}
 		if (info.trueFileMapHandle.IsValid())
 		{
-			DEBUG_LOG_TRUE(L"INTERNAL CloseHandle", L"%llu (%ls) (%ls)", info.trueFileMapHandle.ToU64(), info.name, info.originalName);
-
 			// TODO
+			//DEBUG_LOG_TRUE(L"INTERNAL CloseHandle", L"%llu (%ls) (%ls)", info.trueFileMapHandle.ToU64(), info.name, info.originalName);
 			//CloseHandle(info.trueFileMapHandle);
 		}
 
@@ -530,8 +532,6 @@ void OnModuleLoaded(HMODULE moduleHandle, const StringView& name);
 thread_local const wchar_t* t_renameFileNewName;
 thread_local const wchar_t* t_createFileFileName;
 
-const wchar_t* ToString(BOOL b) { return b ? L"Success" : L"Error"; }
-
 #include "UbaDetoursFunctionsMiMalloc.inl"
 #include "UbaDetoursFunctionsNtDll.inl"
 #include "UbaDetoursFunctionsKernelBase.inl"
diff --git a/Engine/Source/Programs/UnrealBuildAccelerator/Test/Private/UbaTestBasics.cpp b/Engine/Source/Programs/UnrealBuildAccelerator/Test/Private/UbaTestBasics.cpp
--- a/Engine/Source/Programs/UnrealBuildAccelerator/Test/Private/UbaTestBasics.cpp
+++ b/Engine/Source/Programs/UnrealBuildAccelerator/Test/Private/UbaTestBasics.cpp
@@ -184,8 +184,22 @@ namespace uba
 #if PLATFORM_WINDOWS
 		if (!FixPath2(TC("\"e:\\temp\""), ~0u, workingDir, TStrlen(workingDir), buffer, sizeof_array(buffer), &lengthResult))
 			return logger.Error(TC("FixPath2 (1) failed"));
-		if (!FixPath2(TC("\\??\\c:\\windows:::"), 14, workingDir, TStrlen(workingDir), buffer, sizeof_array(buffer), &lengthResult) || lengthResult != 10)
+		if (!FixPath2(TC("\\??\\c:\\windows:::"), 14, workingDir, TStrlen(workingDir), buffer, sizeof_array(buffer), &lengthResult) || lengthResult != 10 || TStrlen(buffer) != 10)
 			return logger.Error(TC("FixPath2 (2) failed"));
+		if (!FixPath2(TC("Z:/UEVFS/Root/Engine/Binaries/Win64/UnrealEditor-InterchangeCore-Win64-Debug.pdb::"), 80, workingDir, TStrlen(workingDir), buffer, sizeof_array(buffer), &lengthResult) || lengthResult != 80 || TStrlen(buffer) != 80)
+			return logger.Error(TC("FixPath2 (2) failed"));
+
+		#if 0
+		const tchar longPath[] = TC("c:\\Program Files (x86)\\Microsoft Visual Studio\\Shared\\Entity Framework Tools\\NuGet Packages");
+		tchar shortPath[MAX_PATH];
+		DWORD len = GetShortPathNameW(longPath, shortPath, MAX_PATH);
+		if (len == 0 || len >= sizeof_array(longPath))
+			return logger.Error(TC("GetShortPathNameW failed"));
+		if (!FixPath2(shortPath, len, workingDir, TStrlen(workingDir), buffer, sizeof_array(buffer), &lengthResult) || lengthResult != sizeof_array(longPath) - 1)
+			return logger.Error(TC("FixPath2 for short path failed"));
+		#endif
+
+
 #else
 
 		if (!TestPath(TC("/..")))
