[UBA]
* Added debug code to be able to write out writes to a file to analyze write throughput
* Removed include
* Fixed spelling in UbaHost.toml

diff --git a/Engine/Source/Programs/UnrealBuildAccelerator/Agent/Private/UbaAgent.cpp b/Engine/Source/Programs/UnrealBuildAccelerator/Agent/Private/UbaAgent.cpp
--- a/Engine/Source/Programs/UnrealBuildAccelerator/Agent/Private/UbaAgent.cpp
+++ b/Engine/Source/Programs/UnrealBuildAccelerator/Agent/Private/UbaAgent.cpp
@@ -9,6 +9,7 @@
 #include "UbaNetworkBackendTcp.h"
 #include "UbaNetworkMessage.h"
 #include "UbaNetworkServer.h"
+#include "UbaProcessStartInfo.h"
 #include "UbaSessionClient.h"
 #include "UbaStorageClient.h"
 #include "UbaStorageProxy.h"
diff --git a/Engine/Source/Programs/UnrealBuildAccelerator/Common/Private/UbaSession.cpp b/Engine/Source/Programs/UnrealBuildAccelerator/Common/Private/UbaSession.cpp
--- a/Engine/Source/Programs/UnrealBuildAccelerator/Common/Private/UbaSession.cpp
+++ b/Engine/Source/Programs/UnrealBuildAccelerator/Common/Private/UbaSession.cpp
@@ -31,6 +31,7 @@ extern char **environ;
 #endif
 
 #define UBA_DEBUG_TRACK_DIR 0 // UBA_DEBUG_LOGGER
+#define UBA_DEBUG_TRACK_WRITES 0 // UBA_DEBUG_LOGGER
 
 //////////////////////////////////////////////////////////////////////////////
 
@@ -233,7 +234,7 @@ namespace uba
 		table.GetValueAsBool(extractObjFilesSymbols, TC("ExtractObjFilesSymbols"));
 		table.GetValueAsBool(useFakeVolumeSerial, TC("UseFakeVolumeSerial"));
 		table.GetValueAsBool(allowLinkDependencyCrawler, TC("AllowLinkDependencyCrawler"));
-		table.GetValueAsU32(traceReserveSizeMb, TC("TraceReserveSizeMb"));
+		table.GetValueAsU32(traceReserveSizeMb, TC("TraceReserveSizeMB"));
 		table.GetValueAsU32(writeFilesBottleneck, TC("WriteFilesBottleneck"));
 		table.GetValueAsU32(writeFilesFileMapMaxMb, TC("WriteFilesFileMapMaxMB"));
 		table.GetValueAsU32(writeFilesNoBufferingMinMb, TC("WriteFilesNoBufferingMinMB"));
@@ -272,6 +273,20 @@ namespace uba
 					m_sharedMemory.UnmapView(memoryHandle, mem);
 			});
 
+		#if UBA_DEBUG_TRACK_WRITES
+		static u64 firstFile = [](Logger* logger)
+			{
+				logger->Info(TC("[%5s %8s] %6s %9s %-13s File"), TC("Wall"), TC("Total"), TC("Time"), TC("Bytes"), TC("Mode"));
+				return GetTime();
+			}(m_debugLogger);
+		u64 start = GetTime(); 
+		StringBuffer<256> writeMode;
+		static Atomic<u64> allFilesSize;
+		#define UBA_WRITE_MODE_BUFFER &writeMode
+		#else
+		#define UBA_WRITE_MODE_BUFFER nullptr
+		#endif
+
 		if (shouldWriteToDisk)
 		{
 			RootsHandle rootsHandle = process.GetStartInfo().rootsHandle;
@@ -300,6 +315,10 @@ namespace uba
 			{
 				UBA_ASSERT(!shouldWritePlaceholder);
 
+				#if UBA_DEBUG_TRACK_WRITES
+				writeMode = TCV("Compressed");
+				#endif
+
 				Storage::WriteResult res;
 				CompressedFileHeader header { CalculateCasKey(mem, fileSize, true, &m_workManager, file.name.c_str()) };
 
@@ -318,6 +337,10 @@ namespace uba
 
 				if (shouldDevirtualize)
 				{
+					#if UBA_DEBUG_TRACK_WRITES
+					writeMode = TCV("Devirtualize");
+					#endif
+
 					UBA_ASSERT(!shouldWritePlaceholder);
 
 					// Need to turn paths back into local paths
@@ -342,7 +365,7 @@ namespace uba
 						shouldKeepInMemory = true; // Must keep in memory in case we want to use the cache
 					}
 
-					if (!WriteMemoryToDisk(destinationFile, mem, writeSize))
+					if (!WriteMemoryToDisk(destinationFile, mem, writeSize, UBA_WRITE_MODE_BUFFER))
 						return false;
 					writtenSize = writeSize;
 				}
@@ -361,6 +384,10 @@ namespace uba
 		}
 		else
 		{
+			#if UBA_DEBUG_TRACK_WRITES
+			writeMode = TCV("Skipped");
+			#endif
+
 			// Delete existing file to make sure it is not picked up (since it is out of date)
 			uba::DeleteFileW(file.name.c_str());
 
@@ -461,10 +488,15 @@ namespace uba
 		if (shouldWriteToDisk)
 			TraceWrittenFile(process.m_id, file.name, writtenSize);
 
+		#if UBA_DEBUG_TRACK_WRITES
+		allFilesSize += fileSize;
+		m_debugLogger->Info(TC("[%5s %8s] %6s %9s %-13s %s"), TimeToText(start-firstFile), BytesToText(allFilesSize).str, TimeToText(GetTime() - start).str, BytesToText(fileSize).str, writeMode.data, StringView(file.name).GetFileName().data);
+		#endif
+		#undef UBA_WRITE_MODE_BUFFER
 		return true;
 	}
 
-	bool Session::WriteMemoryToDisk(FileAccessor& destinationFile, const void* fileMem, u64 fileSize)
+	bool Session::WriteMemoryToDisk(FileAccessor& destinationFile, const void* fileMem, u64 fileSize, StringBufferBase* outWriteMode)
 	{
 		// Seems like best combo (for windows at least) is to use writes with overlap and max 16 at the same time.
 		// On one machine we get twice as fast without overlap if no bottleneck. On another machine (ntfs compression on) we get twice as slow without overlap
@@ -472,6 +504,18 @@ namespace uba
 		bool useFileMapForWrite = fileSize && fileSize <= m_writeFilesFileMapMax; // ::CreateFileMappingW does not work for zero-length files
 		bool useOverlap = !useFileMapForWrite && fileSize >= m_writeFilesNoBufferingMin && EventIsNative;//fileSize > 2*1024*1024;
 
+		#if UBA_DEBUG_TRACK_WRITES
+		if (outWriteMode)
+		{
+			if (useOverlap)
+				outWriteMode->Append(TCV("NoBuffering"));
+			else if (useFileMapForWrite)
+				outWriteMode->Append(TCV("MemoryMap"));
+			else
+				outWriteMode->Append(TCV("Normal"));
+		}
+		#endif
+
 		u32 attributes = DefaultAttributes();
 		if (useOverlap)
 			attributes |= (FILE_FLAG_OVERLAPPED | FILE_FLAG_NO_BUFFERING);
@@ -4031,7 +4075,7 @@ namespace uba
 	{
 		//m_logger.Info(TC("TRYING TO WRITE OUTPUT %s (%s)"), filePath, KeyToString(fileNameKey).data);
 		FileAccessor destinationFile(m_logger, filePath.data);
-		if (!GetFileMemory([&](const void* fileMem, u64 fileSize) { return WriteMemoryToDisk(destinationFile, fileMem, fileSize); }, fileNameKey, filePath, deleteInternalMapping))
+		if (!GetFileMemory([&](const void* fileMem, u64 fileSize) { return WriteMemoryToDisk(destinationFile, fileMem, fileSize, nullptr); }, fileNameKey, filePath, deleteInternalMapping))
 			return false;
 		if (!destinationFile.Close())
 			return false;
diff --git a/Engine/Source/Programs/UnrealBuildAccelerator/Common/Public/UbaSession.h b/Engine/Source/Programs/UnrealBuildAccelerator/Common/Public/UbaSession.h
--- a/Engine/Source/Programs/UnrealBuildAccelerator/Common/Public/UbaSession.h
+++ b/Engine/Source/Programs/UnrealBuildAccelerator/Common/Public/UbaSession.h
@@ -9,7 +9,6 @@
 #include "UbaDirectoryTable.h"
 #include "UbaFileMappingBuffer.h"
 #include "UbaProcessHandle.h"
-#include "UbaProcessStartInfo.h"
 #include "UbaRootPaths.h"
 #include "UbaSessionCreateInfo.h"
 #include "UbaSharedMemoryView.h"
@@ -17,7 +16,7 @@
 #include "UbaTrace.h"
 #include "UbaThread.h"
 
-#define UBA_DEBUG_TRACK_MAPPING UBA_DEBUG_LOGGER
+#define UBA_DEBUG_TRACK_MAPPING 0//UBA_DEBUG_LOGGER
 
 namespace uba
 {
@@ -56,6 +55,7 @@ namespace uba
 	struct ListDirectoryResponse;
 	struct WrittenFile;
 	struct NextProcessInfo;
+	using RootsHandle = u64;
 
 	class Session
 	{
@@ -206,7 +206,7 @@ namespace uba
 
 		void RemoveWrittenFile(ProcessImpl& process, const StringKey& fileKey);
 		bool WriteFileToDisk(ProcessImpl& process, WrittenFile& file);
-		bool WriteMemoryToDisk(FileAccessor& destinationFile, const void* fileMem, u64 fileSize);
+		bool WriteMemoryToDisk(FileAccessor& destinationFile, const void* fileMem, u64 fileSize, StringBufferBase* outWriteMode);
 		bool GetFileMemory(const Function<bool(const void*, u64)>& func, const StringKey& fileNameKey, StringView filePath, bool deleteInternalMapping);
 		void AddEnvironmentVariableNoLock(StringView key, StringView value);
 		bool WriteDirectoryEntriesInternal(DirectoryTable::Directory& dir, const StringKey& dirKey, StringView dirPath, bool isRefresh, u32& outTableOffset);
diff --git a/Engine/Source/Programs/UnrealBuildAccelerator/Test/Private/UbaTestCache.cpp b/Engine/Source/Programs/UnrealBuildAccelerator/Test/Private/UbaTestCache.cpp
--- a/Engine/Source/Programs/UnrealBuildAccelerator/Test/Private/UbaTestCache.cpp
+++ b/Engine/Source/Programs/UnrealBuildAccelerator/Test/Private/UbaTestCache.cpp
@@ -3,13 +3,13 @@
 #include "UbaCacheClient.h"
 #include "UbaCacheBucket.h"
 #include "UbaCacheServer.h"
-#include "UbaConfig.h"
 #include "UbaCompactTables.h"
 #include "UbaFileAccessor.h"
 #include "UbaHashMap.h"
 #include "UbaNetworkBackendTcp.h"
 #include "UbaNetworkClient.h"
 #include "UbaNetworkServer.h"
+#include "UbaProcessStartInfo.h"
 #include "UbaRootPaths.h"
 #include "UbaSessionServer.h"
 #include "UbaStorageServer.h"
