[UBA]
* Fixed bug in detoured SetFileAttributes

diff --git a/Engine/Source/Programs/UnrealBuildAccelerator/Detours/Private/Windows/UbaDetoursFunctionsKernelBase.inl b/Engine/Source/Programs/UnrealBuildAccelerator/Detours/Private/Windows/UbaDetoursFunctionsKernelBase.inl
--- a/Engine/Source/Programs/UnrealBuildAccelerator/Detours/Private/Windows/UbaDetoursFunctionsKernelBase.inl
+++ b/Engine/Source/Programs/UnrealBuildAccelerator/Detours/Private/Windows/UbaDetoursFunctionsKernelBase.inl
@@ -874,7 +874,7 @@ BOOL Detoured_SetFileAttributesW(LPCWSTR lpFileName, DWORD dwFileAttributes)
 
 	if (isMapped || g_rules->IsThrowAway(fixedFile, g_runningRemote))
 	{
-		DEBUG_LOG_DETOURED(L"SetFileAttributesW", L"(IGNORED) (%ls) %u", lpFileName, dwFileAttributes);
+		DEBUG_LOG_DETOURED(L"SetFileAttributesW", L"(IGNORED) (%ls) %u -> %s", lpFileName, dwFileAttributes, ToString(true));
 		SetLastError(ERROR_SUCCESS);
 		return true;
 	}
@@ -888,10 +888,12 @@ BOOL Detoured_SetFileAttributesW(LPCWSTR lpFileName, DWORD dwFileAttributes)
 			DirectoryTable::EntryInformation info;
 			g_directoryTable.GetEntryInformation(info, dirTableOffset);
 			UBA_ASSERTF(info.attributes == dwFileAttributes, TC("Not implemented - Remote process wants to SetFileAttributes on file %s"), fixedFile.data);
-			DEBUG_LOG_TRUE(L"SetFileAttributesW", L"(IGNORED_REMOTE) (%ls) %u", lpFileName, dwFileAttributes);
+			DEBUG_LOG_DETOURED(L"SetFileAttributesW", L"(IGNORED_REMOTE) (%ls) %u", lpFileName, dwFileAttributes);
 			return true;
 		}
-		UBA_ASSERTF(false, L"SetFileAttributesW from remote is not implemented");
+		DEBUG_LOG_DETOURED(L"SetFileAttributesW", L"(NotFound) (%ls) %u -> %s", lpFileName, dwFileAttributes, ToString(false));
+		SetLastError(ERROR_FILE_NOT_FOUND);
+		return FALSE;
 	}
 
 	DEBUG_LOG_TRUE(L"SetFileAttributesW", L"(SUPPRESS)(%ls) %u", lpFileName, dwFileAttributes);
@@ -4240,6 +4242,7 @@ BOOL Detoured_CallNamedPipeW(LPCWSTR lpNamedPipeName, LPVOID lpInBuffer, DWORD n
 BOOL Detoured_PeekNamedPipe(HANDLE hNamedPipe, LPVOID lpBuffer, DWORD nBufferSize, LPDWORD lpBytesRead, LPDWORD lpTotalBytesAvail, LPDWORD lpBytesLeftThisMessage)
 {
 	UBA_ASSERT(!isDetouredHandle(hNamedPipe));
+	DEBUG_LOG_TRUE(L"PeekNamedPipe", L"");
 	return True_PeekNamedPipe(hNamedPipe, lpBuffer, nBufferSize, lpBytesRead, lpTotalBytesAvail, lpBytesLeftThisMessage);
 }
 
