[UBA]
* Added workaround in GetSystemInfo when running on wine since it will break if there are more than 64 cores in the system
* Moved some logic from GetQueuedCompletionStatusEx to NtRemoveIoCompletionEx
* Added lots of more detoured functions in debug just to be able to log and check for bad parameters

diff --git a/Engine/Source/Programs/UnrealBuildAccelerator/Core/Public/UbaWinTypes.h b/Engine/Source/Programs/UnrealBuildAccelerator/Core/Public/UbaWinTypes.h
--- a/Engine/Source/Programs/UnrealBuildAccelerator/Core/Public/UbaWinTypes.h
+++ b/Engine/Source/Programs/UnrealBuildAccelerator/Core/Public/UbaWinTypes.h
@@ -224,4 +224,15 @@ struct RTL_RELATIVE_NAME_U
 	UNICODE_STRING RelativeName;
 	HANDLE ContainingDirectory;
 	PVOID CurDirRef;
-};
\ No newline at end of file
+};
+
+struct FILE_IO_COMPLETION_INFORMATION
+{
+	PVOID KeyContext;
+	PVOID ApcContext;
+	IO_STATUS_BLOCK IoStatusBlock;
+};
+
+struct USER_THREAD_START_ROUTINE;
+
+struct PS_ATTRIBUTE_LIST;
\ No newline at end of file
diff --git a/Engine/Source/Programs/UnrealBuildAccelerator/Detours/Private/Windows/UbaDetoursFunctionsKernelBase.inl b/Engine/Source/Programs/UnrealBuildAccelerator/Detours/Private/Windows/UbaDetoursFunctionsKernelBase.inl
--- a/Engine/Source/Programs/UnrealBuildAccelerator/Detours/Private/Windows/UbaDetoursFunctionsKernelBase.inl
+++ b/Engine/Source/Programs/UnrealBuildAccelerator/Detours/Private/Windows/UbaDetoursFunctionsKernelBase.inl
@@ -96,6 +96,17 @@ LPWSTR Detoured_GetCommandLineW()
 	return str;
 }
 
+VOID Detoured_GetSystemInfo(LPSYSTEM_INFO lpSystemInfo)
+{
+	DETOURED_CALL(GetSystemInfo);
+	DEBUG_LOG_TRUE(L"GetSystemInfo", L"");
+	True_GetSystemInfo(lpSystemInfo);
+
+	// This is a limitation in wine. There is a MAXIMUM_PROCESSORS preventing it from creating all registry keys under HARDWARE\DESCRIPTION\System\CentralProcessor
+	if (g_isRunningWine && lpSystemInfo->dwNumberOfProcessors > 64)
+		lpSystemInfo->dwNumberOfProcessors = 64;
+}
+
 DWORD Detoured_GetCurrentDirectoryW(DWORD nBufferLength, LPWSTR lpBuffer)
 {
 	DETOURED_CALL(GetCurrentDirectoryW);
@@ -615,17 +626,6 @@ LPVOID Detoured_VirtualAlloc(LPVOID lpAddress, SIZE_T dwSize, DWORD flAllocation
 	return nullptr;
 }
 
-BOOL Detoured_GetQueuedCompletionStatusEx(HANDLE CompletionPort, LPOVERLAPPED_ENTRY lpCompletionPortEntries, ULONG ulCount, PULONG ulNumEntriesRemoved, DWORD dwMilliseconds, BOOL fAlertable)
-{
-	DETOURED_CALL(GetQueuedCompletionStatusEx);
-	UBA_ASSERT(!isDetouredHandle(CompletionPort));
-	DEBUG_LOG_TRUE(L"GetQueuedCompletionStatusEx", L"%llu (Timeout: %ums)", u64(CompletionPort), dwMilliseconds);
-	bool res = True_GetQueuedCompletionStatusEx(CompletionPort, lpCompletionPortEntries, ulCount, ulNumEntriesRemoved, dwMilliseconds, fAlertable);
-	if (res)
-		Rpc_UpdateTables(); // This is a bit ugly but we know this is how msbuild worker nodes sync with each other..
-	return res;
-}
-
 BOOL Detoured_PathSearchAndQualifyW(LPCWSTR pszPath, LPWSTR pszBuf, UINT cchBuf)
 {
 	FixPath2(pszPath, TStrlen(pszPath), g_virtualWorkingDir.data, g_virtualWorkingDir.count, pszBuf, cchBuf, nullptr);
@@ -2784,7 +2784,7 @@ HMODULE Shared_LoadLibrary(LPCWSTR lpLibFileName, HANDLE hFile, DWORD dwFlags, b
 	StringBuffer<> path;
 	path.Append(lpLibFileName).FixPathSeparators();
 
-	bool detourDll = path.EndsWith(L".exe") || path.EndsWith(L".dll");
+	bool detourDll = path.EndsWith(L".exe") || path.EndsWith(L".dll") || path.EndsWith(L".node");
 	if (detourDll && g_systemRoot.count && path.StartsWith(g_systemRoot.data))
 		detourDll = g_runningRemote && GetFileAttributes(path.data) == INVALID_FILE_ATTRIBUTES; // It might be that remote machine actually doesn't have the file in system32. then we need to detour
 
@@ -3965,6 +3965,27 @@ BOOL Detoured_DeleteFileA(LPCSTR lpFileName)
 	return res;
 }
 
+DWORD Detoured_GetShortPathNameW(LPCWSTR lpszLongPath, LPWSTR lpszShortPath, DWORD cchBuffer)
+{
+	UBA_ASSERT(!g_runningRemote);
+	DEBUG_LOG_TRUE(L"GetShortPathNameW", L"");
+	return True_GetShortPathNameW(lpszLongPath, lpszShortPath, cchBuffer);
+}
+
+BOOL Detoured_NeedCurrentDirectoryForExePathW(LPCWSTR ExeName)
+{
+	UBA_ASSERT(!g_runningRemote);
+	DEBUG_LOG_TRUE(L"NeedCurrentDirectoryForExePathW", L"");
+	return True_NeedCurrentDirectoryForExePathW(ExeName);
+}
+
+BOOL Detoured_ReadDirectoryChangesW(HANDLE hDirectory, LPVOID lpBuffer, DWORD nBufferLength, BOOL bWatchSubtree, DWORD dwNotifyFilter, LPDWORD lpBytesReturned, LPOVERLAPPED lpOverlapped, LPOVERLAPPED_COMPLETION_ROUTINE lpCompletionRoutine)
+{
+	UBA_ASSERT(!g_runningRemote && !isDetouredHandle(hDirectory));
+	DEBUG_LOG_TRUE(L"ReadDirectoryChangesW", L"");
+	return True_ReadDirectoryChangesW(hDirectory, lpBuffer, nBufferLength, bWatchSubtree, dwNotifyFilter, lpBytesReturned, lpOverlapped, lpCompletionRoutine);
+}
+
 BOOL Detoured_SetCurrentDirectoryA(LPCSTR lpPathName)
 {
 	DETOURED_CALL(SetCurrentDirectoryA);
@@ -4200,11 +4221,54 @@ HANDLE Detoured_CreateWaitableTimerExW(LPSECURITY_ATTRIBUTES lpTimerAttributes,
 HANDLE Detoured_CreateIoCompletionPort(HANDLE FileHandle, HANDLE ExistingCompletionPort, ULONG_PTR CompletionKey, DWORD NumberOfConcurrentThreads)
 {
 	DETOURED_CALL(CreateIoCompletionPort);
-	DEBUG_LOG_TRUE(L"CreateIoCompletionPort", L"%llu %llu", u64(FileHandle), u64(ExistingCompletionPort));
+	UBA_ASSERT(!isDetouredHandle(FileHandle));
 	HANDLE trueHandle = FileHandle;
 	if (isDetouredHandle(FileHandle))
 		trueHandle = asDetouredHandle(FileHandle).trueHandle;
-	return True_CreateIoCompletionPort(trueHandle, ExistingCompletionPort, CompletionKey, NumberOfConcurrentThreads);
+	HANDLE res = True_CreateIoCompletionPort(trueHandle, ExistingCompletionPort, CompletionKey, NumberOfConcurrentThreads);
+	DEBUG_LOG_TRUE(L"CreateIoCompletionPort", L"%llu %llu (%llu) -> %llu", u64(FileHandle), u64(ExistingCompletionPort), u64(CompletionKey), u64(res));
+	return res;
+}
+
+HANDLE Detoured_CreateRemoteThread(HANDLE hProcess, LPSECURITY_ATTRIBUTES lpThreadAttributes, SIZE_T dwStackSize, LPTHREAD_START_ROUTINE lpStartAddress, LPVOID lpParameter, DWORD dwCreationFlags, LPDWORD lpThreadId)
+{
+	DEBUG_LOG_TRUE(L"CreateRemoteThread", L"");
+	return True_CreateRemoteThread(hProcess, lpThreadAttributes, dwStackSize, lpStartAddress, lpParameter, dwCreationFlags, lpThreadId);
+}
+
+BOOL Detoured_PostQueuedCompletionStatus(HANDLE CompletionPort, DWORD dwNumberOfBytesTransferred, ULONG_PTR dwCompletionKey, LPOVERLAPPED lpOverlapped)
+{
+	UBA_ASSERT(!isDetouredHandle(CompletionPort));
+	DEBUG_LOG_TRUE(L"PostQueuedCompletionStatus", L"%llu", u64(dwCompletionKey));
+	return True_PostQueuedCompletionStatus(CompletionPort, dwNumberOfBytesTransferred, dwCompletionKey, lpOverlapped);
+}
+
+BOOL Detoured_GetQueuedCompletionStatusEx(HANDLE CompletionPort, LPOVERLAPPED_ENTRY lpCompletionPortEntries, ULONG ulCount, PULONG ulNumEntriesRemoved, DWORD dwMilliseconds, BOOL fAlertable)
+{
+	DETOURED_CALL(GetQueuedCompletionStatusEx);
+	//DEBUG_LOG_PROXY(L"GetQueuedCompletionStatusEx", L""); // Ends up in NtRemoveIoCompletionEx
+	return True_GetQueuedCompletionStatusEx(CompletionPort, lpCompletionPortEntries, ulCount, ulNumEntriesRemoved, dwMilliseconds, fAlertable);
+}
+
+BOOL Detoured_CancelIo(HANDLE hFile)
+{
+	UBA_ASSERT(!isDetouredHandle(hFile));
+	DEBUG_LOG_TRUE(L"CancelIo", L"");
+	return True_CancelIo(hFile);
+}
+
+BOOL Detoured_CancelIoEx(HANDLE hFile, LPOVERLAPPED lpOverlapped)
+{
+	UBA_ASSERT(!isDetouredHandle(hFile));
+	DEBUG_LOG_TRUE(L"CancelIoEx", L"");
+	return True_CancelIoEx(hFile, lpOverlapped);
+}
+
+BOOL Detoured_SetFileCompletionNotificationModes(HANDLE FileHandle, UCHAR Flags)
+{
+	UBA_ASSERT(!isDetouredHandle(FileHandle));
+	DEBUG_LOG_TRUE(L"SetFileCompletionNotificationModes", L"");
+	return True_SetFileCompletionNotificationModes(FileHandle, Flags);
 }
 
 BOOL Detoured_CreatePipe(PHANDLE hReadPipe, PHANDLE hWritePipe, LPSECURITY_ATTRIBUTES lpPipeAttributes, DWORD nSize)
@@ -4320,6 +4384,30 @@ BOOL Detoured_IsProcessorFeaturePresent(DWORD ProcessorFeature)
 	return res;
 }
 
+PVOID Detoured_VirtualAlloc2(HANDLE Process, PVOID BaseAddress, SIZE_T Size, ULONG AllocationType, ULONG PageProtection, MEM_EXTENDED_PARAMETER *ExtendedParameters, ULONG ParameterCount)
+{
+	UBA_ASSERT(!isDetouredHandle(Process));
+	auto res = True_VirtualAlloc2(Process, BaseAddress, Size, AllocationType, PageProtection, ExtendedParameters, ParameterCount);
+	//DEBUG_LOG_TRUE(L"VirtualAlloc2", L"-> 0x%llx", uintptr_t(res)); // Calling NtAllocateVirtualMemoryEx
+	return res;
+}
+
+PVOID Detoured_VirtualAlloc2FromApp(HANDLE Process, PVOID BaseAddress, SIZE_T Size, ULONG AllocationType, ULONG PageProtection, MEM_EXTENDED_PARAMETER *ExtendedParameters, ULONG ParameterCount)
+{
+	UBA_ASSERT(!isDetouredHandle(Process));
+	auto res = True_VirtualAlloc2FromApp(Process, BaseAddress, Size, AllocationType, PageProtection, ExtendedParameters, ParameterCount);
+	DEBUG_LOG_TRUE(L"VirtualAlloc2FromApp", L"-> 0x%llx", uintptr_t(res));
+	return res;
+}
+
+PVOID Detoured_MapViewOfFile3(HANDLE FileMapping, HANDLE Process, PVOID BaseAddress, ULONG64 Offset, SIZE_T ViewSize, ULONG AllocationType, ULONG PageProtection, MEM_EXTENDED_PARAMETER *ExtendedParameters, ULONG ParameterCount)
+{
+	UBA_ASSERT(!isDetouredHandle(FileMapping));
+	UBA_ASSERT(!isDetouredHandle(Process));
+	DEBUG_LOG_TRUE(L"MapViewOfFile3", L"");
+	return True_MapViewOfFile3(FileMapping, Process, BaseAddress, Offset, ViewSize, AllocationType, PageProtection, ExtendedParameters, ParameterCount);
+}
+
 BOOL Detoured_UnmapViewOfFile2(HANDLE Process, PVOID BaseAddress, ULONG UnmapFlags)
 {
 	DETOURED_CALL(UnmapViewOfFile2);
@@ -4351,13 +4439,6 @@ VOID Detoured_GetStartupInfoW(LPSTARTUPINFOW lpStartupInfo)
 	True_GetStartupInfoW(lpStartupInfo);
 }
 
-VOID Detoured_GetSystemInfo(LPSYSTEM_INFO lpSystemInfo)
-{
-	DETOURED_CALL(GetSystemInfo);
-	DEBUG_LOG_TRUE(L"GetSystemInfo", L"");
-	True_GetSystemInfo(lpSystemInfo);
-}
-
 VOID Detoured_GetSystemTimeAsFileTime(LPFILETIME lpSystemTimeAsFileTime)
 {
 	DETOURED_CALL(GetSystemTimeAsFileTime);
diff --git a/Engine/Source/Programs/UnrealBuildAccelerator/Detours/Private/Windows/UbaDetoursFunctionsNtDll.inl b/Engine/Source/Programs/UnrealBuildAccelerator/Detours/Private/Windows/UbaDetoursFunctionsNtDll.inl
--- a/Engine/Source/Programs/UnrealBuildAccelerator/Detours/Private/Windows/UbaDetoursFunctionsNtDll.inl
+++ b/Engine/Source/Programs/UnrealBuildAccelerator/Detours/Private/Windows/UbaDetoursFunctionsNtDll.inl
@@ -821,6 +821,15 @@ NTSTATUS Detoured_NtSetInformationObject(HANDLE ObjectHandle, OBJECT_INFORMATION
 	return res;
 }
 
+NTSTATUS Detoured_NtRemoveIoCompletionEx(HANDLE IoCompletionHandle, FILE_IO_COMPLETION_INFORMATION *info, ULONG count, ULONG *written, LARGE_INTEGER* timeout, BOOLEAN alertable)
+{
+	auto res = True_NtRemoveIoCompletionEx(IoCompletionHandle, info, count, written, timeout, alertable);
+	DEBUG_LOG_TRUE(L"NtRemoveIoCompletionEx", L"%llu (Count: %u Key: %llu Timeout: %lluns) -> %s (%u)", u64(IoCompletionHandle), count, (written && *written) ? u64((ULONG_PTR)info[0].KeyContext) : 0ull, timeout ? u64(timeout->QuadPart) : 0ull, ToString(res).data, *written);
+	if (res == STATUS_SUCCESS)
+		Rpc_UpdateTables(); // This is a bit ugly but we know this is how msbuild worker nodes sync with each other..
+	return res;
+}
+
 struct JobInfo { Vector<HANDLE> processes; bool shouldKill = false; };
 auto& g_jobsLock = *new ReaderWriterLock();
 auto& g_jobs = *new UnorderedMap<FileObject*, JobInfo>();
@@ -2464,6 +2473,23 @@ NTSTATUS Detoured_NtQueryInformationProcess(HANDLE ProcessHandle, PROCESSINFOCLA
 		ProcessHandle = asDetouredHandle(ProcessHandle).trueHandle;
 	}
 
+	if (g_isRunningWine)
+	{
+		if (ProcessInformationClass == (PROCESSINFOCLASS)49) // ProcessSubsystemInformation
+		{
+			if (ProcessInformationLength < sizeof(ULONG))
+			{
+				DEBUG_LOG_TRUE(L"NtQueryInformationProcess", L"(ProcessSubsystemInformation) %llu -> %s", uintptr_t(ProcessHandle), ToString(STATUS_INFO_LENGTH_MISMATCH).data);
+				return STATUS_INFO_LENGTH_MISMATCH;
+			}
+			*(ULONG*)ProcessInformation = 0; // SubsystemInformationTypeWin32 = 0
+			if (ReturnLength)
+				*ReturnLength = sizeof(ULONG);
+			DEBUG_LOG_TRUE(L"NtQueryInformationProcess", L"(ProcessSubsystemInformation) %llu -> %s", uintptr_t(ProcessHandle), ToString(STATUS_SUCCESS).data);
+			return STATUS_SUCCESS;
+		}
+	}
+
 	NTSTATUS res = True_NtQueryInformationProcess(ProcessHandle, ProcessInformationClass, ProcessInformation, ProcessInformationLength, ReturnLength);
 	DEBUG_LOG_TRUE(L"NtQueryInformationProcess", L"(class %u) %llu -> %ls", ProcessInformationClass, uintptr_t(ProcessHandle), ToString(res).data);
 	return res;
@@ -2579,28 +2605,67 @@ NTSTATUS Detoured_NtQuerySecurityObject(HANDLE Handle, SECURITY_INFORMATION Secu
 
 #if DETOURED_INCLUDE_DEBUG
 
+NTSTATUS Detoured_RtlCreateUserThread(HANDLE process, SECURITY_DESCRIPTOR* descr, BOOLEAN suspended, ULONG zero_bits, SIZE_T stack_reserve, SIZE_T stack_commit, void* start, void *param, HANDLE *handle_ptr, CLIENT_ID* id)
+{
+	DEBUG_LOG_TRUE(L"RtlCreateUserThread", L"");
+	auto res = True_RtlCreateUserThread(process, descr, suspended, zero_bits, stack_reserve, stack_commit, start, param, handle_ptr, id);
+	return res;
+}
+
+NTSTATUS Detoured_NtCreateThreadEx(PHANDLE ThreadHandle, ACCESS_MASK DesiredAccess, POBJECT_ATTRIBUTES ObjectAttributes, HANDLE ProcessHandle, USER_THREAD_START_ROUTINE* StartRoutine, PVOID Argument, ULONG CreateFlags, SIZE_T ZeroBits, SIZE_T StackSize, SIZE_T MaximumStackSize, PS_ATTRIBUTE_LIST* AttributeList)
+{
+	auto res = True_NtCreateThreadEx(ThreadHandle, DesiredAccess, ObjectAttributes, ProcessHandle, StartRoutine, Argument, CreateFlags, ZeroBits, StackSize, MaximumStackSize, AttributeList);
+	DEBUG_LOG_TRUE(L"NtCreateThreadEx", L"-> %s", ToString(res).data);
+	return res;
+}
+
 NTSTATUS Detoured_NtOpenJobObject(HANDLE* handle, ACCESS_MASK access, const OBJECT_ATTRIBUTES *attr)
 {
 	UBA_ASSERT(!isDetouredHandle(handle));
+	DEBUG_LOG_TRUE(L"NtOpenJobObject", L"");
 	return True_NtOpenJobObject(handle, access, attr);
 }
 
+NTSTATUS Detoured_NtAllocateVirtualMemoryEx(HANDLE ProcessHandle, PVOID *BaseAddress, SIZE_T *RegionSize, ULONG AllocationType, ULONG Protect, PMEM_EXTENDED_PARAMETER ExtendedParameters, ULONG ExtendedParameterCount)
+{
+	UBA_ASSERT(!isDetouredHandle(ProcessHandle));
+	//DEBUG_LOG_TRUE(L"NtAllocateVirtualMemoryEx", L"");
+	return True_NtAllocateVirtualMemoryEx(ProcessHandle, BaseAddress, RegionSize, AllocationType, Protect, ExtendedParameters, ExtendedParameterCount);
+}
+
+NTSTATUS Detoured_NtDeviceIoControlFile(HANDLE FileHandle, HANDLE Event, PIO_APC_ROUTINE ApcRoutine, PVOID ApcContext, PIO_STATUS_BLOCK IoStatusBlock, ULONG IoControlCode, PVOID InputBuffer, ULONG InputBufferLength, PVOID OutputBuffer, ULONG OutputBufferLength)
+{
+	UBA_ASSERT(!isDetouredHandle(FileHandle));
+	return True_NtDeviceIoControlFile(FileHandle, Event, ApcRoutine, ApcContext, IoStatusBlock, IoControlCode, InputBuffer, InputBufferLength, OutputBuffer, OutputBufferLength);
+}
+
 NTSTATUS Detoured_NtQueryInformationJobObject(HANDLE handle, JOBOBJECTINFOCLASS infoClass, void* info, ULONG len, ULONG* retLen)
 {
 	UBA_ASSERT(!isDetouredHandle(handle));
+	DEBUG_LOG_TRUE(L"NtQueryInformationJobObject", L"");
 	return True_NtQueryInformationJobObject(handle, infoClass, info, len, retLen);
 }
 
 NTSTATUS Detoured_NtCreateIoCompletion(PHANDLE IoCompletionHandle, ACCESS_MASK DesiredAccess, POBJECT_ATTRIBUTES ObjectAttributes, ULONG Count)
 {
 	UBA_ASSERT(!isDetouredHandle(*IoCompletionHandle));
-	return True_NtCreateIoCompletion(IoCompletionHandle, DesiredAccess, ObjectAttributes, Count);
+	auto res = True_NtCreateIoCompletion(IoCompletionHandle, DesiredAccess, ObjectAttributes, Count);
+	DEBUG_LOG_TRUE(L"NtCreateIoCompletion", L"%llu -> %s", IoCompletionHandle ? u64(*IoCompletionHandle) : 0ull, ToString(res).data);
+	return res;
+}
+
+NTSTATUS Detoured_NtRemoveIoCompletion(HANDLE IoCompletionHandle, PULONG_PTR CompletionKey, PULONG_PTR CompletionValue, PIO_STATUS_BLOCK IoStatusBlock, PLARGE_INTEGER Timeout)
+{
+	auto res = True_NtRemoveIoCompletion(IoCompletionHandle, CompletionKey, CompletionValue, IoStatusBlock, Timeout);
+	DEBUG_LOG_TRUE(L"NtRemoveIoCompletion", L"%llu -> %s", u64(IoCompletionHandle), ToString(res).data);
+	return res;
 }
 
 NTSTATUS NTAPI Detoured_NtFlushBuffersFileEx(HANDLE FileHandle, ULONG Flags, PVOID Parameters, ULONG ParametersSize, PIO_STATUS_BLOCK IoStatusBlock)
 {
 	DETOURED_CALL(NtFlushBuffersFileEx);
 	UBA_ASSERT(!isDetouredHandle(FileHandle));
+	DEBUG_LOG_TRUE(L"NtFlushBuffersFileEx", L"");
 	return True_NtFlushBuffersFileEx(FileHandle, Flags, Parameters, ParametersSize, IoStatusBlock);
 }
 
diff --git a/Engine/Source/Programs/UnrealBuildAccelerator/Detours/Private/Windows/UbaDetoursFunctionsWin.h b/Engine/Source/Programs/UnrealBuildAccelerator/Detours/Private/Windows/UbaDetoursFunctionsWin.h
--- a/Engine/Source/Programs/UnrealBuildAccelerator/Detours/Private/Windows/UbaDetoursFunctionsWin.h
+++ b/Engine/Source/Programs/UnrealBuildAccelerator/Detours/Private/Windows/UbaDetoursFunctionsWin.h
@@ -42,6 +42,7 @@
 
 #define DETOURED_FUNCTIONS_KERNELBASE \
 	DETOURED_FUNCTION(GetCommandLineW) \
+	DETOURED_FUNCTION(GetSystemInfo) \
 	DETOURED_FUNCTION(GetCurrentDirectoryW) \
 	DETOURED_FUNCTION(GetCurrentDirectoryA) \
 	DETOURED_FUNCTION(SetCurrentDirectoryW) \
@@ -117,7 +118,6 @@
 	DETOURED_FUNCTION(WriteConsoleW) \
 	DETOURED_FUNCTION(ReadConsoleW) \
 	DETOURED_FUNCTION(VirtualAlloc) \
-	DETOURED_FUNCTION(GetQueuedCompletionStatusEx) \
 	DETOURED_FUNCTION(PathSearchAndQualifyW) \
 	DETOURED_FUNCTIONS_KERNELBASE_DEBUG \
 
@@ -147,6 +147,7 @@
 	DETOURED_FUNCTION(NtQuerySecurityObject) \
 	DETOURED_FUNCTION(NtSetInformationFile) \
 	DETOURED_FUNCTION(NtSetInformationObject) \
+	DETOURED_FUNCTION(NtRemoveIoCompletionEx) \
 	DETOURED_FUNCTION(NtTerminateProcess) \
 	DETOURED_FUNCTION(NtCreateJobObject) \
 	DETOURED_FUNCTION(NtSetInformationJobObject) \
@@ -307,6 +308,9 @@
 	DETOURED_FUNCTION(ReplaceFileW) \
 	DETOURED_FUNCTION(CreateHardLinkA) \
 	DETOURED_FUNCTION(DeleteFileA) \
+	DETOURED_FUNCTION(GetShortPathNameW) \
+	DETOURED_FUNCTION(NeedCurrentDirectoryForExePathW) \
+	DETOURED_FUNCTION(ReadDirectoryChangesW) \
 	DETOURED_FUNCTION(SetCurrentDirectoryA) \
 	DETOURED_FUNCTION(CreateSymbolicLinkW) \
 	DETOURED_FUNCTION(CreateSymbolicLinkA) \
@@ -327,6 +331,12 @@
 	DETOURED_FUNCTION(CreateMutexExW) \
 	DETOURED_FUNCTION(CreateWaitableTimerExW) \
 	DETOURED_FUNCTION(CreateIoCompletionPort) \
+	DETOURED_FUNCTION(CreateRemoteThread) \
+	DETOURED_FUNCTION(CancelIo) \
+	DETOURED_FUNCTION(CancelIoEx) \
+	DETOURED_FUNCTION(SetFileCompletionNotificationModes) \
+	DETOURED_FUNCTION(PostQueuedCompletionStatus) \
+	DETOURED_FUNCTION(GetQueuedCompletionStatusEx) \
 	DETOURED_FUNCTION(CreatePipe) \
 	DETOURED_FUNCTION(SetHandleInformation) \
 	DETOURED_FUNCTION(CreateNamedPipeW) \
@@ -344,10 +354,12 @@
 	DETOURED_FUNCTION(OpenFileMappingA) \
 	DETOURED_FUNCTION(GetMappedFileNameW) \
 	DETOURED_FUNCTION(IsProcessorFeaturePresent) \
+	DETOURED_FUNCTION(VirtualAlloc2) \
+	DETOURED_FUNCTION(VirtualAlloc2FromApp) \
+	DETOURED_FUNCTION(MapViewOfFile3) \
 	DETOURED_FUNCTION(UnmapViewOfFile2) \
 	DETOURED_FUNCTION(GetModuleHandleW) \
 	DETOURED_FUNCTION(GetStartupInfoW) \
-	DETOURED_FUNCTION(GetSystemInfo) \
 	DETOURED_FUNCTION(GetSystemTimeAsFileTime) \
 	DETOURED_FUNCTION(GetLargePageMinimum) \
 	DETOURED_FUNCTION(UrlCreateFromPathW) \
@@ -377,9 +389,14 @@
 	DETOURED_FUNCTION(RtlGetFullPathName_U) \
 	DETOURED_FUNCTION(RtlGetFullPathName_UEx) \
 	DETOURED_FUNCTION(RtlDosPathNameToRelativeNtPathName_U) \
+	DETOURED_FUNCTION(RtlCreateUserThread) \
+	DETOURED_FUNCTION(NtCreateThreadEx) \
 	DETOURED_FUNCTION(NtOpenJobObject) \
+	DETOURED_FUNCTION(NtAllocateVirtualMemoryEx) \
+	DETOURED_FUNCTION(NtDeviceIoControlFile) \
 	DETOURED_FUNCTION(NtQueryInformationJobObject) \
 	DETOURED_FUNCTION(NtCreateIoCompletion) \
+	DETOURED_FUNCTION(NtRemoveIoCompletion) \
 	DETOURED_FUNCTION(NtFlushBuffersFileEx) \
 	DETOURED_FUNCTION(NtAlpcCreatePort) \
 	DETOURED_FUNCTION(NtAlpcConnectPort) \
@@ -451,6 +468,7 @@ extern "C" {
 	NTSTATUS NTAPI NtQueryFullAttributesFile(POBJECT_ATTRIBUTES ObjectAttributes, PVOID Attributes);
 	NTSTATUS NTAPI NtQueryInformationFile(HANDLE FileHandle, PIO_STATUS_BLOCK IoStatusBlock, PVOID FileInformation, ULONG Length, FILE_INFORMATION_CLASS FileInformationClass);
 	NTSTATUS NTAPI NtQueryDirectoryFile(HANDLE FileHandle, HANDLE Event, PIO_APC_ROUTINE ApcRoutine, PVOID ApcContext, PIO_STATUS_BLOCK IoStatusBlock, PVOID FileInformation, ULONG Length, FILE_INFORMATION_CLASS FileInformationClass, BOOLEAN ReturnSingleEntry, PUNICODE_STRING FileName, BOOLEAN RestartScan);
+	NTSTATUS NTAPI NtDeviceIoControlFile(HANDLE FileHandle, HANDLE Event, PIO_APC_ROUTINE ApcRoutine, PVOID ApcContext, PIO_STATUS_BLOCK IoStatusBlock, ULONG IoControlCode, PVOID InputBuffer, ULONG InputBufferLength, PVOID OutputBuffer, ULONG OutputBufferLength);
 	NTSTATUS NTAPI NtFsControlFile(HANDLE FileHandle, HANDLE Event, PIO_APC_ROUTINE ApcRoutine, PVOID ApcContext, PIO_STATUS_BLOCK IoStatusBlock, ULONG FsControlCode, PVOID InputBuffer, ULONG InputBufferLength, PVOID OutputBuffer, ULONG OutputBufferLength);
 	NTSTATUS NTAPI NtCopyFileChunk(HANDLE Source, HANDLE Dest, HANDLE Event, PIO_STATUS_BLOCK IoStatusBlock, ULONG Length, PULONG SourceOffset, PULONG DestOffset, PULONG SourceKey, PULONG DestKey, ULONG Flags);
 	NTSTATUS NTAPI NtFlushBuffersFileEx(HANDLE FileHandle, ULONG Flags, PVOID Parameters, ULONG ParametersSize, PIO_STATUS_BLOCK IoStatusBlock);
@@ -459,6 +477,8 @@ extern "C" {
 	NTSTATUS NTAPI NtTerminateProcess(HANDLE ProcessHandle, NTSTATUS ExitStatus);
 	NTSTATUS NTAPI NtSetInformationFile(HANDLE FileHandle, PIO_STATUS_BLOCK IoStatusBlock, PVOID FileInformation, ULONG Length, FILE_INFORMATION_CLASS FileInformationClass);
 	NTSTATUS NTAPI NtSetInformationObject(HANDLE ObjectHandle, OBJECT_INFORMATION_CLASS ObjectInformationClass, PVOID ObjectInformation, ULONG Length);
+	NTSTATUS NTAPI RtlCreateUserThread(HANDLE process, SECURITY_DESCRIPTOR* descr, BOOLEAN suspended, ULONG zero_bits, SIZE_T stack_reserve, SIZE_T stack_commit, void* start, void *param, HANDLE *handle_ptr, CLIENT_ID* id);
+	NTSTATUS NTAPI NtCreateThreadEx(PHANDLE ThreadHandle, ACCESS_MASK DesiredAccess, POBJECT_ATTRIBUTES ObjectAttributes, HANDLE ProcessHandle, USER_THREAD_START_ROUTINE* StartRoutine, PVOID Argument, ULONG CreateFlags, SIZE_T ZeroBits, SIZE_T StackSize, SIZE_T MaximumStackSize, PS_ATTRIBUTE_LIST* AttributeList);
 	NTSTATUS NTAPI NtCreateJobObject(HANDLE* handle, ACCESS_MASK access, const OBJECT_ATTRIBUTES* attr);
 	NTSTATUS NTAPI NtOpenJobObject(HANDLE* handle, ACCESS_MASK access, const OBJECT_ATTRIBUTES *attr);
 	NTSTATUS NTAPI NtQueryInformationJobObject(HANDLE handle, JOBOBJECTINFOCLASS infoClass, void* info, ULONG len, ULONG* retLen);
@@ -467,6 +487,9 @@ extern "C" {
 	NTSTATUS NTAPI NtTerminateJobObject(HANDLE handle, NTSTATUS status);
 	NTSTATUS NTAPI NtCreateSection(PHANDLE SectionHandle, ACCESS_MASK DesiredAccess, POBJECT_ATTRIBUTES ObjectAttributes, PLARGE_INTEGER MaximumSize, ULONG SectionPageProtection, ULONG AllocationAttributes, HANDLE FileHandle);
 	NTSTATUS NTAPI NtCreateIoCompletion(PHANDLE IoCompletionHandle, ACCESS_MASK DesiredAccess, POBJECT_ATTRIBUTES ObjectAttributes, ULONG Count);
+	NTSTATUS NTAPI NtRemoveIoCompletion(HANDLE IoCompletionHandle, PULONG_PTR CompletionKey, PULONG_PTR CompletionValue, PIO_STATUS_BLOCK IoStatusBlock, PLARGE_INTEGER Timeout);
+	NTSTATUS NTAPI NtRemoveIoCompletionEx(HANDLE IoCompletionHandle, FILE_IO_COMPLETION_INFORMATION* info, ULONG count, ULONG *written, LARGE_INTEGER *timeout, BOOLEAN alertable);
+	NTSTATUS NTAPI NtAllocateVirtualMemoryEx(HANDLE ProcessHandle, PVOID *BaseAddress, SIZE_T *RegionSize, ULONG AllocationType, ULONG Protect, PMEM_EXTENDED_PARAMETER ExtendedParameters, ULONG ExtendedParameterCount);
 	NTSTATUS NTAPI NtQuerySecurityObject(HANDLE Handle, SECURITY_INFORMATION SecurityInformation, PSECURITY_DESCRIPTOR SecurityDescriptor, ULONG Length, PULONG LengthNeeded);
 	NTSTATUS NTAPI NtAlpcCreatePort(PHANDLE PortHandle, POBJECT_ATTRIBUTES ObjectAttributes, PALPC_PORT_ATTRIBUTES PortAttributes);
 	NTSTATUS NTAPI NtAlpcConnectPort(PHANDLE PortHandle, PUNICODE_STRING PortName, POBJECT_ATTRIBUTES ObjectAttributes, PALPC_PORT_ATTRIBUTES PortAttributes, DWORD ConnectionFlags, PSID RequiredServerSid, PPORT_MESSAGE ConnectionMessage, PSIZE_T ConnectMessageSize, PALPC_MESSAGE_ATTRIBUTES OutMessageAttributes, PALPC_MESSAGE_ATTRIBUTES InMessageAttributes, PLARGE_INTEGER Timeout);
diff --git a/Engine/Source/Programs/UnrealBuildAccelerator/Detours/Private/Windows/UbaDetoursMainWin.cpp b/Engine/Source/Programs/UnrealBuildAccelerator/Detours/Private/Windows/UbaDetoursMainWin.cpp
--- a/Engine/Source/Programs/UnrealBuildAccelerator/Detours/Private/Windows/UbaDetoursMainWin.cpp
+++ b/Engine/Source/Programs/UnrealBuildAccelerator/Detours/Private/Windows/UbaDetoursMainWin.cpp
@@ -205,9 +205,14 @@ BOOL WINAPI DllMain(HINSTANCE hinst, DWORD dwReason, LPVOID reserved)
 	}
 	else if (dwReason == DLL_THREAD_ATTACH)
 	{
+		//DEBUG_LOG_TRUE(TC("DllMain"), TC("THREAD_ATTACH %u"), GetCurrentThreadId());
 		SetThreadErrorMode(GetThreadErrorMode() | SEM_FAILCRITICALERRORS | SEM_NOOPENFILEERRORBOX, nullptr);
 		SetThreadGroupAffinity(GetCurrentThread(), &GroupAffinity, NULL);
 	}
+	else if (dwReason == DLL_THREAD_DETACH)
+	{
+		//DEBUG_LOG_TRUE(TC("DllMain"), TC("THREAD_DETACH %u"), GetCurrentThreadId());
+	}
 	else if (dwReason == DLL_PROCESS_DETACH)
 	{
 		Deinit(startTime);
