[UBA]
* Moved DeviceIoControl logic to Nt* layer
* Removed asserts in GetShortPathName since it is using detoured calls under anyway

diff --git a/Engine/Source/Programs/UnrealBuildAccelerator/Detours/Private/Windows/UbaDetoursFunctionsKernelBase.inl b/Engine/Source/Programs/UnrealBuildAccelerator/Detours/Private/Windows/UbaDetoursFunctionsKernelBase.inl
--- a/Engine/Source/Programs/UnrealBuildAccelerator/Detours/Private/Windows/UbaDetoursFunctionsKernelBase.inl
+++ b/Engine/Source/Programs/UnrealBuildAccelerator/Detours/Private/Windows/UbaDetoursFunctionsKernelBase.inl
@@ -665,51 +665,6 @@ int WriteFileAccessViolation(PEXCEPTION_POINTERS ep)
 	return EXCEPTION_CONTINUE_SEARCH;
 }
 
-BOOL Detoured_DeviceIoControl(HANDLE hDevice, DWORD dwIoControlCode, LPVOID lpInBuffer, DWORD nInBufferSize, LPVOID lpOutBuffer, DWORD  nOutBufferSize, LPDWORD lpBytesReturned, LPOVERLAPPED lpOverlapped)
-{
-	HANDLE trueHandle = hDevice;
-	if (isDetouredHandle(hDevice))
-	{
-		DetouredHandle& h = asDetouredHandle(hDevice);
-		if (h.trueHandle == INVALID_HANDLE_VALUE)
-		{
-			if (dwIoControlCode == FSCTL_GET_REPARSE_POINT)
-			{
-				//UBA_ASSERTF(false, TC("FSCTL_GET_REPARSE_POINT"));
-				DEBUG_LOG_DETOURED(L"DeviceIoControl", L"(FSCTL_GET_REPARSE_POINT) %llu (%ls) -> ERROR_NOT_A_REPARSE_POINT", uintptr_t(hDevice), HandleToName(hDevice));
-				return ERROR_NOT_A_REPARSE_POINT;
-			}
-			else if (dwIoControlCode == FSCTL_SET_SPARSE)
-			{
-				UBA_ASSERTF(false, TC("FSCTL_SET_SPARSE"));
-			}
-			else if (dwIoControlCode == IOCTL_STORAGE_GET_DEVICE_NUMBER)
-			{
-				// This is used by link.exe to figure out if it can use some faster paths to track unique files
-				const wchar_t* originalName = h.fileObject->fileInfo->originalName;
-				UBA_ASSERT(originalName && originalName[0]);
-				DWORD deviceNumber = originalName[0] % 32u;
-
-				auto& number = *(STORAGE_DEVICE_NUMBER*)lpOutBuffer;
-				number.DeviceType = FILE_DEVICE_DISK;
-				number.DeviceNumber = deviceNumber;
-				number.PartitionNumber = 0;
-				if (lpBytesReturned)
-					*lpBytesReturned = sizeof(STORAGE_DEVICE_NUMBER);
-				DEBUG_LOG_DETOURED(L"DeviceIoControl", L"(IOCTL_STORAGE_GET_DEVICE_NUMBER) %llu (%ls) -> Success", uintptr_t(hDevice), HandleToName(hDevice));
-				return TRUE;
-			}
-			else
-			{
-				UBA_ASSERTF(false, TC("DeviceIoControl on detoured handle with code %u not handled (%s)"), dwIoControlCode, HandleToName(hDevice));
-			}
-		}
-		trueHandle = h.trueHandle;
-	}
-	DEBUG_LOG_TRUE(L"DeviceIoControl", L"%llu (%ls) -> Success", uintptr_t(hDevice), HandleToName(hDevice));
-	return True_DeviceIoControl(trueHandle, dwIoControlCode, lpInBuffer, nInBufferSize, lpOutBuffer, nOutBufferSize, lpBytesReturned, lpOverlapped);
-}
-
 BOOL Detoured_FlushFileBuffers(HANDLE hFile)
 {
 	DETOURED_CALL(FlushFileBuffers);
@@ -3678,6 +3633,12 @@ BOOL Detoured_GetThreadPreferredUILanguages(DWORD dwFlags, PULONG pulNumLanguage
 
 #if DETOURED_INCLUDE_DEBUG
 
+BOOL Detoured_DeviceIoControl(HANDLE hDevice, DWORD dwIoControlCode, LPVOID lpInBuffer, DWORD nInBufferSize, LPVOID lpOutBuffer, DWORD  nOutBufferSize, LPDWORD lpBytesReturned, LPOVERLAPPED lpOverlapped)
+{
+	DEBUG_LOG_PROXY(L"DeviceIoControl", L"");
+	return True_DeviceIoControl(hDevice, dwIoControlCode, lpInBuffer, nInBufferSize, lpOutBuffer, nOutBufferSize, lpBytesReturned, lpOverlapped);
+}
+
 BOOL Detoured_GetDiskFreeSpaceExA(LPCSTR lpDirectoryName, PULARGE_INTEGER lpFreeBytesAvailableToCaller, PULARGE_INTEGER lpTotalNumberOfBytes, PULARGE_INTEGER lpTotalNumberOfFreeBytes)
 {
 	DETOURED_CALL(GetDiskFreeSpaceExA);
@@ -3967,14 +3928,15 @@ BOOL Detoured_DeleteFileA(LPCSTR lpFileName)
 
 DWORD Detoured_GetShortPathNameW(LPCWSTR lpszLongPath, LPWSTR lpszShortPath, DWORD cchBuffer)
 {
-	UBA_ASSERT(!g_runningRemote);
+	// This is using FindFirstFileExW/FindClose so it is fine to let it through
+	DETOURED_CALL(GetShortPathNameW);
 	DEBUG_LOG_TRUE(L"GetShortPathNameW", L"");
 	return True_GetShortPathNameW(lpszLongPath, lpszShortPath, cchBuffer);
 }
 
 BOOL Detoured_NeedCurrentDirectoryForExePathW(LPCWSTR ExeName)
 {
-	UBA_ASSERT(!g_runningRemote);
+	//UBA_ASSERT(!g_runningRemote); // This is probably fine
 	DEBUG_LOG_TRUE(L"NeedCurrentDirectoryForExePathW", L"");
 	return True_NeedCurrentDirectoryForExePathW(ExeName);
 }
@@ -4080,6 +4042,7 @@ UINT Detoured_GetTempFileNameW(LPCWSTR lpPathName, LPCWSTR lpPrefixString, UINT
 
 DWORD Detoured_GetShortPathNameA(LPCSTR lpszLongPath, LPSTR lpszShortPath, DWORD cchBuffer)
 {
+	// This is using FindFirstFileExW/FindClose so it is fine to let it through
 	DETOURED_CALL(GetShortPathNameA);
 	DWORD res = True_GetShortPathNameA(lpszLongPath, lpszShortPath, cchBuffer);
 	DEBUG_LOG_TRUE(L"GetShortPathNameA", L"%hs (%u) -> %u (%hs)", lpszLongPath, cchBuffer, res, res ? lpszShortPath : "error");
diff --git a/Engine/Source/Programs/UnrealBuildAccelerator/Detours/Private/Windows/UbaDetoursFunctionsNtDll.inl b/Engine/Source/Programs/UnrealBuildAccelerator/Detours/Private/Windows/UbaDetoursFunctionsNtDll.inl
--- a/Engine/Source/Programs/UnrealBuildAccelerator/Detours/Private/Windows/UbaDetoursFunctionsNtDll.inl
+++ b/Engine/Source/Programs/UnrealBuildAccelerator/Detours/Private/Windows/UbaDetoursFunctionsNtDll.inl
@@ -2117,15 +2117,95 @@ NTSTATUS NTAPI Detoured_NtFsControlFile(HANDLE FileHandle, HANDLE Event, PIO_APC
 	{
 		auto& dh = asDetouredHandle(FileHandle);
 		trueHandle = dh.trueHandle;
-		//if (trueHandle == INVALID_HANDLE_VALUE)
-		//{
-		//	return STATUS_UNSUCCESSFUL;
-		//}
-		UBA_ASSERTF(trueHandle != INVALID_HANDLE_VALUE, L"NtFsControlFile code %u (%ls)", FsControlCode, HandleToName(FileHandle));
+
+		if (dh.type >= HandleType_StdOut)
+			trueHandle = True_GetStdHandle(STD_OUTPUT_HANDLE);
+		else if (dh.type >= HandleType_StdErr)
+			trueHandle = True_GetStdHandle(STD_ERROR_HANDLE);
+		else if (trueHandle == INVALID_HANDLE_VALUE)
+		{
+			if (FsControlCode == FSCTL_GET_REPARSE_POINT)
+			{
+				DEBUG_LOG_DETOURED(L"NtFsControlFile", L"(FSCTL_GET_REPARSE_POINT) %llu (%s) -> %s", uintptr_t(FileHandle), HandleToName(FileHandle), ToString(STATUS_NOT_A_REPARSE_POINT).data);
+				return STATUS_NOT_A_REPARSE_POINT;
+			}
+			else
+			{
+				UBA_ASSERTF(false, L"NtFsControlFile code 0x%x not handled (%s)", FsControlCode, HandleToName(FileHandle));
+			}
+		}
 	}
 	UBA_ASSERT(!isListDirectoryHandle(FileHandle));
 
-	return True_NtFsControlFile(trueHandle, Event, ApcRoutine, ApcContext, IoStatusBlock, FsControlCode, InputBuffer, InputBufferLength, OutputBuffer, OutputBufferLength);
+	auto res = True_NtFsControlFile(trueHandle, Event, ApcRoutine, ApcContext, IoStatusBlock, FsControlCode, InputBuffer, InputBufferLength, OutputBuffer, OutputBufferLength);
+	DEBUG_LOG_DETOURED(L"NtFsControlFile", L"(0x%x) %llu (%s) -> %s", FsControlCode, uintptr_t(FileHandle), HandleToName(FileHandle), ToString(res).data);
+	return res;
+}
+
+NTSTATUS Detoured_NtDeviceIoControlFile(HANDLE FileHandle, HANDLE Event, PIO_APC_ROUTINE ApcRoutine, PVOID ApcContext, PIO_STATUS_BLOCK IoStatusBlock, ULONG IoControlCode, PVOID InputBuffer, ULONG InputBufferLength, PVOID OutputBuffer, ULONG OutputBufferLength)
+{
+	HANDLE trueHandle = FileHandle;
+	if (isDetouredHandle(FileHandle))
+	{
+		DetouredHandle& dh = asDetouredHandle(FileHandle);
+		trueHandle = dh.trueHandle;
+
+		if (dh.type >= HandleType_StdOut)
+			trueHandle = True_GetStdHandle(STD_OUTPUT_HANDLE);
+		else if (dh.type >= HandleType_StdErr)
+			trueHandle = True_GetStdHandle(STD_ERROR_HANDLE);
+		else if (dh.trueHandle == INVALID_HANDLE_VALUE)
+		{
+			if (IoControlCode == FSCTL_GET_REPARSE_POINT)
+			{
+				DEBUG_LOG_DETOURED(L"NtDeviceIoControlFile", L"(FSCTL_GET_REPARSE_POINT) %llu (%ls) -> STATUS_NOT_A_REPARSE_POINT", uintptr_t(FileHandle), HandleToName(FileHandle));
+				return STATUS_NOT_A_REPARSE_POINT;
+			}
+			else if (IoControlCode == FSCTL_SET_SPARSE)
+			{
+				UBA_ASSERTF(false, TC("FSCTL_SET_SPARSE"));
+			}
+			else if (IoControlCode == IOCTL_STORAGE_GET_DEVICE_NUMBER)
+			{
+				const wchar_t* originalName = dh.fileObject->fileInfo->originalName;
+				UBA_ASSERT(originalName && originalName[0]);
+				DWORD deviceNumber = originalName[0] % 32u;
+
+				auto& number = *(STORAGE_DEVICE_NUMBER*)OutputBuffer;
+				number.DeviceType = FILE_DEVICE_DISK;
+				number.DeviceNumber = deviceNumber;
+				number.PartitionNumber = 0;
+				if (IoStatusBlock)
+					IoStatusBlock->Information = sizeof(STORAGE_DEVICE_NUMBER);
+				DEBUG_LOG_DETOURED(L"NtDeviceIoControlFile", L"(IOCTL_STORAGE_GET_DEVICE_NUMBER) %llu (%ls) -> Success", uintptr_t(FileHandle), HandleToName(FileHandle));
+				if (IoStatusBlock)
+					IoStatusBlock->Status = STATUS_SUCCESS;
+				return STATUS_SUCCESS;
+			}
+			else
+			{
+				UBA_ASSERTF(false, TC("NtDeviceIoControlFile on detoured handle with code %u not handled (%s)"), IoControlCode, HandleToName(FileHandle));
+			}
+		}
+	}
+	else if (isListDirectoryHandle(FileHandle))
+	{
+		if (IoControlCode == FSCTL_GET_REPARSE_POINT)
+		{
+			DEBUG_LOG_DETOURED(L"NtDeviceIoControlFile", L"(FSCTL_GET_REPARSE_POINT) %llu (%ls) -> STATUS_NOT_A_REPARSE_POINT", uintptr_t(FileHandle), HandleToName(FileHandle));
+			return STATUS_NOT_A_REPARSE_POINT;
+		}
+		if (IoControlCode == 0x504000) // IOCTL_CONDRV_GET_MODE
+		{
+			DEBUG_LOG_DETOURED(L"NtDeviceIoControlFile", L"(FSCTL_GET_REPARSE_POINT) %llu (%ls) -> STATUS_NOT_A_REPARSE_POINT", uintptr_t(FileHandle), HandleToName(FileHandle));
+			return STATUS_INVALID_DEVICE_REQUEST;
+		}
+		UBA_ASSERTF(false, TC("NtDeviceIoControlFile on detoured list handle with code %u not handled (%s)"), IoControlCode, HandleToName(FileHandle));
+	}
+
+	auto res = True_NtDeviceIoControlFile(trueHandle, Event, ApcRoutine, ApcContext, IoStatusBlock, IoControlCode, InputBuffer, InputBufferLength, OutputBuffer, OutputBufferLength);
+	DEBUG_LOG_TRUE(L"NtDeviceIoControlFile", L"%llu (%s) -> %s", uintptr_t(FileHandle), HandleToName(FileHandle), ToString(res).data);
+	return res;
 }
 
 NTSTATUS NTAPI Detoured_NtCopyFileChunk(HANDLE Source, HANDLE Dest, HANDLE Event, PIO_STATUS_BLOCK IoStatusBlock, ULONG Length, PULONG SourceOffset, PULONG DestOffset, PULONG SourceKey, PULONG DestKey, ULONG Flags)
@@ -2147,6 +2227,7 @@ NTSTATUS NTAPI Detoured_NtCopyFileChunk(HANDLE Source, HANDLE Dest, HANDLE Event
 	}
 	return True_NtCopyFileChunk(trueSourceHandle, trueDestHandle, Event, IoStatusBlock, Length, SourceOffset, DestOffset, SourceKey, DestKey, Flags);
 }
+
 NTSTATUS NTAPI Detoured_NtClose(HANDLE handle)
 {
 	DETOURED_CALL(NtClose);
@@ -2633,12 +2714,6 @@ NTSTATUS Detoured_NtAllocateVirtualMemoryEx(HANDLE ProcessHandle, PVOID *BaseAdd
 	return True_NtAllocateVirtualMemoryEx(ProcessHandle, BaseAddress, RegionSize, AllocationType, Protect, ExtendedParameters, ExtendedParameterCount);
 }
 
-NTSTATUS Detoured_NtDeviceIoControlFile(HANDLE FileHandle, HANDLE Event, PIO_APC_ROUTINE ApcRoutine, PVOID ApcContext, PIO_STATUS_BLOCK IoStatusBlock, ULONG IoControlCode, PVOID InputBuffer, ULONG InputBufferLength, PVOID OutputBuffer, ULONG OutputBufferLength)
-{
-	UBA_ASSERT(!isDetouredHandle(FileHandle));
-	return True_NtDeviceIoControlFile(FileHandle, Event, ApcRoutine, ApcContext, IoStatusBlock, IoControlCode, InputBuffer, InputBufferLength, OutputBuffer, OutputBufferLength);
-}
-
 NTSTATUS Detoured_NtQueryInformationJobObject(HANDLE handle, JOBOBJECTINFOCLASS infoClass, void* info, ULONG len, ULONG* retLen)
 {
 	UBA_ASSERT(!isDetouredHandle(handle));
diff --git a/Engine/Source/Programs/UnrealBuildAccelerator/Detours/Private/Windows/UbaDetoursFunctionsWin.h b/Engine/Source/Programs/UnrealBuildAccelerator/Detours/Private/Windows/UbaDetoursFunctionsWin.h
--- a/Engine/Source/Programs/UnrealBuildAccelerator/Detours/Private/Windows/UbaDetoursFunctionsWin.h
+++ b/Engine/Source/Programs/UnrealBuildAccelerator/Detours/Private/Windows/UbaDetoursFunctionsWin.h
@@ -55,7 +55,6 @@
 	DETOURED_FUNCTION(LockFileEx) \
 	DETOURED_FUNCTION(UnlockFile) \
 	DETOURED_FUNCTION(UnlockFileEx) \
-	DETOURED_FUNCTION(DeviceIoControl) \
 	DETOURED_FUNCTION(FlushFileBuffers) \
 	DETOURED_FUNCTION(SetFileTime) \
 	DETOURED_FUNCTION(GetFileType) \
@@ -137,6 +136,7 @@
 	DETOURED_FUNCTION(NtReadFile) \
 	DETOURED_FUNCTION(NtWriteFile) \
 	DETOURED_FUNCTION(NtFsControlFile) \
+	DETOURED_FUNCTION(NtDeviceIoControlFile) \
 	DETOURED_FUNCTION(NtCopyFileChunk) \
 	DETOURED_FUNCTION(NtQueryVolumeInformationFile) \
 	DETOURED_FUNCTION(NtQueryInformationFile) \
@@ -279,6 +279,7 @@
 	DETOURED_FUNCTION(ReadConsoleInputA) \
 	DETOURED_FUNCTION(GetLocaleInfoEx) \
 	DETOURED_FUNCTION(GetUserDefaultLocaleName) \
+	DETOURED_FUNCTION(DeviceIoControl) \
 	DETOURED_FUNCTION(GetDiskFreeSpaceExA) \
 	DETOURED_FUNCTION(GetLongPathNameA) \
 	DETOURED_FUNCTION(GetVolumePathNameA) \
@@ -393,7 +394,6 @@
 	DETOURED_FUNCTION(NtCreateThreadEx) \
 	DETOURED_FUNCTION(NtOpenJobObject) \
 	DETOURED_FUNCTION(NtAllocateVirtualMemoryEx) \
-	DETOURED_FUNCTION(NtDeviceIoControlFile) \
 	DETOURED_FUNCTION(NtQueryInformationJobObject) \
 	DETOURED_FUNCTION(NtCreateIoCompletion) \
 	DETOURED_FUNCTION(NtRemoveIoCompletion) \
