[UBA]
* Removed double logging of NtClose
* Fixed bug in test application for special case unit test

diff --git a/Engine/Source/Programs/UnrealBuildAccelerator/Detours/Private/Windows/UbaDetoursFunctionsNtDll.inl b/Engine/Source/Programs/UnrealBuildAccelerator/Detours/Private/Windows/UbaDetoursFunctionsNtDll.inl
--- a/Engine/Source/Programs/UnrealBuildAccelerator/Detours/Private/Windows/UbaDetoursFunctionsNtDll.inl
+++ b/Engine/Source/Programs/UnrealBuildAccelerator/Detours/Private/Windows/UbaDetoursFunctionsNtDll.inl
@@ -2319,7 +2319,7 @@ NTSTATUS NTAPI Detoured_NtClose(HANDLE handle)
 		return closeTrueHandle(false);
 	}
 
-	NTSTATUS res = closeTrueHandle(true);
+	NTSTATUS res = closeTrueHandle(false);
 
 	SharedMemoryHandle memoryHandle;
 	u64 memoryWritten = 0;
diff --git a/Engine/Source/Programs/UnrealBuildAccelerator/Test/Private/UbaTestSession.cpp b/Engine/Source/Programs/UnrealBuildAccelerator/Test/Private/UbaTestSession.cpp
--- a/Engine/Source/Programs/UnrealBuildAccelerator/Test/Private/UbaTestSession.cpp
+++ b/Engine/Source/Programs/UnrealBuildAccelerator/Test/Private/UbaTestSession.cpp
@@ -131,7 +131,7 @@ namespace uba
 		sessionServerInfo.rootDir = rootDir.data;
 		sessionServerInfo.useUniqueId = false;
 
-		#if UBA_DEBUG
+		#if UBA_DEBUG_LOG_ENABLED
 		sessionServerInfo.logToFile = true;
 		sessionServerInfo.remoteLogEnabled = true;
 		#endif
@@ -218,7 +218,7 @@ namespace uba
 		sessionClientInfo.useUniqueId = false;
 		//sessionClientInfo.allowKeepFilesInMemory = false;
 
-		#if UBA_DEBUG
+		#if UBA_DEBUG_LOG_ENABLED
 		sessionClientInfo.logToFile = true;
 		#endif
 
diff --git a/Engine/Source/Programs/UnrealBuildAccelerator/TestApp/Private/Windows/UbaTestAppWin.cpp b/Engine/Source/Programs/UnrealBuildAccelerator/TestApp/Private/Windows/UbaTestAppWin.cpp
--- a/Engine/Source/Programs/UnrealBuildAccelerator/TestApp/Private/Windows/UbaTestAppWin.cpp
+++ b/Engine/Source/Programs/UnrealBuildAccelerator/TestApp/Private/Windows/UbaTestAppWin.cpp
@@ -405,7 +405,7 @@ int wmain(int argc, wchar_t* argv[])
 		int read = ReadTestFile(buf, 32*2, L"SpecialFile1");
 		if (read != 2)
 			return LogError(L"Expected to read two bytes but got %i", read);
-		buf[read] = 0;
+		buf[1] = 0;
 
 		int readIndex = _wtoi(buf);
 		if (callIndex != readIndex)
@@ -431,7 +431,7 @@ int wmain(int argc, wchar_t* argv[])
 		int read = ReadTestFile(buf, 32*2, L"SpecialFile1");
 		if (read != 2)
 			return LogError(L"Expected to read two bytes but got %i", read);
-		buf[read] = 0;
+		buf[1] = 0;
 
 		int readIndex = _wtoi(buf);
 		if (callIndex != readIndex)
