[UBA]
* Added GetResults to scheduler that can show number of succeeded,failed and skipped processes

diff --git a/Engine/Source/Programs/UnrealBuildAccelerator/Common/Private/UbaScheduler.cpp b/Engine/Source/Programs/UnrealBuildAccelerator/Common/Private/UbaScheduler.cpp
--- a/Engine/Source/Programs/UnrealBuildAccelerator/Common/Private/UbaScheduler.cpp
+++ b/Engine/Source/Programs/UnrealBuildAccelerator/Common/Private/UbaScheduler.cpp
@@ -468,6 +468,20 @@ namespace uba
 		outQueued = m_queuedProcesses;
 	}
 
+	void Scheduler::GetResults(u32& outSucceeded, u32& outFailed, u32& outSkipped)
+	{
+		outSucceeded = 0;
+		outFailed = 0;
+		outSkipped = 0;
+		SCOPED_FUTEX_READ(m_processEntriesLock, lock);
+		for (auto& entry : m_processEntries)
+		{
+			outSucceeded += entry.runStatus == RunStatus_Success;
+			outFailed += entry.runStatus == RunStatus_Failed;
+			outSkipped += entry.runStatus == RunStatus_Skipped;
+		}
+	}
+
 	bool Scheduler::IsEmpty()
 	{
 		SCOPED_FUTEX_READ(m_processEntriesLock, lock);
diff --git a/Engine/Source/Programs/UnrealBuildAccelerator/Common/Public/UbaScheduler.h b/Engine/Source/Programs/UnrealBuildAccelerator/Common/Public/UbaScheduler.h
--- a/Engine/Source/Programs/UnrealBuildAccelerator/Common/Public/UbaScheduler.h
+++ b/Engine/Source/Programs/UnrealBuildAccelerator/Common/Public/UbaScheduler.h
@@ -96,6 +96,7 @@ namespace uba
 
 		u32 GetTotalProcesses();
 		void GetStats(u32& outQueued, u32& outActiveLocal, u32& outActiveRemote, u32& outFinished); // This is not "threadsafe".. which means that you are not guaranteed to get the exact correct state.
+		void GetResults(u32& outSucceeded, u32& outFailed, u32& outSkipped);
 
 		bool IsEmpty(); // Returns true if scheduler is entirely empty.. as in no processes are left in the system
 
