[UBA]
* Fixed detoured ReadConsole so it properly handles detoured handles

diff --git a/Engine/Source/Programs/UnrealBuildAccelerator/Detours/Private/Windows/UbaDetoursFunctionsKernelBase.inl b/Engine/Source/Programs/UnrealBuildAccelerator/Detours/Private/Windows/UbaDetoursFunctionsKernelBase.inl
--- a/Engine/Source/Programs/UnrealBuildAccelerator/Detours/Private/Windows/UbaDetoursFunctionsKernelBase.inl
+++ b/Engine/Source/Programs/UnrealBuildAccelerator/Detours/Private/Windows/UbaDetoursFunctionsKernelBase.inl
@@ -449,8 +449,16 @@ BOOL Detoured_ReadConsoleW(HANDLE hConsoleInput, LPVOID lpBuffer, DWORD nNumberO
 {
 	DETOURED_CALL(ReadConsoleW);
 	//DEBUG_LOG_DETOURED(L"ReadConsoleW", L"(%llu)", u64(hConsoleInput)); // Too much spam
-	UBA_ASSERT(!isDetouredHandle(hConsoleInput));
-	return True_ReadConsoleW(hConsoleInput, lpBuffer, nNumberOfCharsToRead, lpNumberOfCharsRead, pInputControl);
+	HANDLE trueHandle = hConsoleInput;
+	if (isDetouredHandle(hConsoleInput))
+	{
+		auto& dh = asDetouredHandle(hConsoleInput);
+		trueHandle = dh.trueHandle;
+		if (dh.type == HandleType_StdIn)
+			trueHandle = True_GetStdHandle(STD_INPUT_HANDLE);
+		UBA_ASSERTF(trueHandle != INVALID_HANDLE_VALUE, TC("ReadConsoleW got handle that can't be resolved %llu (%s)"), u64(hConsoleInput), HandleToName(hConsoleInput));
+	}
+	return True_ReadConsoleW(trueHandle, lpBuffer, nNumberOfCharsToRead, lpNumberOfCharsRead, pInputControl);
 }
 
 UINT Detoured_GetDriveTypeW(LPCWSTR lpRootPathName)
