[UBA]
* Minor opt in scheduler when searching for number of processes that can be scheduled remotely

diff --git a/Engine/Source/Programs/UnrealBuildAccelerator/Common/Private/UbaScheduler.cpp b/Engine/Source/Programs/UnrealBuildAccelerator/Common/Private/UbaScheduler.cpp
--- a/Engine/Source/Programs/UnrealBuildAccelerator/Common/Private/UbaScheduler.cpp
+++ b/Engine/Source/Programs/UnrealBuildAccelerator/Common/Private/UbaScheduler.cpp
@@ -519,9 +519,12 @@ namespace uba
 			return 0;
 
 		u32 count = 0;
+		auto processEntries = m_processEntries.data();
 		SCOPED_FUTEX_READ(m_processEntriesLock, lock);
-		for (auto& entry : m_processEntries)
+		for (u32 i=m_processEntriesRunStart, e=u32(m_processEntries.size()); i!=e; ++i)
 		{
+			auto& entry = processEntries[i];
+
 			if (!entry.canExecuteRemotely)
 				continue;
 			if (entry.runStatus != RunStatus_QueuedForRun || entry.cacheStatus != CacheStatus_Failed)
@@ -556,9 +559,12 @@ namespace uba
 		float crossArchitecture = 0;
 		float crossPlatform = 0;
 
+		auto processEntries = m_processEntries.data();
 		SCOPED_FUTEX_READ(m_processEntriesLock, lock);
-		for (auto& entry : m_processEntries)
+		for (u32 i=m_processEntriesRunStart, e=u32(m_processEntries.size()); i!=e; ++i)
 		{
+			auto& entry = processEntries[i];
+
 			if (!entry.canExecuteRemotely)
 				continue;
 			if (entry.runStatus != RunStatus_QueuedForRun || entry.cacheStatus != CacheStatus_Failed)
