[UBA]
* Fixed so path validation requires ':' to always be as second character in a valid path if it exists

diff --git a/Engine/Source/Programs/UnrealBuildAccelerator/Core/Public/UbaPathUtils.h b/Engine/Source/Programs/UnrealBuildAccelerator/Core/Public/UbaPathUtils.h
--- a/Engine/Source/Programs/UnrealBuildAccelerator/Core/Public/UbaPathUtils.h
+++ b/Engine/Source/Programs/UnrealBuildAccelerator/Core/Public/UbaPathUtils.h
@@ -126,7 +126,7 @@ namespace uba
 				hasDotDot = true;
 			else if (c == ':')
 			{
-				if (isAbsolute)
+				if ((write - buffer) != 1) // colon must be on position 1
 					return false;
 				isAbsolute = true;
 			}
