[UBA]
* Moved python caches (.pyc) to be under temp folder when running remotely. Probably faster to create the python cache than transferring them back and forth between helpers

diff --git a/Engine/Source/Programs/UnrealBuildAccelerator/Common/Private/UbaSessionClient.cpp b/Engine/Source/Programs/UnrealBuildAccelerator/Common/Private/UbaSessionClient.cpp
--- a/Engine/Source/Programs/UnrealBuildAccelerator/Common/Private/UbaSessionClient.cpp
+++ b/Engine/Source/Programs/UnrealBuildAccelerator/Common/Private/UbaSessionClient.cpp
@@ -1525,6 +1525,7 @@ namespace uba
 		AddEnvironmentVariableNoLock(TCV("TEMP"), m_tempPath);
 		AddEnvironmentVariableNoLock(TCV("TMP"), m_tempPath);
 		//AddEnvironmentVariableNoLock(TCV("WINEDEBUG"), TCV("+process"));
+		AddEnvironmentVariableNoLock(TCV("PYTHONPYCACHEPREFIX"), m_tempPath);
 		#else
 		AddEnvironmentVariableNoLock(TCV("TMPDIR"), m_tempPath);
 		#endif
