[UBA]
* Fixed so racing processes are not trying to write log files if log files are disabled

diff --git a/Engine/Source/Programs/UnrealBuildAccelerator/Common/Private/UbaSessionServer.cpp b/Engine/Source/Programs/UnrealBuildAccelerator/Common/Private/UbaSessionServer.cpp
--- a/Engine/Source/Programs/UnrealBuildAccelerator/Common/Private/UbaSessionServer.cpp
+++ b/Engine/Source/Programs/UnrealBuildAccelerator/Common/Private/UbaSessionServer.cpp
@@ -399,15 +399,18 @@ namespace uba
 			return {};
 
 		StringBuffer<> logFile;
-		logFile.Append(psi.logFile);
-		if (const tchar* dot = logFile.Last('.'))
+		if (*psi.logFile)
 		{
-			u64 offset = dot - logFile.data;
-			logFile.Resize(offset).Append(TCV("_racing")).Append(psi.logFile + offset);
+			logFile.Append(psi.logFile);
+			if (const tchar* dot = logFile.Last('.'))
+			{
+				u64 offset = dot - logFile.data;
+				logFile.Resize(offset).Append(TCV("_racing")).Append(psi.logFile + offset);
+			}
+			else
+				logFile.Append(TCV("_racing"));
+			psi.logFile = logFile.data;
 		}
-		else
-			logFile.Append(TCV("_racing"));
-		psi.logFile = logFile.data;
 
 		auto rpi = new RacingProcessExitInfo;
 		rpi->server = this;
