[UBA]
* Added SocketRecv2 which can take multiple buffers to receive into

diff --git a/Engine/Source/Programs/UnrealBuildAccelerator/Common/Private/UbaNetworkBackendTcp.cpp b/Engine/Source/Programs/UnrealBuildAccelerator/Common/Private/UbaNetworkBackendTcp.cpp
--- a/Engine/Source/Programs/UnrealBuildAccelerator/Common/Private/UbaNetworkBackendTcp.cpp
+++ b/Engine/Source/Programs/UnrealBuildAccelerator/Common/Private/UbaNetworkBackendTcp.cpp
@@ -1378,24 +1378,13 @@ namespace uba
 			}
 
 			int read;
-#if PLATFORM_WINDOWS
-			if (false)//!IsRunningWine())
+
 			{
-				WSABUF bufs[2];
-				bufs[0] = WSABUFINIT(buffer, recvLeft);
-				bufs[1] = WSABUFINIT(recvCache.bytes, sizeof(recvCache.bytes));
-				u32 bufCount = 2;
-
-				DWORD bytesReceived = 0;
-				DWORD flags = 0;
-				read = WSARecv(ToPlatformSocket(socket), bufs, bufCount, &bytesReceived, &flags, NULL, NULL);
-				if (read == 0)
-					read = int(bytesReceived);
-				else
-					read = SOCKET_ERROR;
+				RecvBuf bufs[2] = { { recvLeft, buffer }, { sizeof(recvCache.bytes), recvCache.bytes } };
+				read = SocketRecv2(socket, bufs, 2);
 				if (read > 0)
 				{
-					if (u32(read) >= recvLeft)
+					if (u32(read) > recvLeft)
 					{
 						recvCache.written = read - recvLeft;
 						return true;
@@ -1403,58 +1392,11 @@ namespace uba
 					recvLeft -= read;
 					buffer += read;
 
-					if (allowLess)
-					{
-						bufferLen = read;
-						return true;
-					}
-					continue;
-				}
-			}
-			else
-#endif
-			{
-				u8* writePos;
-				u32 writeCapacity;
-
-				if (sizeof(recvCache.bytes) < recvLeft)
-				{
-					writePos = buffer;
-					writeCapacity = recvLeft;
-				}
-				else
-				{
-					writePos = recvCache.bytes;
-					writeCapacity = sizeof(recvCache.bytes);
-				}
-
-				read = SocketRecv(socket, writePos, writeCapacity);
-				if (read > 0)
-				{
-					if (writePos == buffer)
-					{
-						recvLeft -= read;
-						buffer += read;
-					}
-					else
-					{
-						recvCache.written = read;
-						u32 toCopy = Min(recvCache.written, recvLeft);
-						memcpy(buffer, recvCache.bytes, toCopy);
-						recvCache.pos += toCopy;
-						recvLeft -= toCopy;
-						buffer += toCopy;
-					}
-
-					if (recvLeft == 0)
-						return true;
+					if (!allowLess)
+						continue;
 
-					if (allowLess)
-					{
-						bufferLen = read;
-						return true;
-					}
-					continue;
+					bufferLen = read;
+					return true;
 				}
 			}
 
diff --git a/Engine/Source/Programs/UnrealBuildAccelerator/Common/Public/UbaSocket.h b/Engine/Source/Programs/UnrealBuildAccelerator/Common/Public/UbaSocket.h
--- a/Engine/Source/Programs/UnrealBuildAccelerator/Common/Public/UbaSocket.h
+++ b/Engine/Source/Programs/UnrealBuildAccelerator/Common/Public/UbaSocket.h
@@ -38,6 +38,12 @@ namespace uba
 		const void* data;
 	};
 
+	struct RecvBuf
+	{
+		u32 len;
+		void* data;
+	};
+
 	Socket SocketCreate(int family = AF_INET, int type = SOCK_STREAM, int protocol = IPPROTO_TCP);
 	int SocketConnect(Socket s, const sockaddr* addr, u32 addrSize, bool isBlocking = true);
 	int SocketClose(Socket s);
@@ -55,6 +61,7 @@ namespace uba
 	int SocketPoll(Socket s, int events, int timeoutMs, int* outRevents);
 	Socket SocketAccept(Socket s, sockaddr* outAddr, int addrLen);
 	int SocketRecv(Socket s, void* data, int len);
+	int SocketRecv2(Socket s, RecvBuf* bufs, int bufCount);
 	int SocketSend(Socket s, SendBuf* bufs, int bufCount);
 	bool SocketShouldPoll(bool* outRetry);
 	int SocketCheckConnect(Socket s, bool* outTimedOut);
diff --git a/Engine/Source/Programs/UnrealBuildAccelerator/Common/Public/UbaSocket.inl b/Engine/Source/Programs/UnrealBuildAccelerator/Common/Public/UbaSocket.inl
--- a/Engine/Source/Programs/UnrealBuildAccelerator/Common/Public/UbaSocket.inl
+++ b/Engine/Source/Programs/UnrealBuildAccelerator/Common/Public/UbaSocket.inl
@@ -230,6 +230,26 @@ int SocketRecv(Socket s, void* data, int len)
 	return (int)recv(ToPlatformSocket(s), (char*)data, len, 0);
 }
 
+int SocketRecv2(Socket s, RecvBuf* bufs, int bufCount)
+{
+#if PLATFORM_WINDOWS
+	DWORD flags = 0;
+	DWORD received = 0;
+	int res = WSARecv(ToPlatformSocket(s), (LPWSABUF)bufs, bufCount, &received, &flags, NULL, NULL);
+	if (!res)
+		return int(received);
+	return SOCKET_ERROR;
+#else
+	iovec bufs2[2];
+	for (int i=0; i!=bufCount; ++i)
+		bufs2[i] = { (void*)bufs[i].data, bufs[i].len };
+	msghdr msg = {};
+	msg.msg_iov = bufs2;
+	msg.msg_iovlen = bufCount;
+	return recvmsg(ToPlatformSocket(s), &msg, 0);
+#endif
+}
+
 int SocketSend(Socket s, SendBuf* bufs, int bufCount)
 {
 	if (bufCount == 1)
@@ -241,7 +261,7 @@ int SocketSend(Socket s, SendBuf* bufs, int bufCount)
 		res = int(bytesSent);
 	return res;
 #else
-	iovec bufs2[256];
+	iovec bufs2[256]; // 256 is also set in UbaNetworkBackendTcp.cpp
 	for (int i=0; i!=bufCount; ++i)
 		bufs2[i] = { (void*)bufs[i].data, bufs[i].len };
 	return writev(ToPlatformSocket(s), bufs2, bufCount);
diff --git a/Engine/Source/Programs/UnrealBuildAccelerator/Common/Public/UbaSocketWin.h b/Engine/Source/Programs/UnrealBuildAccelerator/Common/Public/UbaSocketWin.h
--- a/Engine/Source/Programs/UnrealBuildAccelerator/Common/Public/UbaSocketWin.h
+++ b/Engine/Source/Programs/UnrealBuildAccelerator/Common/Public/UbaSocketWin.h
@@ -31,6 +31,7 @@ namespace uba
 		SOCKET_FUNCTION(SetSendBuf) \
 		SOCKET_FUNCTION(SetPriority) \
 		SOCKET_FUNCTION(Recv) \
+		SOCKET_FUNCTION(Recv2) \
 		SOCKET_FUNCTION(Send) \
 		SOCKET_FUNCTION(ShouldPoll) \
 		SOCKET_FUNCTION(CheckConnect) \
@@ -61,6 +62,7 @@ namespace uba
 	int SocketPoll(Socket s, int events, int timeoutMs, int* outRevents) { return g_socketPollFunc(s, events, timeoutMs, outRevents); }
 	Socket SocketAccept(Socket s, sockaddr* outAddr, int addrLen) { return g_socketAcceptFunc(s, outAddr, addrLen); }
 	int SocketRecv(Socket s, void* data, int len) { return g_socketRecvFunc(s, data, len); }
+	int SocketRecv2(Socket s, RecvBuf* bufs, int bufCount) { return g_socketRecv2Func(s, bufs, bufCount); }
 	int SocketSend(Socket s, SendBuf* bufs, int bufCount) { return g_socketSendFunc(s, bufs, bufCount); }
 	int SocketCheckConnect(Socket s, bool* outTimedOut) { return g_socketCheckConnectFunc(s, outTimedOut); }
 	bool SocketShouldPoll(bool* outRetry) { return g_socketShouldPollFunc(outRetry); }
diff --git a/Engine/Source/Programs/UnrealBuildAccelerator/Wine/Build.sh b/Engine/Source/Programs/UnrealBuildAccelerator/Wine/Build.sh
--- a/Engine/Source/Programs/UnrealBuildAccelerator/Wine/Build.sh
+++ b/Engine/Source/Programs/UnrealBuildAccelerator/Wine/Build.sh
@@ -1,3 +1,10 @@
 #!/bin/sh
 cd "$(dirname "$0")"
-winegcc -shared -m64 UbaWineWindows.cpp UbaWineLinux.cpp UbaWineWindows.def -o ../../../../Binaries/Win64/UnrealBuildAccelerator/x64/UbaWine.dll.so
\ No newline at end of file
+
+EXTRA_FLAGS="-O2 -DNDEBUG"
+
+if [ -d "/usr/include/wine/wine/windows" ]; then
+    EXTRA_FLAGS="$EXTRA_FLAGS -I/usr/include/wine/wine/windows"
+fi
+
+winegcc -shared -m64 -Wno-changes-meaning $EXTRA_FLAGS UbaWineWindows.cpp UbaWineLinux.cpp UbaWineWindows.def -o ../../../../Binaries/Win64/UnrealBuildAccelerator/x64/UbaWine.dll.so
\ No newline at end of file
diff --git a/Engine/Source/Programs/UnrealBuildAccelerator/Wine/UbaWineWindows.cpp b/Engine/Source/Programs/UnrealBuildAccelerator/Wine/UbaWineWindows.cpp
--- a/Engine/Source/Programs/UnrealBuildAccelerator/Wine/UbaWineWindows.cpp
+++ b/Engine/Source/Programs/UnrealBuildAccelerator/Wine/UbaWineWindows.cpp
@@ -1,5 +1,9 @@
 // Copyright Epic Games, Inc. All Rights Reserved.
 
+#define PLATFORM_WINDOWS 0
+#define PLATFORM_MAC 0
+#define PLATFORM_LINUX 1
+
 #include <winsock2.h>
 #include <stdint.h>
 #include <unistd.h>
