[UBA] Add memory usage tracking and visualization to build monitor

Extended SchedulerUpdate trace events with memory statistics and bumped TraceVersion to 53, with backward compatibility for older traces.
Added periodic memory stat emissions from the memory check loop, tracking spawn-paused transitions to preserve throttle timing accuracy.
Visualizer now displays a memory usage progress bar showing estimated memory percentage, with color changes to indicate active throttling.

#Synced-CL 51944883

diff --git a/Engine/Source/Programs/UnrealBuildAccelerator/Common/Private/UbaScheduler.cpp b/Engine/Source/Programs/UnrealBuildAccelerator/Common/Private/UbaScheduler.cpp
--- a/Engine/Source/Programs/UnrealBuildAccelerator/Common/Private/UbaScheduler.cpp
+++ b/Engine/Source/Programs/UnrealBuildAccelerator/Common/Private/UbaScheduler.cpp
@@ -1712,7 +1712,7 @@ namespace uba
 		{
 			if (m_memDelaySpawning)
 			{
-				m_session.GetTrace().SchedulerUpdate(false);
+				m_session.GetTrace().SchedulerUpdate(false, m_memLastTotal, m_memEstimatedUsage);
 				m_memDelaySpawning = false;
 			}
 			return true;
@@ -1725,7 +1725,7 @@ namespace uba
 			return false;
 
 		m_memDelaySpawning = true;
-		m_session.GetTrace().SchedulerUpdate(true);
+		m_session.GetTrace().SchedulerUpdate(true, m_memLastTotal, m_memEstimatedUsage);
 
 #if 0//PLATFORM_WINDOWS
 		static bool hasBeenRunOnce;
@@ -2094,6 +2094,9 @@ namespace uba
 				trace.StatusUpdate(StartRow+5, 6, BytesToText(memKillThreshold), LogEntryType_Info);
 			}
 
+			m_memLastTotal = memTotal;
+			trace.SchedulerUpdate(m_memDelaySpawning, memTotal, m_memEstimatedUsage);
+
 			if (memKillThreshold)
 			{
 				if (m_lastKillTime && m_lastKillTime + MsToTime(3000) < GetTime())
diff --git a/Engine/Source/Programs/UnrealBuildAccelerator/Common/Private/UbaTrace.cpp b/Engine/Source/Programs/UnrealBuildAccelerator/Common/Private/UbaTrace.cpp
--- a/Engine/Source/Programs/UnrealBuildAccelerator/Common/Private/UbaTrace.cpp
+++ b/Engine/Source/Programs/UnrealBuildAccelerator/Common/Private/UbaTrace.cpp
@@ -639,10 +639,12 @@ namespace uba
 		BEGIN_TRACE_ENTRY(RemoteExecutionDisabled);
 	}
 
-	void Trace::SchedulerUpdate(bool spawnPaused)
+	void Trace::SchedulerUpdate(bool spawnPaused, u64 memTotal, u64 memEstimated)
 	{
 		BEGIN_TRACE_ENTRY(SchedulerUpdate);
 		writer.WriteBool(spawnPaused);
+		writer.Write7BitEncoded(memTotal);
+		writer.Write7BitEncoded(memEstimated);
 	}
 
 	void Trace::SchedulerKillProcess(u32 processId, StringView reason)
diff --git a/Engine/Source/Programs/UnrealBuildAccelerator/Common/Public/UbaScheduler.h b/Engine/Source/Programs/UnrealBuildAccelerator/Common/Public/UbaScheduler.h
--- a/Engine/Source/Programs/UnrealBuildAccelerator/Common/Public/UbaScheduler.h
+++ b/Engine/Source/Programs/UnrealBuildAccelerator/Common/Public/UbaScheduler.h
@@ -242,6 +242,7 @@ namespace uba
 		Atomic<u64> m_memEstimatedUsage = 0;
 		Atomic<u64> m_memWaitThreshold = 0;
 		Atomic<u64> m_lastKillTime = 0;
+		Atomic<u64> m_memLastTotal = 0;
 		u8 m_memStartWaitPercent = 0;
 		u8 m_memStartKillPercent = 0;
 		u8 m_memTracingLevel = 0;
diff --git a/Engine/Source/Programs/UnrealBuildAccelerator/Common/Public/UbaTrace.h b/Engine/Source/Programs/UnrealBuildAccelerator/Common/Public/UbaTrace.h
--- a/Engine/Source/Programs/UnrealBuildAccelerator/Common/Public/UbaTrace.h
+++ b/Engine/Source/Programs/UnrealBuildAccelerator/Common/Public/UbaTrace.h
@@ -66,7 +66,7 @@ namespace uba
 		#undef UBA_TRACE_TYPE
 	};
 
-	static constexpr u32 TraceVersion = 52;
+	static constexpr u32 TraceVersion = 53;
 	static constexpr u32 TraceReadCompatibilityVersion = 6;
 
 	////////////////////////////////////////////////////////////////////////////////////////////////////
@@ -128,7 +128,7 @@ namespace uba
 		void StatusUpdate2(u32 statusRow, u32 statusColumn, u32 stringId);
 		void DriveUpdate(const tchar drive, u8 busyPercent, u32 readCount, u64 readBytes, u32 writeCount, u64 writeBytes);
 		void RemoteExecutionDisabled();
-		void SchedulerUpdate(bool spawnPaused);
+		void SchedulerUpdate(bool spawnPaused, u64 memTotal = 0, u64 memEstimated = 0);
 		void SchedulerKillProcess(u32 processId, StringView reason);
 
 		void CacheBeginFetch(u32 fetchId, const tchar* description);
diff --git a/Engine/Source/Programs/UnrealBuildAccelerator/Visualizer/Private/UbaTraceReader.cpp b/Engine/Source/Programs/UnrealBuildAccelerator/Visualizer/Private/UbaTraceReader.cpp
--- a/Engine/Source/Programs/UnrealBuildAccelerator/Visualizer/Private/UbaTraceReader.cpp
+++ b/Engine/Source/Programs/UnrealBuildAccelerator/Visualizer/Private/UbaTraceReader.cpp
@@ -78,6 +78,10 @@ namespace uba
 		lastKillReason = nullptr;
 		lastSpawningDelayStartTime = 0;
 		lastSpawningDelayEndTime = 0;
+		memTotal = 0;
+		memEstimated = 0;
+		hasMemoryStats = false;
+		lastSpawnPausedState = false;
 		remoteExecutionDisabled = false;
 		finished = true;
 	};
@@ -1503,13 +1507,26 @@ namespace uba
 		}
 		case TraceType_SchedulerUpdate:
 		{
-			if (reader.ReadBool())
+			bool spawnPaused = reader.ReadBool();
+
+			if (spawnPaused && !out.lastSpawnPausedState)
 			{
 				out.lastSpawningDelayStartTime = time;
 				out.lastSpawningDelayEndTime = ~0llu - MsToTime(10000);
 			}
-			else
+			else if (!spawnPaused && out.lastSpawnPausedState)
+			{
 				out.lastSpawningDelayEndTime = time;
+			}
+			out.lastSpawnPausedState = spawnPaused;
+
+			if (out.version >= 53)
+			{
+				out.memTotal = reader.Read7BitEncoded();
+				out.memEstimated = reader.Read7BitEncoded();
+				if (out.memTotal > 0)
+					out.hasMemoryStats = true;
+			}
 			break;
 		}
 		case TraceType_SchedulerKillProcess:
diff --git a/Engine/Source/Programs/UnrealBuildAccelerator/Visualizer/Private/UbaTraceReader.h b/Engine/Source/Programs/UnrealBuildAccelerator/Visualizer/Private/UbaTraceReader.h
--- a/Engine/Source/Programs/UnrealBuildAccelerator/Visualizer/Private/UbaTraceReader.h
+++ b/Engine/Source/Programs/UnrealBuildAccelerator/Visualizer/Private/UbaTraceReader.h
@@ -222,6 +222,10 @@ namespace uba
 		u32 progressProcessesDone = 0;
 		u32 progressErrorCount = 0;
 		u16 maxActiveProcessCount = 0;
+		u64 memTotal = 0;
+		u64 memEstimated = 0;
+		bool hasMemoryStats = false;
+		bool lastSpawnPausedState = false;
 		bool remoteExecutionDisabled = false;
 		bool finished = true;
 	};
diff --git a/Engine/Source/Programs/UnrealBuildAccelerator/Visualizer/Private/UbaVisualizer.cpp b/Engine/Source/Programs/UnrealBuildAccelerator/Visualizer/Private/UbaVisualizer.cpp
--- a/Engine/Source/Programs/UnrealBuildAccelerator/Visualizer/Private/UbaVisualizer.cpp
+++ b/Engine/Source/Programs/UnrealBuildAccelerator/Visualizer/Private/UbaVisualizer.cpp
@@ -894,6 +894,33 @@ namespace uba
 				text.Appendf(TC("Killing (%s)"), killReason);
 				drawStatusText(text, LogEntryType_Info, rect.left, rect.right, 0, false, true);
 			}
+			else if (m_traceView.hasMemoryStats && m_traceView.memTotal > 0)
+			{
+				u32 estPct = u32(m_traceView.memEstimated * 100 / m_traceView.memTotal);
+				bool isThrottling = m_traceView.lastSpawningDelayStartTime && playTime >= m_traceView.lastSpawningDelayStartTime && playTime < m_traceView.lastSpawningDelayEndTime + MsToTime(500);
+
+				rect.left = start + width + 20;
+				u32 barWidth = activeFont.height * 8;
+				rect.right = rect.left + barWidth;
+
+				// Gray background
+				context.FillRect(rect, m_processBrushes[0].returned);
+
+				// Fill proportional to estimated usage
+				u32 fillPct = Min(estPct, 100u);
+				RECT fillRect = rect;
+				fillRect.right = fillRect.left + int(float(fillPct) / 100.0f * float(barWidth));
+				BrushHandle fillBrush = isThrottling ? m_processBrushes[0].warning : m_processBrushes[0].success;
+				if (fillPct > 0)
+					context.FillRect(fillRect, fillBrush);
+
+				StringBuffer<64> memText;
+				if (isThrottling)
+					memText.Appendf(TC("Throttling %u%%"), estPct);
+				else
+					memText.Appendf(TC("Mem %u%%"), estPct);
+				drawStatusText(memText, LogEntryType_Info, rect.left, rect.right, 0, false, true);
+			}
 			else if (m_traceView.lastSpawningDelayStartTime && playTime >= m_traceView.lastSpawningDelayStartTime && playTime < m_traceView.lastSpawningDelayEndTime + MsToTime(500))
 			{
 				rect.left = start + width + 20;
