[UBA]
* Detoured the Dos->Nt path translation functions in order to try to get rid of all the filename related detouring in kernel layer
* Untangled the first part of NtCreateFile a bit. Still a mess but slightly better

diff --git a/Engine/Source/Programs/UnrealBuildAccelerator/Core/Private/UbaApplicationRules.cpp b/Engine/Source/Programs/UnrealBuildAccelerator/Core/Private/UbaApplicationRules.cpp
--- a/Engine/Source/Programs/UnrealBuildAccelerator/Core/Private/UbaApplicationRules.cpp
+++ b/Engine/Source/Programs/UnrealBuildAccelerator/Core/Private/UbaApplicationRules.cpp
@@ -42,9 +42,9 @@ namespace uba
 		using Super = ApplicationRules;
 
 	public:
-		virtual bool CanDetour(const tchar* file, bool isRunningRemote) const override
+		virtual bool CanDetour(StringView file, bool isRunningRemote) const override
 		{
-			return !isRunningRemote || TStrchr(file, ',') == 0;
+			return !isRunningRemote || !file.Contains(',');
 		}
 	};
 
diff --git a/Engine/Source/Programs/UnrealBuildAccelerator/Core/Private/UbaStringBuffer.cpp b/Engine/Source/Programs/UnrealBuildAccelerator/Core/Private/UbaStringBuffer.cpp
--- a/Engine/Source/Programs/UnrealBuildAccelerator/Core/Private/UbaStringBuffer.cpp
+++ b/Engine/Source/Programs/UnrealBuildAccelerator/Core/Private/UbaStringBuffer.cpp
@@ -504,7 +504,7 @@ namespace uba
 
 	StringView ToView(const tchar* s)
 	{
-		return StringView(s, TStrlen(s));
+		return s ? StringView(s, TStrlen(s)) : StringView();;
 	}
 
 	////////////////////////////////////////////////////////////////////////////////////////////////////
diff --git a/Engine/Source/Programs/UnrealBuildAccelerator/Core/Public/UbaApplicationRules.h b/Engine/Source/Programs/UnrealBuildAccelerator/Core/Public/UbaApplicationRules.h
--- a/Engine/Source/Programs/UnrealBuildAccelerator/Core/Public/UbaApplicationRules.h
+++ b/Engine/Source/Programs/UnrealBuildAccelerator/Core/Public/UbaApplicationRules.h
@@ -37,7 +37,7 @@ namespace uba
 		}
 
 		// Control if certain files should not be detoured
-		virtual bool CanDetour(const tchar* file, bool isRunningRemote) const
+		virtual bool CanDetour(StringView file, bool isRunningRemote) const
 		{
 			return true;
 		}
diff --git a/Engine/Source/Programs/UnrealBuildAccelerator/Detours/Private/PosixOS/UbaDetoursFunctionsPosix.cpp b/Engine/Source/Programs/UnrealBuildAccelerator/Detours/Private/PosixOS/UbaDetoursFunctionsPosix.cpp
--- a/Engine/Source/Programs/UnrealBuildAccelerator/Detours/Private/PosixOS/UbaDetoursFunctionsPosix.cpp
+++ b/Engine/Source/Programs/UnrealBuildAccelerator/Detours/Private/PosixOS/UbaDetoursFunctionsPosix.cpp
@@ -312,7 +312,7 @@ namespace uba
 	}
 
 	bool CouldBeCompressedFile(const StringView& fileName) { return false; }
-	bool CanDetour(const tchar* file)
+	bool CanDetour(StringView file)
 	{
 		if (t_disallowDetour)
 			return false;
@@ -807,13 +807,15 @@ int Shared_stat(const char* funcName, const char* file, struct stat* attr, const
 //template<typename True_stat>
 int Shared_access(const char* funcName, const char* pathname, int mode)
 {
+	StringView original(ToView(pathname));
+
 	StringBuffer<> fixedPath;
-	if (!FixPath(fixedPath, pathname) || fixedPath.StartsWith("/proc") || !CanDetour2(fixedPath))
+	if (!FixPath(fixedPath, original) || fixedPath.StartsWith("/proc") || !CanDetour2(fixedPath))
 	{
 		TimerScope ts(g_kernelStats.getFileInfo);
-		auto res = TRUE_WRAPPER(access)(pathname, mode);
+		auto res = TRUE_WRAPPER(access)(original.data, mode);
 		ts.Leave();
-		DEBUG_LOG_TRUE("access", "%s %i -> %i (%s)", pathname, mode, res, StrError(res, errno));
+		DEBUG_LOG_TRUE("access", "%s %i -> %i (%s)", original.data, mode, res, StrError(res, errno));
 		return res;
 	}
 
@@ -833,12 +835,12 @@ int Shared_access(const char* funcName, const char* pathname, int mode)
 
 	DevirtualizePath(fixedPath);
 
-	if (!CanDetour(pathname))
+	if (!CanDetour(original))
 	{
 		TimerScope ts(g_kernelStats.getFileInfo);
-		auto res = TRUE_WRAPPER(access)(pathname, mode);
+		auto res = TRUE_WRAPPER(access)(original.data, mode);
 		ts.Leave();
-		DEBUG_LOG_TRUE(funcName, "%s %i -> %i (%s)", pathname, mode, res, StrError(res, errno));
+		DEBUG_LOG_TRUE(funcName, "%s %i -> %i (%s)", original.data, mode, res, StrError(res, errno));
 		return res;
 	}
 
@@ -850,13 +852,13 @@ int Shared_access(const char* funcName, const char* pathname, int mode)
 		TimerScope ts(g_kernelStats.getFileInfo);
 		auto res = TRUE_WRAPPER(access)(realName, mode);
 		ts.Leave();
-		DEBUG_LOG_TRUE(funcName, "%s %i (%s) -> %i %s", pathname, mode, realName, res, StrError(res, errno));
+		DEBUG_LOG_TRUE(funcName, "%s %i (%s) -> %i %s", original.data, mode, realName, res, StrError(res, errno));
 		return res;
 	}
 
 	int res = attr.lastError == 0 ? 0 : -1;
 	
-	DEBUG_LOG_DETOURED(funcName, "%s %i (%s) -> %i %s", pathname, mode, realName, res, StrError(res, attr.lastError));
+	DEBUG_LOG_DETOURED(funcName, "%s %i (%s) -> %i %s", original.data, mode, realName, res, StrError(res, attr.lastError));
 
 	#if UBA_DEBUG_VALIDATE
 	if (!g_runningRemote)
@@ -865,7 +867,7 @@ int Shared_access(const char* funcName, const char* pathname, int mode)
 		auto res2 = TRUE_WRAPPER(access)(realName, mode);
 		//int eo = errno;
 		
-		DEBUG_LOG_DETOURED(funcName, "%s %i (%s) -> %i %s", pathname, mode, realName, res, StrError(res, attr.lastError));
+		DEBUG_LOG_DETOURED(funcName, "%s %i (%s) -> %i %s", original.data, mode, realName, res, StrError(res, attr.lastError));
 		
 		UBA_ASSERTF(res2 == res, "MISMATCH OF RESULTS for %s - %i %i (err = %s) (exedir %s)", realName, res2, res, StrError(res, attr.lastError), g_exeDir.data);
 		//UBA_ASSERTF(eo == attr.lastError, "MISMATCH OF ERRORS for %s - %i %i", realName, eo, attr.lastError);
diff --git a/Engine/Source/Programs/UnrealBuildAccelerator/Detours/Private/UbaDetoursShared.cpp b/Engine/Source/Programs/UnrealBuildAccelerator/Detours/Private/UbaDetoursShared.cpp
--- a/Engine/Source/Programs/UnrealBuildAccelerator/Detours/Private/UbaDetoursShared.cpp
+++ b/Engine/Source/Programs/UnrealBuildAccelerator/Detours/Private/UbaDetoursShared.cpp
@@ -235,6 +235,11 @@ namespace uba
 		return FixPath2(path, pathLen, g_virtualWorkingDir.data, g_virtualWorkingDir.count, out.data, out.capacity, &out.count);
 	}
 
+	bool FixPath(StringBufferBase& out, const StringView path)
+	{
+		return FixPath2(path.data, path.count, g_virtualWorkingDir.data, g_virtualWorkingDir.count, out.data, out.capacity, &out.count);
+	}
+
 	struct VfsEntry { StringView vfs; StringView local; VfsEntry() : vfs(NoInit), local(NoInit) {}; };
 	VfsEntry g_vfsEntries[32];
 
diff --git a/Engine/Source/Programs/UnrealBuildAccelerator/Detours/Private/UbaDetoursShared.h b/Engine/Source/Programs/UnrealBuildAccelerator/Detours/Private/UbaDetoursShared.h
--- a/Engine/Source/Programs/UnrealBuildAccelerator/Detours/Private/UbaDetoursShared.h
+++ b/Engine/Source/Programs/UnrealBuildAccelerator/Detours/Private/UbaDetoursShared.h
@@ -23,7 +23,8 @@ namespace uba
 
 	const tchar* GetApplicationShortName();
 	bool FixPath(StringBufferBase& out, const tchar* path, u32 pathLen = ~0u);
-	
+	bool FixPath(StringBufferBase& out, const StringView path);
+
 	void PopulateVfs(BinaryReader& vfsReader);
 	bool IsVfsEnabled();
 	bool DevirtualizePath(StringBufferBase& path);
diff --git a/Engine/Source/Programs/UnrealBuildAccelerator/Detours/Private/Windows/UbaDetoursFunctionsKernelBase.inl b/Engine/Source/Programs/UnrealBuildAccelerator/Detours/Private/Windows/UbaDetoursFunctionsKernelBase.inl
--- a/Engine/Source/Programs/UnrealBuildAccelerator/Detours/Private/Windows/UbaDetoursFunctionsKernelBase.inl
+++ b/Engine/Source/Programs/UnrealBuildAccelerator/Detours/Private/Windows/UbaDetoursFunctionsKernelBase.inl
@@ -216,19 +216,9 @@ HANDLE Detoured_CreateFileW(LPCWSTR lpFileName, DWORD dwDesiredAccess, DWORD dwS
 	DEBUG_LOG_PROXY(L"CreateFileW", L"%ls", lpFileName);
 	u32 disallowDetour = !lpFileName || Equals(lpFileName, L"nul");
 	t_disallowDetour += disallowDetour;
-
-	#if UBA_CHECK_PATH_CONVERSION
-	t_rtlCalled = false;
-	#endif
-
-	t_createFileFileName = lpFileName;
+	t_hackInsideCreateFile = true;
 	HANDLE h = True_CreateFileW(lpFileName, dwDesiredAccess, dwShareMode, lpSecurityAttributes, dwCreationDisposition, dwFlagsAndAttributes, hTemplateFile);
-
-	#if UBA_CHECK_PATH_CONVERSION
-	bool doNotCheckRtlCalled = g_isRunningWine && (!lpFileName || !lpFileName[0] || dwCreationDisposition < CREATE_NEW || dwCreationDisposition > TRUNCATE_EXISTING);
-	UBA_ASSERTF(doNotCheckRtlCalled || t_rtlCalled, TC("Rtl path conversion not called fo rpath %s"), lpFileName);
-	#endif
-	t_createFileFileName = nullptr;
+	t_hackInsideCreateFile = false;
 	t_disallowDetour -= disallowDetour;
 	return h;
 }
@@ -242,19 +232,9 @@ HANDLE Detoured_CreateFileA(LPCSTR lpFileName, DWORD dwDesiredAccess, DWORD dwSh
 	fileName.Appendf(L"%hs", lpFileName);
 	u32 disallowDetour = fileName.Equals(L"nul");
 	t_disallowDetour += disallowDetour;
-
-	#if UBA_CHECK_PATH_CONVERSION
-	t_rtlCalled = false;
-	#endif
-
-	t_createFileFileName = fileName.data;
+	t_hackInsideCreateFile = true;
 	HANDLE h = True_CreateFileA(lpFileName, dwDesiredAccess, dwShareMode, lpSecurityAttributes, dwCreationDisposition, dwFlagsAndAttributes, hTemplateFile);
-
-	#if UBA_CHECK_PATH_CONVERSION
-	UBA_ASSERT(t_rtlCalled);
-	#endif
-
-	t_createFileFileName = nullptr;
+	t_hackInsideCreateFile = false;
 	t_disallowDetour -= disallowDetour;
 	return h;
 }
@@ -772,7 +752,8 @@ BOOL Shared_GetFileAttributesExW(const tchar* caller, StringView fileName, GET_F
 BOOL Detoured_GetFileAttributesExW(LPCWSTR lpFileName, GET_FILEEX_INFO_LEVELS fInfoLevelId, LPVOID lpFileInformation)
 {
 	DETOURED_CALL(GetFileAttributesExW);
-	if (!CanDetour(lpFileName) || Contains(lpFileName, L"::")) // Some weird .net path used by dotnet.exe ... ignore for now!
+	StringView dosName(ToView(lpFileName));
+	if (!CanDetour(dosName) || dosName.Contains(L"::")) // Some weird .net path used by dotnet.exe ... ignore for now!
 	{
 		//UBA_ASSERTF(!g_runningRemote, TC("GetFileAttributesExW can't handle this path when running remote: %s"), lpFileName);
 		DEBUG_LOG_TRUE(L"GetFileAttributesExW", L"(%ls)", lpFileName);
@@ -781,7 +762,7 @@ BOOL Detoured_GetFileAttributesExW(LPCWSTR lpFileName, GET_FILEEX_INFO_LEVELS fI
 	}
 
 	StringBuffer<MaxPath> fixedName;
-	FixPath(fixedName, lpFileName);
+	FixPath(fixedName, dosName);
 	DevirtualizePath(fixedName);
 
 	if (!g_rules->CanExist(fixedName))
@@ -797,7 +778,8 @@ BOOL Detoured_GetFileAttributesExW(LPCWSTR lpFileName, GET_FILEEX_INFO_LEVELS fI
 DWORD Detoured_GetFileAttributesW(LPCWSTR lpFileName)
 {
 	DETOURED_CALL(GetFileAttributesW);
-	if (!CanDetour(lpFileName))
+	StringView dosName(ToView(lpFileName));
+	if (!CanDetour(dosName))
 	{
 		TimerScope ts(g_kernelStats.getFileInfo);
 		DWORD res = True_GetFileAttributesW(lpFileName);
@@ -806,7 +788,7 @@ DWORD Detoured_GetFileAttributesW(LPCWSTR lpFileName)
 	}
 
 	StringBuffer<MaxPath> fixedPath;
-	if (!FixPath(fixedPath, lpFileName))
+	if (!FixPath(fixedPath, dosName))
 	{
 		DEBUG_LOG_DETOURED(L"GetFileAttributesW", L"(BAD_PATH) (%ls) -> INVALID_FILE_ATTRIBUTES", lpFileName);
 		SetLastError(ERROR_INVALID_NAME);
@@ -1309,15 +1291,16 @@ BOOL Detoured_DeleteFileW(LPCWSTR lpFileName)
 	#endif
 
 	LPCWSTR original = lpFileName;
+	StringView dosName(ToView(lpFileName));
 
-	if (!CanDetour(lpFileName))
+	if (!CanDetour(dosName))
 	{
 		DEBUG_LOG_TRUE(L"DeleteFileW", L"(%ls)", original);
 		return True_DeleteFileW(original);
 	}
 
 	StringBuffer<> fixedName;
-	FixPath(fixedName, lpFileName);
+	FixPath(fixedName, dosName);
 
 	DevirtualizePath(fixedName);
 
diff --git a/Engine/Source/Programs/UnrealBuildAccelerator/Detours/Private/Windows/UbaDetoursFunctionsNtDll.inl b/Engine/Source/Programs/UnrealBuildAccelerator/Detours/Private/Windows/UbaDetoursFunctionsNtDll.inl
--- a/Engine/Source/Programs/UnrealBuildAccelerator/Detours/Private/Windows/UbaDetoursFunctionsNtDll.inl
+++ b/Engine/Source/Programs/UnrealBuildAccelerator/Detours/Private/Windows/UbaDetoursFunctionsNtDll.inl
@@ -82,6 +82,87 @@ StringBuffer<32> ToString(NTSTATUS s)
 }
 #endif
 
+extern "C" RTL_PATH_TYPE NTAPI RtlDetermineDosPathNameType_U(PCWSTR DosFileName);
+extern "C" ULONG NTAPI RtlIsDosDeviceName_U(PCWSTR DosFileName);
+
+static void FixupDosPath(PCWSTR& dosPath, StringBufferBase& temp)
+{
+	if (RtlIsDosDeviceName_U(dosPath) != 0)
+		return;
+
+	switch (RtlDetermineDosPathNameType_U(dosPath))
+	{
+	case RtlPathTypeRelative:
+		temp.Append(g_virtualWorkingDir).EnsureEndsWithSlash().Append(dosPath);
+		dosPath = temp.data;
+		break;
+	case RtlPathTypeRooted:
+		temp.Append(g_virtualWorkingDir[0]).Append(':').Append(dosPath);
+		dosPath = temp.data;
+		break;
+	case RtlPathTypeDriveRelative:
+		wchar_t envVar[4] = { L'=', dosPath[0], L':', L'\0' };
+		temp.count = GetEnvironmentVariableW(envVar, temp.data, temp.capacity);
+		if (temp.count == 0)
+			temp.Append(dosPath[0]).Append(TCV(":\\"));
+		temp.EnsureEndsWithSlash().Append(dosPath + 2);
+		dosPath = temp.data;
+		break;
+	}
+}
+
+NTSTATUS Detoured_RtlDosPathNameToNtPathName_U_WithStatus(PCWSTR dosPath, PUNICODE_STRING ntPath, PWSTR* filePart, VOID* reserved)
+{
+	const tchar* original = dosPath;(void)original;
+	StringBuffer<> temp;
+	FixupDosPath(dosPath, temp);
+	auto res = True_RtlDosPathNameToNtPathName_U_WithStatus(dosPath, ntPath, filePart, reserved);
+	//DEBUG_LOG(L"RtlDosPathNameToNtPathName_U_WithStatus:           %s -> %.*s", original, ntPath->Length/2, ntPath->Buffer);
+	return res;
+}
+
+NTSTATUS Detoured_RtlDosPathNameToRelativeNtPathName_U_WithStatus(PCWSTR dosPath, PUNICODE_STRING ntPath, PWSTR* filePart, RTL_RELATIVE_NAME_U* relativeName)
+{
+	const tchar* original = dosPath;(void)original;
+	StringBuffer<> temp;
+	UBA_ASSERT(!relativeName || !relativeName->ContainingDirectory);
+	FixupDosPath(dosPath, temp);
+	auto res = True_RtlDosPathNameToRelativeNtPathName_U_WithStatus(dosPath, ntPath, filePart, relativeName);
+	//DEBUG_LOG(L"RtlDosPathNameToRelativeNtPathName_U_WithStatus:   %s -> %.*s", original, ntPath->Length/2, ntPath->Buffer);
+	return res;
+}
+
+DWORD WINAPI Detoured_RtlGetFullPathName_U(const WCHAR* name, ULONG size, WCHAR* buffer, WCHAR** file_part)
+{
+	const tchar* original = name;(void)original;
+	StringBuffer<> temp;
+	FixupDosPath(name, temp);
+	auto res = True_RtlGetFullPathName_U(name, size, buffer, file_part);
+	//DEBUG_LOG(L"RtlGetFullPathName_U:                              %s -> %s", original, buffer);
+	return res;
+}
+
+ULONG WINAPI Detoured_RtlGetFullPathName_UEx(const WCHAR* name, ULONG size, WCHAR* buffer, WCHAR** file_part, RTL_PATH_TYPE* type)
+{
+	const tchar* original = name;(void)original;
+	StringBuffer<> temp;
+	FixupDosPath(name, temp);
+	auto res = True_RtlGetFullPathName_UEx(name, size, buffer, file_part, type);
+	//DEBUG_LOG(L"RtlGetFullPathName_UEx:                            %s -> %s", original, buffer);
+	return res;
+}
+
+BOOLEAN NTAPI Detoured_RtlDosPathNameToRelativeNtPathName_U(PCWSTR dosPath, PUNICODE_STRING ntPath, PWSTR* filePart, RTL_RELATIVE_NAME_U* relativeName)
+{
+	const tchar* original = dosPath;(void)original;
+	StringBuffer<> temp;
+	UBA_ASSERT(!relativeName || !relativeName->ContainingDirectory);
+	FixupDosPath(dosPath, temp);
+	auto res = True_RtlDosPathNameToRelativeNtPathName_U(dosPath, ntPath, filePart, relativeName);
+	//DEBUG_LOG(L"RtlDosPathNameToRelativeNtPathName_U:              %s -> %.*s", original, ntPath->Length/2, ntPath->Buffer);
+	return res;
+}
+
 NTSTATUS Detoured_NtQueryVolumeInformationFile(HANDLE FileHandle, PIO_STATUS_BLOCK IoStatusBlock, PVOID FsInformation, ULONG Length, FS_INFORMATION_CLASS FsInformationClass)
 {
 	IoStatusBlock->Status = STATUS_SUCCESS;
@@ -620,15 +701,13 @@ NTSTATUS NTAPI Detoured_NtQueryDirectoryFile(HANDLE FileHandle, HANDLE Event, PI
 
 NTSTATUS Detoured_NtQueryFullAttributesFile(POBJECT_ATTRIBUTES ObjectAttributes, PVOID Attributes)
 {
-	auto fileName = (const wchar_t*)ObjectAttributes->ObjectName->Buffer;
-	u32 fileNameLen = ObjectAttributes->ObjectName->Length/sizeof(tchar);(void)fileNameLen;
-	UBA_ASSERT(fileName[fileNameLen] == 0);
+	StringView fileName(ObjectAttributes->ObjectName->Buffer, ObjectAttributes->ObjectName->Length/sizeof(tchar));
 
-	if (!CanDetour(fileName) || Contains(fileName, L"::")) // Some weird .net path used by dotnet.exe ... ignore for now!
+	if (!g_rules->CanDetour(fileName, g_runningRemote) || fileName.Contains(L"::")) // Some weird .net path used by dotnet.exe ... ignore for now!
 	{
 		TimerScope ts(g_kernelStats.getFileInfo);
 		auto res = True_NtQueryFullAttributesFile(ObjectAttributes, Attributes);
-		DEBUG_LOG_TRUE(L"NtQueryFullAttributesFile", L"(%.*s) -> %s", fileNameLen, fileName, ToString(res).data);
+		DEBUG_LOG_TRUE(L"NtQueryFullAttributesFile", L"(%.*s) -> %s", fileName.count, fileName.data, ToString(res).data);
 		return res;
 	}
 
@@ -652,7 +731,7 @@ NTSTATUS Detoured_NtQueryFullAttributesFile(POBJECT_ATTRIBUTES ObjectAttributes,
 	if (!attr.useCache)
 	{
 		auto res = True_NtQueryFullAttributesFile(ObjectAttributes, Attributes);
-		DEBUG_LOG_TRUE(L"NtQueryFullAttributesFile", L"(%.*s) -> %s", fileNameLen, fileName, ToString(res).data);
+		DEBUG_LOG_TRUE(L"NtQueryFullAttributesFile", L"(%.*s) -> %s", fileName.count, fileName.data, ToString(res).data);
 		return res;
 	}
 
@@ -673,7 +752,7 @@ NTSTATUS Detoured_NtQueryFullAttributesFile(POBJECT_ATTRIBUTES ObjectAttributes,
 		info.FileAttributes = data.dwFileAttributes;
 	}
 
-	DEBUG_LOG_DETOURED(L"NtQueryFullAttributesFile", L"(%.*s) -> %s (Size: %llu)", fileNameLen, fileName, ToString(res).data, ToLargeInteger(attr.data.nFileSizeHigh, attr.data.nFileSizeLow).QuadPart);
+	DEBUG_LOG_DETOURED(L"NtQueryFullAttributesFile", L"(%.*s) -> %s (Size: %llu)", fileName.count, fileName.data, ToString(res).data, ToLargeInteger(attr.data.nFileSizeHigh, attr.data.nFileSizeLow).QuadPart);
 	return res;
 }
 
@@ -1137,12 +1216,6 @@ NTSTATUS NTAPI Local_NtCreateFile(bool IsCreateFunc, PHANDLE hFileHandle, ACCESS
 	}
 }
 
-#define UBA_CHECK_PATH_CONVERSION 0
-
-#if UBA_CHECK_PATH_CONVERSION
-thread_local bool t_rtlCalled;
-#endif
-
 NTSTATUS NTAPI Shared_NtCreateFile(bool IsCreateFunc, PHANDLE hFileHandle, ACCESS_MASK DesiredAccess, POBJECT_ATTRIBUTES ObjectAttributes, PIO_STATUS_BLOCK IoStatusBlock, PLARGE_INTEGER AllocationSize, ULONG FileAttributes,
 	ULONG ShareAccess, ULONG CreateDisposition, ULONG CreateOptions, PVOID EaBuffer, ULONG EaLength)
 {
@@ -1152,120 +1225,100 @@ NTSTATUS NTAPI Shared_NtCreateFile(bool IsCreateFunc, PHANDLE hFileHandle, ACCES
 	const wchar_t* funcName = IsCreateFunc ? L"NtCreateFile" : L"NtOpenFile"; (void)funcName;
 	#endif
 
-	#if UBA_CHECK_PATH_CONVERSION
-	UBA_ASSERT(!t_createFileFileName || t_rtlCalled);
-	#endif
+	StringView ntPath(ObjectAttributes->ObjectName->Buffer, ObjectAttributes->ObjectName->Length / 2);
 
-	const wchar_t* createFileName = t_createFileFileName;
+	auto trueFunc = [&]()
+		{
+			NTSTATUS res = Local_NtCreateFile(IsCreateFunc, hFileHandle, DesiredAccess, ObjectAttributes, IoStatusBlock, AllocationSize, FileAttributes, ShareAccess, CreateDisposition, CreateOptions, EaBuffer, EaLength);
+			DEBUG_LOG_TRUE(funcName, L"(SUPPRESSDETOUR) %llu (%.*s) -> %s", uintptr_t(*hFileHandle), ntPath.count, ntPath.data, ToString(res).data);
+			if (ntPath.Equals(L"\\??\\NUL"))
+				g_nullFile = *hFileHandle;
+			return res;
+		};
+
+	
+	if (t_disallowCreateFileDetour || t_disallowDetour)
+		return trueFunc();
+	if (!ntPath.count || ntPath[1] != '?')
+		return trueFunc();
+
+	bool allowDirectoryCache = g_allowDirectoryCache && t_hackInsideCreateFile; // TODO: This needs to be fixed.. t_hackInsideCreateFile just happens to work but is wrong!
 
-	// NOTE - ObjectAttributes->ObjectName->Buffer might not be null terminated, so we need to copy over to another buffer
-	StringBuffer<> fileName;
-	bool suppressCreateFileDetour = t_disallowCreateFileDetour || t_disallowDetour;
 	HANDLE rootDir = ObjectAttributes->RootDirectory;
+	StringBuffer<> fileName;
+
+	if (ntPath.StartsWith(TCV("\\??\\")))
 	{
-		const wchar_t* buf = ObjectAttributes->ObjectName->Buffer;
-		u32 bufChars = ObjectAttributes->ObjectName->Length / 2;
+		UBA_ASSERT(!ObjectAttributes->RootDirectory);
+		StringView dosPath = ntPath.Skip(4);
 
-		if (suppressCreateFileDetour)
-		{
-		}
-		else if (!buf)
+		if (!g_rules->CanDetour(dosPath, g_runningRemote))
+			return trueFunc();
+
+		if (dosPath[dosPath.count - 1] == '$')
 		{
-			suppressCreateFileDetour = true;
+			constexpr const wchar_t* stdStr[] = { L"conerr$", L"conout$", L"conin$" };
+			for (u32 i = 0; i != 3; ++i)
+			{
+				if (!dosPath.EndsWith(stdStr[i]))
+					continue;
+				if (g_isDetachedProcess)
+				{
+					IoStatusBlock->Status = STATUS_SUCCESS;
+					IoStatusBlock->Status = FILE_OPENED;
+					*hFileHandle = g_stdHandle[i];
+					DEBUG_LOG_DETOURED(funcName, L"(STD) %llu (%.*s) -> %s", uintptr_t(*hFileHandle), dosPath.count, dosPath.data, ToString(STATUS_SUCCESS).data);
+					return STATUS_SUCCESS;
+				}
+				return trueFunc();
+			}
 		}
-		else if ((bufChars >= 7u && wcsncmp(buf, L"\\Device", 7u) == 0) || (bufChars >= 10u && wcsncmp(buf, L"\\Global??\\", 10u) == 0)) // \Global is for FilterConnectCommunicationPort and friends
+
+		if (dosPath.StartsWith(TCV("pipe\\")))
 		{
-			suppressCreateFileDetour = true;
+			if (!g_rules->IsThrowAway(dosPath, g_runningRemote))
+				return trueFunc();
+			//allowDirectoryCache = false;
+			//fileName.Append(TCV("\\\\.\\")).Append(dosPath);
+			DEBUG_LOG_DETOURED(funcName, L"(THROWAWAY) (%.*s) -> STATUS_OBJECT_NAME_NOT_FOUND", dosPath.count, dosPath.data);
+			return STATUS_OBJECT_NAME_NOT_FOUND;
 		}
-		else if (createFileName)
+		else if (!FixPath(fileName, dosPath))
+			UBA_ASSERTF(false, L"FixPath failed for string '%.*s'", dosPath.count, dosPath.data);
+	}
+	else if (ObjectAttributes->RootDirectory)
+	{
+		if (isDetouredHandle(ObjectAttributes->RootDirectory))
 		{
-			if (!CanDetour(createFileName))
-				suppressCreateFileDetour = true;
-			else if (StartsWith(createFileName, TC("\\\\?\\pipe\\")) && !g_rules->IsThrowAway(ToView(createFileName), g_runningRemote))
-				suppressCreateFileDetour = true;
-			else if (!FixPath(fileName, createFileName))
-			{
-				DEBUG_LOG(L"FixPath failed for string '%ls'", createFileName);
-				suppressCreateFileDetour = true;
-			}
-
-			// TODO: Instead of using t_createFileFileName.. should we just resolve the path from ObjectAttributes->RootDirectory?
-			ObjectAttributes->RootDirectory = nullptr;
+			auto& dh = asDetouredHandle(ObjectAttributes->RootDirectory);
+			fileName.Append(dh.fileObject->fileInfo->originalName);
+			rootDir = dh.trueHandle;
 		}
-		else if (bufChars >= 4 && memcmp(buf, L"\\??\\", 8) == 0)
+		else if (isListDirectoryHandle(ObjectAttributes->RootDirectory))
 		{
-			if (!CanDetour(buf))
-				suppressCreateFileDetour = true;
-			else if (!FixPath(fileName, buf, bufChars))
-				UBA_ASSERTF(false, L"FixPath failed for string '%ls'", buf + 4);
+			auto& lh = asListDirectoryHandle(ObjectAttributes->RootDirectory);
+			fileName.Append(lh.originalName);
+			//rootDir = lh.trueHandle;
+			UBA_ASSERT(false);
 		}
 		else
 		{
-			if (ObjectAttributes->RootDirectory)
-			{
-				if (isDetouredHandle(ObjectAttributes->RootDirectory))
-				{
-					auto& dh = asDetouredHandle(ObjectAttributes->RootDirectory);
-					fileName.Append(dh.fileObject->fileInfo->originalName);
-					if (bufChars)
-						fileName.EnsureEndsWithSlash().Append(buf, bufChars);
-					rootDir = dh.trueHandle;
-					ObjectAttributes->RootDirectory = nullptr;
-				}
-				else if (isListDirectoryHandle(ObjectAttributes->RootDirectory))
-				{
-					auto& lh = asListDirectoryHandle(ObjectAttributes->RootDirectory);
-					fileName.Append(lh.originalName);
-					if (bufChars)
-						fileName.EnsureEndsWithSlash().Append(buf, bufChars);
-					ObjectAttributes->RootDirectory = nullptr;
-				}
-				else
-				{
-					tchar rootDirPath[1024];
-					u32 rootDirLen = True_GetFinalPathNameByHandleW(rootDir, rootDirPath, 1024, 0);
-					UBA_ASSERT(rootDirLen != 0 && rootDirLen < sizeof(rootDirPath));
-					UBA_ASSERT(StartsWith(rootDirPath, TC("\\\\?\\")));
-					fileName.Append(rootDirPath + 4, rootDirLen - 4);
-					if (bufChars)
-						fileName.EnsureEndsWithSlash().Append(buf, bufChars);
-				}
-			}
-			else
-			{
-				fileName.Append(buf, bufChars);
-				if (fileName.StartsWith(L"\\DosDevices")) // Something used in msbuild.. 
-					suppressCreateFileDetour = true;
-			}
+			tchar rootDirPath[1024];
+			u32 rootDirLen = True_GetFinalPathNameByHandleW(rootDir, rootDirPath, 1024, 0);
+			UBA_ASSERT(rootDirLen != 0 && rootDirLen < sizeof(rootDirPath));
+			UBA_ASSERT(StartsWith(rootDirPath, TC("\\\\?\\")));
+			fileName.Append(rootDirPath + 4, rootDirLen - 4);
 		}
+		if (ntPath.count)
+			fileName.EnsureEndsWithSlash().Append(ntPath);
+		ObjectAttributes->RootDirectory = nullptr;
 	}
-
-	if (!suppressCreateFileDetour && fileName[fileName.count - 1] == '$')
+	else
 	{
-		constexpr const wchar_t* stdStr[] = { L"conerr$", L"conout$", L"conin$" };
-		for (u32 i = 0; i != 3; ++i)
-		{
-			if (!fileName.EndsWith(stdStr[i]))
-				continue;
-			if (g_isDetachedProcess)
-			{
-				*hFileHandle = g_stdHandle[i];
-				return STATUS_SUCCESS;
-			}
-
-			suppressCreateFileDetour = true;
-			break;
-		}
+		UBA_ASSERT(false);
+		return trueFunc();
 	}
 
-	if (suppressCreateFileDetour)
-	{
-		NTSTATUS res = Local_NtCreateFile(IsCreateFunc, hFileHandle, DesiredAccess, ObjectAttributes, IoStatusBlock, AllocationSize, FileAttributes, ShareAccess, CreateDisposition, CreateOptions, EaBuffer, EaLength);
-		DEBUG_LOG_TRUE(funcName, L"(SUPPRESSDETOUR) %llu (%ls) -> %ls", uintptr_t(*hFileHandle), ObjectAttributes->ObjectName->Buffer, ToString(res).data);
-		if (createFileName && Equals(createFileName, L"NUL"))
-			g_nullFile = *hFileHandle;
-		return res;
-	}
 
 	DevirtualizePath(fileName);
 
@@ -1317,86 +1370,83 @@ NTSTATUS NTAPI Shared_NtCreateFile(bool IsCreateFunc, PHANDLE hFileHandle, ACCES
 	fileNameLower.MakeLower();
 	StringKey fileNameKey = ToStringKey(fileNameLower);
 
-	if (createFileName)// && (keepInMemory || canDetour))
+	if (allowDirectoryCache)
 	{
-		if (g_allowDirectoryCache)
+		// This is an optimization where we populate directory table and use that to figure out if file exists or not..
+		// .. in msvc's case it doesn't matter much because these tables are already up to date when msvc use CreateFile.
+		// .. clang otoh is using CreateFile with tooons of different paths trying to open files.. in remote worker case this becomes super expensive
+		if (!isWrite && !isSystemOrTempFile) // We need to skip SystemTemp.. lots of stuff going on there.
 		{
-			// This is an optimization where we populate directory table and use that to figure out if file exists or not..
-			// .. in msvc's case it doesn't matter much because these tables are already up to date when msvc use CreateFile.
-			// .. clang otoh is using CreateFile with tooons of different paths trying to open files.. in remote worker case this becomes super expensive
-			if (!isWrite && !isSystemOrTempFile) // We need to skip SystemTemp.. lots of stuff going on there.
-			{
-				CHECK_PATH(fileNameLower);
-				dirTableOffset = Rpc_GetEntryOffset(fileNameKey, fileName, checkIfDir);
+			CHECK_PATH(fileNameLower);
+			dirTableOffset = Rpc_GetEntryOffset(fileNameKey, fileName, checkIfDir);
 
-				bool allowEarlyOut = true;
-				if (dirTableOffset == ~u32(0))
+			bool allowEarlyOut = true;
+			if (dirTableOffset == ~u32(0))
+			{
+				// This could be a written file not reported to server yet
 				{
-					// This could be a written file not reported to server yet
-					{
-						SCOPED_READ_LOCK(g_mappedFileTable.m_lookupLock, lock);
-						auto findIt = g_mappedFileTable.m_lookup.find(fileNameKey);
-						if (findIt != g_mappedFileTable.m_lookup.end())
-							allowEarlyOut = findIt->second.deleted;
-					}
-					if (allowEarlyOut)
-					{
-						//SetLastError(ERROR_FILE_NOT_FOUND); // Don't think this is needed
-						*hFileHandle = INVALID_HANDLE_VALUE;
-						DEBUG_LOG_DETOURED(funcName, L"NOTFOUND_USINGTABLE (%ls) -> Error", fileName.data);
+					SCOPED_READ_LOCK(g_mappedFileTable.m_lookupLock, lock);
+					auto findIt = g_mappedFileTable.m_lookup.find(fileNameKey);
+					if (findIt != g_mappedFileTable.m_lookup.end())
+						allowEarlyOut = findIt->second.deleted;
+				}
+				if (allowEarlyOut)
+				{
+					//SetLastError(ERROR_FILE_NOT_FOUND); // Don't think this is needed
+					*hFileHandle = INVALID_HANDLE_VALUE;
+					DEBUG_LOG_DETOURED(funcName, L"NOTFOUND_USINGTABLE (%ls) -> Error", fileName.data);
 
 #if UBA_DEBUG_VALIDATE
-						if (g_validateFileAccess)
-						{
-							SuppressDetourScope _;
-							UBA_ASSERTF(True_GetFileAttributesW(fileName.data) == INVALID_FILE_ATTRIBUTES, L"DIRTABLE claims file %ls does not exist but it does", fileName.data);
-						}
+					if (g_validateFileAccess)
+					{
+						SuppressDetourScope _;
+						UBA_ASSERTF(True_GetFileAttributesW(fileName.data) == INVALID_FILE_ATTRIBUTES, L"DIRTABLE claims file %ls does not exist but it does", fileName.data);
+					}
 #endif
 
-						return STATUS_OBJECT_NAME_NOT_FOUND;
-					}
+					return STATUS_OBJECT_NAME_NOT_FOUND;
 				}
-				else if (!checkIfDir)
+			}
+			else if (!checkIfDir)
+			{
+				// File could have been deleted.
+				DirectoryTable::EntryInformation entryInfo;
+				g_directoryTable.GetEntryInformation(entryInfo, dirTableOffset);
+				if (entryInfo.attributes == 0)
 				{
-					// File could have been deleted.
-					DirectoryTable::EntryInformation entryInfo;
-					g_directoryTable.GetEntryInformation(entryInfo, dirTableOffset);
-					if (entryInfo.attributes == 0)
-					{
-						DEBUG_LOG_DETOURED(funcName, L"DELETED %llu, (%ls) -> Success", uintptr_t(*hFileHandle), fileName.data);
-						return STATUS_OBJECT_NAME_NOT_FOUND;
-					}
-					else if (useContent && IsDirectory(entryInfo.attributes))
-					{
-						DEBUG_LOG_DETOURED(funcName, L"%llu, (%ls) -> STATUS_FILE_IS_A_DIRECTORY", uintptr_t(*hFileHandle), fileName.data);
-						return STATUS_FILE_IS_A_DIRECTORY;
-					}
+					DEBUG_LOG_DETOURED(funcName, L"DELETED %llu, (%ls) -> Success", uintptr_t(*hFileHandle), fileName.data);
+					return STATUS_OBJECT_NAME_NOT_FOUND;
 				}
-
-				bool isWriteAttributes = (DesiredAccess & FILE_WRITE_ATTRIBUTES) != 0;
-
-				// If file is an output file we still allow this path and accept wrong (compressed) size.
-				// This is a bit hacky but we don't want to transfer and decompress file just to get the size
-				if (allowEarlyOut && !useContent && !isWriteAttributes && (!CouldBeCompressedFile(fileName) || g_rules->IsOutputFile(fileName, g_systemTemp)))
+				else if (useContent && IsDirectory(entryInfo.attributes))
 				{
-					auto dh = new DetouredHandle(HandleType_File);
-					dh->fileObject = new FileObject();
-					dh->fileObject->desiredAccess = DesiredAccess;
-					dh->dirTableOffset = dirTableOffset;
-
-					FileInfo* tempFileInfo = new FileInfo();
-					dh->fileObject->fileInfo = tempFileInfo;
-					dh->fileObject->ownsFileInfo = true;
-					dh->fileObject->deleteOnClose = isDeleteOnClose;
-					tempFileInfo->originalName = _wcsdup(fileName.data);
-					tempFileInfo->name = L"GETATTRIBUTES";
-					tempFileInfo->refCount = 1;
-					*hFileHandle = makeDetouredHandle(dh);
-					//SetLastError(ERROR_SUCCESS); // Don't think this is needed
-					DEBUG_LOG_DETOURED(funcName, L"GETATTRIBUTES %llu, (%ls) -> Success", uintptr_t(*hFileHandle), fileName.data);
-					return STATUS_SUCCESS;
+					DEBUG_LOG_DETOURED(funcName, L"%llu, (%ls) -> STATUS_FILE_IS_A_DIRECTORY", uintptr_t(*hFileHandle), fileName.data);
+					return STATUS_FILE_IS_A_DIRECTORY;
 				}
 			}
+
+			bool isWriteAttributes = (DesiredAccess & FILE_WRITE_ATTRIBUTES) != 0;
+
+			// If file is an output file we still allow this path and accept wrong (compressed) size.
+			// This is a bit hacky but we don't want to transfer and decompress file just to get the size
+			if (allowEarlyOut && !useContent && !isWriteAttributes && (!CouldBeCompressedFile(fileName) || g_rules->IsOutputFile(fileName, g_systemTemp)))
+			{
+				auto dh = new DetouredHandle(HandleType_File);
+				dh->fileObject = new FileObject();
+				dh->fileObject->desiredAccess = DesiredAccess;
+				dh->dirTableOffset = dirTableOffset;
+
+				FileInfo* tempFileInfo = new FileInfo();
+				dh->fileObject->fileInfo = tempFileInfo;
+				dh->fileObject->ownsFileInfo = true;
+				dh->fileObject->deleteOnClose = isDeleteOnClose;
+				tempFileInfo->originalName = _wcsdup(fileName.data);
+				tempFileInfo->name = L"GETATTRIBUTES";
+				tempFileInfo->refCount = 1;
+				*hFileHandle = makeDetouredHandle(dh);
+				//SetLastError(ERROR_SUCCESS); // Don't think this is needed
+				DEBUG_LOG_DETOURED(funcName, L"GETATTRIBUTES %llu, (%ls) -> Success", uintptr_t(*hFileHandle), fileName.data);
+				return STATUS_SUCCESS;
+			}
 		}
 	}
 
@@ -2556,17 +2606,22 @@ NTSTATUS Detoured_NtQueryInformationProcess(HANDLE ProcessHandle, PROCESSINFOCLA
 
 	if (g_isRunningWine)
 	{
-		if (ProcessInformationClass == (PROCESSINFOCLASS)49) // ProcessSubsystemInformation
+		if (ProcessInformationClass == (PROCESSINFOCLASS)49) // ProcessConsoleHostProcess
 		{
-			if (ProcessInformationLength < sizeof(ULONG))
+			struct PROCESS_CONSOLE_HOST_PROCESS_INFORMATION
 			{
-				DEBUG_LOG_TRUE(L"NtQueryInformationProcess", L"(ProcessSubsystemInformation) %llu -> %s", uintptr_t(ProcessHandle), ToString(STATUS_INFO_LENGTH_MISMATCH).data);
+				ULONG_PTR ConsoleHostProcess;
+			};
+
+			if (ProcessInformationLength < sizeof(PROCESS_CONSOLE_HOST_PROCESS_INFORMATION))
+			{
+				DEBUG_LOG_TRUE(L"NtQueryInformationProcess", L"(ProcessConsoleHostProcess) %llu -> %s", uintptr_t(ProcessHandle), ToString(STATUS_INFO_LENGTH_MISMATCH).data);
 				return STATUS_INFO_LENGTH_MISMATCH;
 			}
-			*(ULONG*)ProcessInformation = 0; // SubsystemInformationTypeWin32 = 0
+			((PROCESS_CONSOLE_HOST_PROCESS_INFORMATION*)ProcessInformation)->ConsoleHostProcess = 0;
 			if (ReturnLength)
-				*ReturnLength = sizeof(ULONG);
-			DEBUG_LOG_TRUE(L"NtQueryInformationProcess", L"(ProcessSubsystemInformation) %llu -> %s", uintptr_t(ProcessHandle), ToString(STATUS_SUCCESS).data);
+				*ReturnLength = sizeof(PROCESS_CONSOLE_HOST_PROCESS_INFORMATION);
+			DEBUG_LOG_TRUE(L"NtQueryInformationProcess", L"(ProcessConsoleHostProcess) %llu -> %s", uintptr_t(ProcessHandle), ToString(STATUS_SUCCESS).data);
 			return STATUS_SUCCESS;
 		}
 	}
@@ -2851,48 +2906,4 @@ BOOLEAN Detoured_RtlValidateHeap(HANDLE HeapPtr, ULONG Flags, PVOID Block)
 	return True_RtlValidateHeap(HeapPtr, Flags, Block);
 }
 
-NTSTATUS Detoured_RtlDosPathNameToNtPathName_U_WithStatus(PCWSTR dos_path, PUNICODE_STRING ntpath, PWSTR* file_part, VOID* reserved)
-{
-	#if UBA_CHECK_PATH_CONVERSION
-	t_rtlCalled = true;
-	#endif
-	auto res = True_RtlDosPathNameToNtPathName_U_WithStatus(dos_path, ntpath, file_part, reserved);
-	//DEBUG_LOG_TRUE(L"Detoured_RtlDosPathNameToNtPathName_U_WithStatus", L"-> %s", ToString(res).data);
-	return res;
-}
-
-NTSTATUS Detoured_RtlDosPathNameToRelativeNtPathName_U_WithStatus(PCWSTR DosFileName, PUNICODE_STRING NtFileName, PWSTR* FilePart, RTL_RELATIVE_NAME_U* RelativeName)
-{
-	#if UBA_CHECK_PATH_CONVERSION
-	t_rtlCalled = true;
-	#endif
-	//UBA_ASSERT(!FilePart);
-	//UBA_ASSERT(!RelativeName);
-	return True_RtlDosPathNameToRelativeNtPathName_U_WithStatus(DosFileName, NtFileName, FilePart, RelativeName);
-}
-
-DWORD WINAPI Detoured_RtlGetFullPathName_U(const WCHAR* name, ULONG size, WCHAR* buffer, WCHAR** file_part)
-{
-	#if UBA_CHECK_PATH_CONVERSION
-	t_rtlCalled = true;
-	#endif
-	return True_RtlGetFullPathName_U(name, size, buffer, file_part);
-}
-
-ULONG WINAPI Detoured_RtlGetFullPathName_UEx(const WCHAR* name, ULONG size, WCHAR* buffer, WCHAR** file_part, RTL_PATH_TYPE* type)
-{
-	#if UBA_CHECK_PATH_CONVERSION
-	t_rtlCalled = true;
-	#endif
-	return True_RtlGetFullPathName_UEx(name, size, buffer, file_part, type);
-}
-
-BOOLEAN NTAPI Detoured_RtlDosPathNameToRelativeNtPathName_U(PCWSTR DosName, PUNICODE_STRING NtName, PWSTR *FilePart, RTL_RELATIVE_NAME_U* RelativeName)
-{
-	#if UBA_CHECK_PATH_CONVERSION
-	t_rtlCalled = true;
-	#endif
-	return True_RtlDosPathNameToRelativeNtPathName_U(DosName, NtName, FilePart, RelativeName);
-}
-
 #endif
diff --git a/Engine/Source/Programs/UnrealBuildAccelerator/Detours/Private/Windows/UbaDetoursFunctionsUcrtBase.inl b/Engine/Source/Programs/UnrealBuildAccelerator/Detours/Private/Windows/UbaDetoursFunctionsUcrtBase.inl
--- a/Engine/Source/Programs/UnrealBuildAccelerator/Detours/Private/Windows/UbaDetoursFunctionsUcrtBase.inl
+++ b/Engine/Source/Programs/UnrealBuildAccelerator/Detours/Private/Windows/UbaDetoursFunctionsUcrtBase.inl
@@ -155,9 +155,9 @@ errno_t Detoured__waccess_s(const wchar_t* path, int mode)
 	}
 
 
-	if (!CanDetour(path))
+	if (!CanDetour(fixedPath))
 	{
-		auto res = True__waccess_s(path, mode);
+		auto res = True__waccess_s(fixedPath.data, mode);
 		DEBUG_LOG_TRUE(L"_waccess_s", L"(NODETOUR) %ls %i -> %ls", originalPath, mode, WaccessResultToString(res));
 		return res;
 	}
diff --git a/Engine/Source/Programs/UnrealBuildAccelerator/Detours/Private/Windows/UbaDetoursFunctionsWin.cpp b/Engine/Source/Programs/UnrealBuildAccelerator/Detours/Private/Windows/UbaDetoursFunctionsWin.cpp
--- a/Engine/Source/Programs/UnrealBuildAccelerator/Detours/Private/Windows/UbaDetoursFunctionsWin.cpp
+++ b/Engine/Source/Programs/UnrealBuildAccelerator/Detours/Private/Windows/UbaDetoursFunctionsWin.cpp
@@ -200,9 +200,9 @@ bool CouldBeCompressedFile(const StringView& fileName)
 	return g_filesCouldBeCompressed && g_globalRules.FileCanBeCompressed(fileName);
 }
 
-bool CanDetour(const tchar* path)
+bool CanDetour(StringView path)
 {
-	if (t_disallowDetour || !path)
+	if (t_disallowDetour || !path.count)
 		return false;
 
 	if (path[0] == '\\')
@@ -213,13 +213,13 @@ bool CanDetour(const tchar* path)
 			{
 				if (path[5] != ':') // Not file
 					return false;
-				path += 4;
+				path = path.Skip(4);
 			}
 			else if (path[2] == '?') // \\?\ - Win32 path prefix to send through unmodified to nt layer
 			{
 				if (path[3] != '\\')
 					return false;
-				path += 4;
+				path = path.Skip(4);
 			}
 			else
 			{
@@ -232,12 +232,12 @@ bool CanDetour(const tchar* path)
 			if (path[4] == 'U' && path[5] == 'N' && path[6] == 'C') // All network paths ok?
 				return true;
 			if (path[5] == ':')
-				path += 4;
+				path = path.Skip(4);
 			else
 				return false; // Unknown
 		}
 	}
-	if (Equals(path, TC("nul")))
+	if (path.Equals(TC("nul")))
 		return false;
 
 	return g_rules->CanDetour(path, g_runningRemote);
@@ -530,7 +530,7 @@ void OnModuleLoaded(HMODULE moduleHandle, const StringView& name);
 
 // Variables used to communicate state from kernelbase functions to ntdll functions
 thread_local const wchar_t* t_renameFileNewName;
-thread_local const wchar_t* t_createFileFileName;
+thread_local bool t_hackInsideCreateFile;
 
 #include "UbaDetoursFunctionsMiMalloc.inl"
 #include "UbaDetoursFunctionsNtDll.inl"
