[UBA]
* Some fixes caused by previous change where CreateFileW is no longer sending filename to NtCreateFile
* Improved some logging

diff --git a/Engine/Source/Programs/UnrealBuildAccelerator/Detours/Private/Windows/UbaDetoursFunctionsKernelBase.inl b/Engine/Source/Programs/UnrealBuildAccelerator/Detours/Private/Windows/UbaDetoursFunctionsKernelBase.inl
--- a/Engine/Source/Programs/UnrealBuildAccelerator/Detours/Private/Windows/UbaDetoursFunctionsKernelBase.inl
+++ b/Engine/Source/Programs/UnrealBuildAccelerator/Detours/Private/Windows/UbaDetoursFunctionsKernelBase.inl
@@ -1774,7 +1774,7 @@ BOOL Detoured_GetFileInformationByHandleEx(HANDLE hFile, FILE_INFO_BY_HANDLE_CLA
 			MemoryFile* mf = fi.memoryFile;
 			UBA_ASSERTF(mf || fi.name[0] == '$', L"GetFileInformationByHandleEx called on file %s which has no entry offset or real handle", HandleToName(hFile));
 
-			DEBUG_LOG_DETOURED(L"GetFileInformationByHandleEx", L"(MEMORY) (%u) %llu (%ls)", fileInformationClass, uintptr_t(hFile), HandleToName(hFile));
+			DEBUG_LOG_DETOURED(L"GetFileInformationByHandleEx", L"(MEMORY) (%u) %llu (%ls) -> %s", fileInformationClass, uintptr_t(hFile), HandleToName(hFile), ToString(TRUE));
 
 			if (fileInformationClass == FileIdInfo)
 			{
@@ -4358,7 +4358,7 @@ PVOID Detoured_MapViewOfFile3(HANDLE FileMapping, HANDLE Process, PVOID BaseAddr
 {
 	UBA_ASSERT(!isDetouredHandle(FileMapping));
 	UBA_ASSERT(!isDetouredHandle(Process));
-	DEBUG_LOG_TRUE(L"MapViewOfFile3", L"");
+	//DEBUG_LOG_TRUE(L"MapViewOfFile3", L""); // Spammy since SharedMemoryAllocator use it
 	return True_MapViewOfFile3(FileMapping, Process, BaseAddress, Offset, ViewSize, AllocationType, PageProtection, ExtendedParameters, ParameterCount);
 }
 
@@ -4366,7 +4366,7 @@ BOOL Detoured_UnmapViewOfFile2(HANDLE Process, PVOID BaseAddress, ULONG UnmapFla
 {
 	DETOURED_CALL(UnmapViewOfFile2);
 	auto res = True_UnmapViewOfFile2(Process, BaseAddress, UnmapFlags);
-	DEBUG_LOG_TRUE(L"UnmapViewOfFile2", L"0x%llx -> %ls", uintptr_t(BaseAddress), ToString(res));
+	//DEBUG_LOG_TRUE(L"UnmapViewOfFile2", L"0x%llx -> %ls", uintptr_t(BaseAddress), ToString(res)); // Spammy since SharedMemoryAllocator use it
 	return res;
 }
 
diff --git a/Engine/Source/Programs/UnrealBuildAccelerator/Detours/Private/Windows/UbaDetoursFunctionsNtDll.inl b/Engine/Source/Programs/UnrealBuildAccelerator/Detours/Private/Windows/UbaDetoursFunctionsNtDll.inl
--- a/Engine/Source/Programs/UnrealBuildAccelerator/Detours/Private/Windows/UbaDetoursFunctionsNtDll.inl
+++ b/Engine/Source/Programs/UnrealBuildAccelerator/Detours/Private/Windows/UbaDetoursFunctionsNtDll.inl
@@ -93,6 +93,8 @@ static void FixupDosPath(PCWSTR& dosPath, StringBufferBase& temp)
 	switch (RtlDetermineDosPathNameType_U(dosPath))
 	{
 	case RtlPathTypeRelative:
+		if (TStrchr(dosPath, ':')) // There was a colon somewhere later in the path
+			break;
 		temp.Append(g_virtualWorkingDir).EnsureEndsWithSlash().Append(dosPath);
 		dosPath = temp.data;
 		break;
@@ -517,7 +519,11 @@ NTSTATUS NTAPI Detoured_NtQueryDirectoryFile(HANDLE FileHandle, HANDLE Event, PI
 
 	if (isListDirectoryHandle(FileHandle))
 	{
-		DEBUG_LOG_DETOURED(L"NtQueryDirectoryFile", L"%llu %s", u64(FileHandle), HandleToName(FileHandle));
+		auto returnFunc = [&](NTSTATUS res)
+			{
+				DEBUG_LOG_DETOURED(L"NtQueryDirectoryFile", L"%llu %s -> %s", u64(FileHandle), HandleToName(FileHandle), ToString(res).data);
+				return res;
+			};
 
 		IoStatusBlock->Information = 0;
 
@@ -565,11 +571,11 @@ NTSTATUS NTAPI Detoured_NtQueryDirectoryFile(HANDLE FileHandle, HANDLE Event, PI
 		{
 			UBA_ASSERT(false);
 			IoStatusBlock->Status = STATUS_OBJECT_NAME_NOT_FOUND;
-			return STATUS_OBJECT_NAME_NOT_FOUND;
+			return returnFunc(STATUS_OBJECT_NAME_NOT_FOUND);
 		}
 
 		if (Length < structSize)
-			return STATUS_BUFFER_OVERFLOW;
+			return returnFunc(STATUS_BUFFER_OVERFLOW);
 
 		// Fix for buggy program (rustc.exe). even if there are zero entries they still seem to check this variable
 		memset(FileInformation, 0, structSize);
@@ -655,7 +661,7 @@ NTSTATUS NTAPI Detoured_NtQueryDirectoryFile(HANDLE FileHandle, HANDLE Event, PI
 			UBA_ASSERT(res < 0 || wcscmp(info.FileName, info2.FileName) == 0);
 		}
 #endif
-		return res;
+		return returnFunc(res);
 	}
 
 	HANDLE trueHandle = FileHandle;
@@ -1226,20 +1232,20 @@ NTSTATUS NTAPI Shared_NtCreateFile(bool IsCreateFunc, PHANDLE hFileHandle, ACCES
 	#endif
 
 	StringView ntPath(ObjectAttributes->ObjectName->Buffer, ObjectAttributes->ObjectName->Length / 2);
+	bool allowLocal = true;(void)allowLocal;
 
 	auto trueFunc = [&]()
 		{
+			UBA_ASSERT(allowLocal);
 			NTSTATUS res = Local_NtCreateFile(IsCreateFunc, hFileHandle, DesiredAccess, ObjectAttributes, IoStatusBlock, AllocationSize, FileAttributes, ShareAccess, CreateDisposition, CreateOptions, EaBuffer, EaLength);
 			DEBUG_LOG_TRUE(funcName, L"(SUPPRESSDETOUR) %llu (%.*s) -> %s", uintptr_t(*hFileHandle), ntPath.count, ntPath.data, ToString(res).data);
-			if (ntPath.Equals(L"\\??\\NUL"))
-				g_nullFile = *hFileHandle;
 			return res;
 		};
 
 	
 	if (t_disallowCreateFileDetour || t_disallowDetour)
 		return trueFunc();
-	if (!ntPath.count || ntPath[1] != '?')
+	if (ntPath.count && ntPath[1] != '?') // we can have an empty ntPath but RootDirectory set
 		return trueFunc();
 
 	bool allowDirectoryCache = g_allowDirectoryCache && t_hackInsideCreateFile; // TODO: This needs to be fixed.. t_hackInsideCreateFile just happens to work but is wrong!
@@ -1252,12 +1258,22 @@ NTSTATUS NTAPI Shared_NtCreateFile(bool IsCreateFunc, PHANDLE hFileHandle, ACCES
 		UBA_ASSERT(!ObjectAttributes->RootDirectory);
 		StringView dosPath = ntPath.Skip(4);
 
+		if (!dosPath.count)
+			return trueFunc();
+
+		if (dosPath.Equals(TCV("NUL")))
+		{
+			auto res = trueFunc();
+			g_nullFile = *hFileHandle;
+			return res;
+		}
+
 		if (!g_rules->CanDetour(dosPath, g_runningRemote))
 			return trueFunc();
 
 		if (dosPath[dosPath.count - 1] == '$')
 		{
-			constexpr const wchar_t* stdStr[] = { L"conerr$", L"conout$", L"conin$" };
+			constexpr const StringView stdStr[] = { TCV("conerr$"), TCV("conout$"), TCV("conin$") };
 			for (u32 i = 0; i != 3; ++i)
 			{
 				if (!dosPath.EndsWith(stdStr[i]))
@@ -1284,7 +1300,10 @@ NTSTATUS NTAPI Shared_NtCreateFile(bool IsCreateFunc, PHANDLE hFileHandle, ACCES
 			return STATUS_OBJECT_NAME_NOT_FOUND;
 		}
 		else if (!FixPath(fileName, dosPath))
-			UBA_ASSERTF(false, L"FixPath failed for string '%.*s'", dosPath.count, dosPath.data);
+		{
+			DEBUG_LOG_DETOURED(funcName, L"(INVALID) (%.*s) -> STATUS_OBJECT_NAME_INVALID", dosPath.count, dosPath.data);
+			return STATUS_OBJECT_NAME_INVALID;
+		}
 	}
 	else if (ObjectAttributes->RootDirectory)
 	{
@@ -1298,11 +1317,12 @@ NTSTATUS NTAPI Shared_NtCreateFile(bool IsCreateFunc, PHANDLE hFileHandle, ACCES
 		{
 			auto& lh = asListDirectoryHandle(ObjectAttributes->RootDirectory);
 			fileName.Append(lh.originalName);
-			//rootDir = lh.trueHandle;
-			UBA_ASSERT(false);
+			allowLocal = false;
 		}
 		else
 		{
+			if (!ntPath.count)
+				return trueFunc();
 			tchar rootDirPath[1024];
 			u32 rootDirLen = True_GetFinalPathNameByHandleW(rootDir, rootDirPath, 1024, 0);
 			UBA_ASSERT(rootDirLen != 0 && rootDirLen < sizeof(rootDirPath));
@@ -2254,7 +2274,7 @@ NTSTATUS Detoured_NtDeviceIoControlFile(HANDLE FileHandle, HANDLE Event, PIO_APC
 	}
 
 	auto res = True_NtDeviceIoControlFile(trueHandle, Event, ApcRoutine, ApcContext, IoStatusBlock, IoControlCode, InputBuffer, InputBufferLength, OutputBuffer, OutputBufferLength);
-	DEBUG_LOG_TRUE(L"NtDeviceIoControlFile", L"%llu (%s) -> %s", uintptr_t(FileHandle), HandleToName(FileHandle), ToString(res).data);
+	DEBUG_LOG_TRUE(L"NtDeviceIoControlFile", L"(0x%x) %llu (%s) -> %s", IoControlCode, uintptr_t(FileHandle), HandleToName(FileHandle), ToString(res).data);
 	return res;
 }
 
@@ -2751,7 +2771,7 @@ NTSTATUS Detoured_RtlCreateUserThread(HANDLE process, SECURITY_DESCRIPTOR* descr
 NTSTATUS Detoured_NtCreateThreadEx(PHANDLE ThreadHandle, ACCESS_MASK DesiredAccess, POBJECT_ATTRIBUTES ObjectAttributes, HANDLE ProcessHandle, USER_THREAD_START_ROUTINE* StartRoutine, PVOID Argument, ULONG CreateFlags, SIZE_T ZeroBits, SIZE_T StackSize, SIZE_T MaximumStackSize, PS_ATTRIBUTE_LIST* AttributeList)
 {
 	auto res = True_NtCreateThreadEx(ThreadHandle, DesiredAccess, ObjectAttributes, ProcessHandle, StartRoutine, Argument, CreateFlags, ZeroBits, StackSize, MaximumStackSize, AttributeList);
-	DEBUG_LOG_TRUE(L"NtCreateThreadEx", L"-> %s", ToString(res).data);
+	DEBUG_LOG_TRUE(L"NtCreateThreadEx", L"(%llu) -> %s", ThreadHandle?u64(*ThreadHandle):0ull, ToString(res).data);
 	return res;
 }
 
