[UBA]
* Fixed so NtCreateFile is properly handling all devices (\\.\Nsi for example) and use true path for the ones that are unknown

diff --git a/Engine/Source/Programs/UnrealBuildAccelerator/Detours/Private/Windows/UbaDetoursFunctionsNtDll.inl b/Engine/Source/Programs/UnrealBuildAccelerator/Detours/Private/Windows/UbaDetoursFunctionsNtDll.inl
--- a/Engine/Source/Programs/UnrealBuildAccelerator/Detours/Private/Windows/UbaDetoursFunctionsNtDll.inl
+++ b/Engine/Source/Programs/UnrealBuildAccelerator/Detours/Private/Windows/UbaDetoursFunctionsNtDll.inl
@@ -1258,48 +1258,52 @@ NTSTATUS NTAPI Shared_NtCreateFile(bool IsCreateFunc, PHANDLE hFileHandle, ACCES
 		UBA_ASSERT(!ObjectAttributes->RootDirectory);
 		StringView dosPath = ntPath.Skip(4);
 
-		if (!dosPath.count)
+		if (dosPath.count < 2)
 			return trueFunc();
 
-		if (dosPath.Equals(TCV("NUL")))
+		if (dosPath[1] != ':') // No ':' means it is a device
 		{
-			auto res = trueFunc();
-			g_nullFile = *hFileHandle;
-			return res;
-		}
-
-		if (!g_rules->CanDetour(dosPath, g_runningRemote))
-			return trueFunc();
-
-		if (dosPath[dosPath.count - 1] == '$')
-		{
-			constexpr const StringView stdStr[] = { TCV("conerr$"), TCV("conout$"), TCV("conin$") };
-			for (u32 i = 0; i != 3; ++i)
+			if (dosPath.Equals(TCV("NUL")))
 			{
-				if (!dosPath.EndsWith(stdStr[i]))
-					continue;
-				if (g_isDetachedProcess)
+				auto res = trueFunc();
+				g_nullFile = *hFileHandle;
+				return res;
+			}
+			if (dosPath.StartsWith(TCV("pipe\\")))
+			{
+				if (g_rules->IsThrowAway(dosPath, g_runningRemote))
 				{
-					IoStatusBlock->Status = STATUS_SUCCESS;
-					IoStatusBlock->Status = FILE_OPENED;
-					*hFileHandle = g_stdHandle[i];
-					DEBUG_LOG_DETOURED(funcName, L"(STD) %llu (%.*s) -> %s", uintptr_t(*hFileHandle), dosPath.count, dosPath.data, ToString(STATUS_SUCCESS).data);
-					return STATUS_SUCCESS;
+					//allowDirectoryCache = false;
+					//fileName.Append(TCV("\\\\.\\")).Append(dosPath);
+					DEBUG_LOG_DETOURED(funcName, L"(THROWAWAY) (%.*s) -> STATUS_OBJECT_NAME_NOT_FOUND", dosPath.count, dosPath.data);
+					return STATUS_OBJECT_NAME_NOT_FOUND;
 				}
-				return trueFunc();
 			}
+			else if (dosPath[dosPath.count - 1] == '$')
+			{
+				constexpr const StringView stdStr[] = { TCV("conerr$"), TCV("conout$"), TCV("conin$") };
+				for (u32 i = 0; i != 3; ++i)
+				{
+					if (!dosPath.EndsWith(stdStr[i]))
+						continue;
+					if (g_isDetachedProcess)
+					{
+						IoStatusBlock->Status = STATUS_SUCCESS;
+						IoStatusBlock->Status = FILE_OPENED;
+						*hFileHandle = g_stdHandle[i];
+						DEBUG_LOG_DETOURED(funcName, L"(STD) %llu (%.*s) -> %s", uintptr_t(*hFileHandle), dosPath.count, dosPath.data, ToString(STATUS_SUCCESS).data);
+						return STATUS_SUCCESS;
+					}
+					break;
+				}
+			}
+			return trueFunc();
 		}
 
-		if (dosPath.StartsWith(TCV("pipe\\")))
-		{
-			if (!g_rules->IsThrowAway(dosPath, g_runningRemote))
-				return trueFunc();
-			//allowDirectoryCache = false;
-			//fileName.Append(TCV("\\\\.\\")).Append(dosPath);
-			DEBUG_LOG_DETOURED(funcName, L"(THROWAWAY) (%.*s) -> STATUS_OBJECT_NAME_NOT_FOUND", dosPath.count, dosPath.data);
-			return STATUS_OBJECT_NAME_NOT_FOUND;
-		}
-		else if (!FixPath(fileName, dosPath))
+		if (!g_rules->CanDetour(dosPath, g_runningRemote))
+			return trueFunc();
+
+		if (!FixPath(fileName, dosPath))
 		{
 			DEBUG_LOG_DETOURED(funcName, L"(INVALID) (%.*s) -> STATUS_OBJECT_NAME_INVALID", dosPath.count, dosPath.data);
 			return STATUS_OBJECT_NAME_INVALID;
