[UBA]
* Redid changes that shows used memory since we already have that stats from the session and can present that in visualizer instead (will also work when scheduler is not used)

diff --git a/Engine/Source/Programs/UnrealBuildAccelerator/Common/Private/UbaScheduler.cpp b/Engine/Source/Programs/UnrealBuildAccelerator/Common/Private/UbaScheduler.cpp
--- a/Engine/Source/Programs/UnrealBuildAccelerator/Common/Private/UbaScheduler.cpp
+++ b/Engine/Source/Programs/UnrealBuildAccelerator/Common/Private/UbaScheduler.cpp
@@ -1712,7 +1712,7 @@ namespace uba
 		{
 			if (m_memDelaySpawning)
 			{
-				m_session.GetTrace().SchedulerUpdate(false, m_memLastTotal, m_memEstimatedUsage);
+				m_session.GetTrace().SchedulerUpdate(false);
 				m_memDelaySpawning = false;
 			}
 			return true;
@@ -1725,7 +1725,7 @@ namespace uba
 			return false;
 
 		m_memDelaySpawning = true;
-		m_session.GetTrace().SchedulerUpdate(true, m_memLastTotal, m_memEstimatedUsage);
+		m_session.GetTrace().SchedulerUpdate(true);
 
 #if 0//PLATFORM_WINDOWS
 		static bool hasBeenRunOnce;
@@ -2094,9 +2094,6 @@ namespace uba
 				trace.StatusUpdate(StartRow+5, 6, BytesToText(memKillThreshold), LogEntryType_Info);
 			}
 
-			m_memLastTotal = memTotal;
-			trace.SchedulerUpdate(m_memDelaySpawning, memTotal, m_memEstimatedUsage);
-
 			if (memKillThreshold)
 			{
 				if (m_lastKillTime && m_lastKillTime + MsToTime(3000) < GetTime())
diff --git a/Engine/Source/Programs/UnrealBuildAccelerator/Common/Private/UbaTrace.cpp b/Engine/Source/Programs/UnrealBuildAccelerator/Common/Private/UbaTrace.cpp
--- a/Engine/Source/Programs/UnrealBuildAccelerator/Common/Private/UbaTrace.cpp
+++ b/Engine/Source/Programs/UnrealBuildAccelerator/Common/Private/UbaTrace.cpp
@@ -639,12 +639,10 @@ namespace uba
 		BEGIN_TRACE_ENTRY(RemoteExecutionDisabled);
 	}
 
-	void Trace::SchedulerUpdate(bool spawnPaused, u64 memTotal, u64 memEstimated)
+	void Trace::SchedulerUpdate(bool spawnPaused)
 	{
 		BEGIN_TRACE_ENTRY(SchedulerUpdate);
 		writer.WriteBool(spawnPaused);
-		writer.Write7BitEncoded(memTotal);
-		writer.Write7BitEncoded(memEstimated);
 	}
 
 	void Trace::SchedulerKillProcess(u32 processId, StringView reason)
diff --git a/Engine/Source/Programs/UnrealBuildAccelerator/Common/Public/UbaScheduler.h b/Engine/Source/Programs/UnrealBuildAccelerator/Common/Public/UbaScheduler.h
--- a/Engine/Source/Programs/UnrealBuildAccelerator/Common/Public/UbaScheduler.h
+++ b/Engine/Source/Programs/UnrealBuildAccelerator/Common/Public/UbaScheduler.h
@@ -242,7 +242,6 @@ namespace uba
 		Atomic<u64> m_memEstimatedUsage = 0;
 		Atomic<u64> m_memWaitThreshold = 0;
 		Atomic<u64> m_lastKillTime = 0;
-		Atomic<u64> m_memLastTotal = 0;
 		u8 m_memStartWaitPercent = 0;
 		u8 m_memStartKillPercent = 0;
 		u8 m_memTracingLevel = 0;
diff --git a/Engine/Source/Programs/UnrealBuildAccelerator/Common/Public/UbaTrace.h b/Engine/Source/Programs/UnrealBuildAccelerator/Common/Public/UbaTrace.h
--- a/Engine/Source/Programs/UnrealBuildAccelerator/Common/Public/UbaTrace.h
+++ b/Engine/Source/Programs/UnrealBuildAccelerator/Common/Public/UbaTrace.h
@@ -66,7 +66,7 @@ namespace uba
 		#undef UBA_TRACE_TYPE
 	};
 
-	static constexpr u32 TraceVersion = 53;
+	static constexpr u32 TraceVersion = 52;
 	static constexpr u32 TraceReadCompatibilityVersion = 6;
 
 	////////////////////////////////////////////////////////////////////////////////////////////////////
@@ -128,7 +128,7 @@ namespace uba
 		void StatusUpdate2(u32 statusRow, u32 statusColumn, u32 stringId);
 		void DriveUpdate(const tchar drive, u8 busyPercent, u32 readCount, u64 readBytes, u32 writeCount, u64 writeBytes);
 		void RemoteExecutionDisabled();
-		void SchedulerUpdate(bool spawnPaused, u64 memTotal = 0, u64 memEstimated = 0);
+		void SchedulerUpdate(bool spawnPaused);
 		void SchedulerKillProcess(u32 processId, StringView reason);
 
 		void CacheBeginFetch(u32 fetchId, const tchar* description);
diff --git a/Engine/Source/Programs/UnrealBuildAccelerator/Visualizer/Private/UbaTraceReader.cpp b/Engine/Source/Programs/UnrealBuildAccelerator/Visualizer/Private/UbaTraceReader.cpp
--- a/Engine/Source/Programs/UnrealBuildAccelerator/Visualizer/Private/UbaTraceReader.cpp
+++ b/Engine/Source/Programs/UnrealBuildAccelerator/Visualizer/Private/UbaTraceReader.cpp
@@ -78,10 +78,6 @@ namespace uba
 		lastKillReason = nullptr;
 		lastSpawningDelayStartTime = 0;
 		lastSpawningDelayEndTime = 0;
-		memTotal = 0;
-		memEstimated = 0;
-		hasMemoryStats = false;
-		lastSpawnPausedState = false;
 		remoteExecutionDisabled = false;
 		finished = true;
 	};
@@ -1507,26 +1503,13 @@ namespace uba
 		}
 		case TraceType_SchedulerUpdate:
 		{
-			bool spawnPaused = reader.ReadBool();
-
-			if (spawnPaused && !out.lastSpawnPausedState)
+			if (reader.ReadBool())
 			{
 				out.lastSpawningDelayStartTime = time;
 				out.lastSpawningDelayEndTime = ~0llu - MsToTime(10000);
 			}
-			else if (!spawnPaused && out.lastSpawnPausedState)
-			{
+			else
 				out.lastSpawningDelayEndTime = time;
-			}
-			out.lastSpawnPausedState = spawnPaused;
-
-			if (out.version >= 53)
-			{
-				out.memTotal = reader.Read7BitEncoded();
-				out.memEstimated = reader.Read7BitEncoded();
-				if (out.memTotal > 0)
-					out.hasMemoryStats = true;
-			}
 			break;
 		}
 		case TraceType_SchedulerKillProcess:
diff --git a/Engine/Source/Programs/UnrealBuildAccelerator/Visualizer/Private/UbaTraceReader.h b/Engine/Source/Programs/UnrealBuildAccelerator/Visualizer/Private/UbaTraceReader.h
--- a/Engine/Source/Programs/UnrealBuildAccelerator/Visualizer/Private/UbaTraceReader.h
+++ b/Engine/Source/Programs/UnrealBuildAccelerator/Visualizer/Private/UbaTraceReader.h
@@ -222,10 +222,6 @@ namespace uba
 		u32 progressProcessesDone = 0;
 		u32 progressErrorCount = 0;
 		u16 maxActiveProcessCount = 0;
-		u64 memTotal = 0;
-		u64 memEstimated = 0;
-		bool hasMemoryStats = false;
-		bool lastSpawnPausedState = false;
 		bool remoteExecutionDisabled = false;
 		bool finished = true;
 	};
diff --git a/Engine/Source/Programs/UnrealBuildAccelerator/Visualizer/Private/UbaVisualizer.cpp b/Engine/Source/Programs/UnrealBuildAccelerator/Visualizer/Private/UbaVisualizer.cpp
--- a/Engine/Source/Programs/UnrealBuildAccelerator/Visualizer/Private/UbaVisualizer.cpp
+++ b/Engine/Source/Programs/UnrealBuildAccelerator/Visualizer/Private/UbaVisualizer.cpp
@@ -863,27 +863,23 @@ namespace uba
 			u32 start = 3 + 6*activeFont.height;
 			u32 width = activeFont.height * 18;
 			RECT rect;
-			rect.left = start;
-			rect.right = start + width;
-			rect.top = posY;
-			rect.bottom = posY + activeFont.height;
-			context.FillRect(rect, m_processBrushes[0].returned);
 
-			rect.right = rect.left + int(progress*float(width));//durationMs/100;
-			context.FillRect(rect, m_traceView.progressErrorCount ? m_processBrushes[0].error : m_processBrushes[0].success);
-
-			StringBuffer<> str;
-			str.Appendf(TC("%u%%    %u / %u"), u32(progress*100.0f), m_traceView.progressProcessesDone, m_traceView.progressProcessesTotal);
+			{
+				rect.top = posY;
+				rect.bottom = posY + activeFont.height;
 
-			const tchar* remoteDisabled = TC(""); // TODO: Only show this in "advanced mode" m_traceView.remoteExecutionDisabled ? TC("   (Remote spawn disabled)") : TC("");
-			str.Append(remoteDisabled);
+				rect.left = start;
+				rect.right = rect.left + int(progress*float(width));//durationMs/100;
+				context.FillRect(rect, m_traceView.progressErrorCount ? m_processBrushes[0].error : m_processBrushes[0].success);
 
-			if (u32 active = m_traceView.totalProcessActiveCount)
-				str.Appendf(TC("   (%u active)"), active);
+				rect.left = rect.right;
+				rect.right = start + width;
+				context.FillRect(rect, m_processBrushes[0].returned);
+			}
 
 			if (m_traceView.lastKillProcessTime && playTime >= m_traceView.lastKillProcessTime && playTime < m_traceView.lastKillProcessTime + MsToTime(800))
 			{
-				rect.left = start + width + 20;
+				rect.left = rect.right + 20;
 				rect.right = rect.left + activeFont.height * 7;
 				const tchar* killReason = m_traceView.lastKillReason;
 				if (!killReason)
@@ -894,13 +890,14 @@ namespace uba
 				text.Appendf(TC("Killing (%s)"), killReason);
 				drawStatusText(text, LogEntryType_Info, rect.left, rect.right, 0, false, true);
 			}
-			else if (m_traceView.hasMemoryStats && m_traceView.memTotal > 0)
+			else if (!m_traceView.finished && !m_traceView.sessions[0].memAvail.empty())
 			{
-				u32 estPct = u32(m_traceView.memEstimated * 100 / m_traceView.memTotal);
+				u64 memAvail = *m_traceView.sessions[0].memAvail.rbegin();
+				u32 estPct = 100u - u32(memAvail * 100ull / m_traceView.sessions[0].memTotal);
 				bool isThrottling = m_traceView.lastSpawningDelayStartTime && playTime >= m_traceView.lastSpawningDelayStartTime && playTime < m_traceView.lastSpawningDelayEndTime + MsToTime(500);
 
-				rect.left = start + width + 20;
 				u32 barWidth = activeFont.height * 8;
+				rect.left = rect.right + 20;
 				rect.right = rect.left + barWidth;
 
 				// Gray background
@@ -923,14 +920,39 @@ namespace uba
 			}
 			else if (m_traceView.lastSpawningDelayStartTime && playTime >= m_traceView.lastSpawningDelayStartTime && playTime < m_traceView.lastSpawningDelayEndTime + MsToTime(500))
 			{
-				rect.left = start + width + 20;
+				rect.left = rect.right + 20;
 				rect.right = rect.left + activeFont.height * 7;
 				context.FillRect(rect, m_processBrushes[0].warning);
 				drawStatusText(TCV("Throttling (mem)"), LogEntryType_Info, rect.left, rect.right, 0, false, true);
 			}
 
+			#if 0
+			if (!m_traceView.finished && !m_traceView.sessions[0].cpuLoad.empty())
+			{
+				u32 cpuLoad = u32(*m_traceView.sessions[0].cpuLoad.rbegin() * 100.0f);
+				rect.left = rect.right + 20;
+				u32 barWidth = activeFont.height * 8;
+				rect.right = rect.left + barWidth;
+				context.FillRect(rect, m_processBrushes[0].returned);
+				StringBuffer<64> memText;
+				memText.Appendf(TC("Cpu %u%%"), cpuLoad);
+				drawStatusText(memText, LogEntryType_Info, rect.left, rect.right, 0, false, true);
+			}
+			#endif
+
+			// This is last since drawIndentedText line feeds
+			{
+				StringBuffer<> str;
+				str.Appendf(TC("%u%%    %u / %u"), u32(progress*100.0f), m_traceView.progressProcessesDone, m_traceView.progressProcessesTotal);
+
+				const tchar* remoteDisabled = TC(""); // TODO: Only show this in "advanced mode" m_traceView.remoteExecutionDisabled ? TC("   (Remote spawn disabled)") : TC("");
+				str.Append(remoteDisabled);
 
-			drawIndentedText(str, LogEntryType_Info, 6, true);
+				if (u32 active = m_traceView.totalProcessActiveCount)
+					str.Appendf(TC("   (%u active)"), active);
+
+				drawIndentedText(str, LogEntryType_Info, 6, true);
+			}
 		}
 
 		if (m_config.showStatus && !m_traceView.statusMap.empty())
