[UBA]
* Added detouring of RtlIsPackageSid and RtlIsCapabilitySid on wine when running remote

diff --git a/Engine/Source/Programs/UnrealBuildAccelerator/Detours/Private/Windows/UbaDetoursFunctionsWin.cpp b/Engine/Source/Programs/UnrealBuildAccelerator/Detours/Private/Windows/UbaDetoursFunctionsWin.cpp
--- a/Engine/Source/Programs/UnrealBuildAccelerator/Detours/Private/Windows/UbaDetoursFunctionsWin.cpp
+++ b/Engine/Source/Programs/UnrealBuildAccelerator/Detours/Private/Windows/UbaDetoursFunctionsWin.cpp
@@ -943,8 +943,13 @@ void PreInit(const DetoursPayload& payload)
 				FatalError(1351, L"Failed to preload %s (%u)", *it, GetLastError());
 }
 
-#if 0
-static bool PatchImportByName(HMODULE module, const char* importedDll, const char* funcName, void* replacement)
+struct PatchImportInfo
+{
+	const char* funcName;
+	void* replacement;
+};
+
+static bool PatchImportByName(HMODULE module, const char* importedDll, PatchImportInfo* infos, u32 infoCount)
 {
 	auto base = (unsigned char*)module;
 	auto dos  = (IMAGE_DOS_HEADER*)base;
@@ -960,52 +965,59 @@ static bool PatchImportByName(HMODULE module, const char* importedDll, const cha
 		return false;
 
 	auto imp = (IMAGE_IMPORT_DESCRIPTOR*)(base + dir.VirtualAddress);
-
 	for (; imp->Name; ++imp)
 	{
 		const char* dllName = (const char*)(base + imp->Name);
-		if (_stricmp(dllName, importedDll) != 0)
+		if (_stricmp(dllName, importedDll) == 0)
+			break;
+	}
+	if (!imp->Name)
+		return false;
+
+	auto thunk     = (IMAGE_THUNK_DATA*)(base + imp->FirstThunk);
+	auto origThunk = (IMAGE_THUNK_DATA*)(base + imp->OriginalFirstThunk);
+
+	u32 patchCount = 0;
+	for (; origThunk->u1.AddressOfData; ++origThunk, ++thunk)
+	{
+		if (IMAGE_SNAP_BY_ORDINAL(origThunk->u1.Ordinal))
 			continue;
 
-		auto thunk     = (IMAGE_THUNK_DATA*)(base + imp->FirstThunk);
-		auto origThunk = (IMAGE_THUNK_DATA*)(base + imp->OriginalFirstThunk);
+		auto ibn = (IMAGE_IMPORT_BY_NAME*)(base + origThunk->u1.AddressOfData);
 
-		for (; origThunk->u1.AddressOfData; ++origThunk, ++thunk)
+		for (PatchImportInfo* i=infos, *e=i+infoCount; i!=e; ++i)
 		{
-			if (IMAGE_SNAP_BY_ORDINAL(origThunk->u1.Ordinal))
-				continue;
-
-			auto ibn = (IMAGE_IMPORT_BY_NAME*)(base + origThunk->u1.AddressOfData);
-			if (strcmp((const char*)ibn->Name, funcName) != 0)
+			if (strcmp((const char*)ibn->Name, i->funcName) != 0)
 				continue;
 
 			DWORD oldProt;
 			if (!VirtualProtect(&thunk->u1.Function, sizeof(void*), PAGE_READWRITE, &oldProt))
 				return false;
 
-			thunk->u1.Function = (ULONG_PTR)replacement;
+			thunk->u1.Function = (ULONG_PTR)i->replacement;
 
 			DWORD dummy;
 			VirtualProtect(&thunk->u1.Function, sizeof(void*), oldProt, &dummy);
 			FlushInstructionCache(GetCurrentProcess(), nullptr, 0);
-			return true;
+			++patchCount;
+			break;
 		}
 	}
-
-	return false;
+	return patchCount == infoCount;
 }
-#endif
 
 void Init(const DetoursPayload& payload, u64 startTime)
 {
 	AddExceptionHandler();
 	InitMemory();
 	
-	#if 0
-	HMODULE exe = GetModuleHandleW(nullptr);
-	PatchImportByName(exe, "ntdll.dll", "RtlIsPackageSid", (void*)Detoured_RtlIsPackageSid);
-	PatchImportByName(exe, "ntdll.dll", "RtlIsCapabilitySid", (void*)Detoured_RtlIsPackageSid);
-	#endif
+	if (g_isRunningWine && g_runningRemote)
+	{
+		// Wine (11.2) does not have RtlIsPackageSid and RtlIsCapabilitySid implemented and route them to a function that terminates the process
+		HMODULE exe = GetModuleHandleW(nullptr);
+		PatchImportInfo infos[] = { { "RtlIsPackageSid", (void*)Detoured_RtlIsPackageSid }, { "RtlIsCapabilitySid", (void*)Detoured_RtlIsCapabilitySid }, };
+		PatchImportByName(exe, "ntdll.dll", infos, sizeof_array(infos));
+	}
 
 	DetourAttachFunctions(g_runningRemote);
 
