[UBA]
* Fixed so parsing of coff symbols is skipping auxiliary entries. These entries could cause m_imports/m_exports end up storing garbage if we were unlucky and a hash for an aux entry was found first and had the same hash as a real symbol
* Skipping comdat section definition symbols
* Fixed so export symbol index is not increased if symbol is duplicated

diff --git a/Engine/Source/Programs/UnrealBuildAccelerator/Common/Private/UbaObjectFileCoff.cpp b/Engine/Source/Programs/UnrealBuildAccelerator/Common/Private/UbaObjectFileCoff.cpp
--- a/Engine/Source/Programs/UnrealBuildAccelerator/Common/Private/UbaObjectFileCoff.cpp
+++ b/Engine/Source/Programs/UnrealBuildAccelerator/Common/Private/UbaObjectFileCoff.cpp
@@ -291,13 +291,19 @@ namespace uba
 		auto symbols = (SymbolType*)(m_data + m_info.symbolsMemPos);
 		auto sections = (ImageSectionHeader*)(m_data + m_info.sectionsMemOffset);
 
-		for (u32 i=0; i!=m_info.symbolCount; ++i)
+		for (u32 i=0, e=m_info.symbolCount; i!=e; ++i)
 		{
 			auto& symbol = symbols[i];
+			i += symbol.NumberOfAuxSymbols;
 
 			if (symbol.StorageClass != ImageSymClassExternal && symbol.StorageClass != ImageSymClassStatic)
 				continue;
 
+			// Skip COMDAT section definition symbols — these are section headers, not real code/data symbols.
+			// They always have: STATIC storage class, Value=0 (section offset), Type=0 (no type info), 1 aux record.
+			if (symbol.StorageClass == ImageSymClassStatic && symbol.Value == 0 && symbol.Type == 0 && symbol.NumberOfAuxSymbols == 1)
+				continue;
+
 			AnsiStringView symbolName = GetSymbolName(symbol, m_data, m_info.stringTableMemPos);
 
 			if (symbolName[0] == '$') // exception handling etc
@@ -310,17 +316,20 @@ namespace uba
 				if (symbolName[0] == '.') // sections
 					continue;
 
-				bool isData = false;
+				u8 isData = 0;
 				if (symbol.SectionNumber != decltype(symbol.SectionNumber)(-1) && symbol.SectionNumber != decltype(symbol.SectionNumber)(-2))
 				{
 					UBA_ASSERTF(symbol.SectionNumber <= m_info.sectionCount, TC("%i / %i"), symbol.SectionNumber, m_info.sectionCount);
 					auto& section = sections[symbol.SectionNumber-1];
 					if (memcmp(section.Name, ".data", 6) == 0 || memcmp(section.Name, ".rdata", 7) == 0 || memcmp(section.Name, ".bss", 5) == 0)
-						isData = true;
+						isData = 1;
 				}
 				symbolName.ToString(symbolString);
 				StringKey symbolKey = ToStringKeyRaw(symbolString.data(), symbolString.size());
-				m_exports.emplace(symbolKey, ExportInfo{symbolString, isData, index++});
+				auto res = m_exports.emplace(symbolKey, ExportInfo{symbolString, isData, index});
+				UBA_ASSERT(res.second || (res.first->second.isData == isData && res.first->second.symbol == symbolString));
+				if (res.second)
+					++index;
 			}
 			else
 			{
