[UBA]
* Moved more code to Nt* layer to remove duplication

diff --git a/Engine/Source/Programs/UnrealBuildAccelerator/Core/Public/UbaWinTypes.h b/Engine/Source/Programs/UnrealBuildAccelerator/Core/Public/UbaWinTypes.h
--- a/Engine/Source/Programs/UnrealBuildAccelerator/Core/Public/UbaWinTypes.h
+++ b/Engine/Source/Programs/UnrealBuildAccelerator/Core/Public/UbaWinTypes.h
@@ -157,6 +157,16 @@ struct FILE_ALIGNMENT_INFORMATION
 	ULONG AlignmentRequirement;
 };
 
+struct FILE_DISPOSITION_INFORMATION
+{
+	BOOLEAN DeleteFile;
+};
+
+struct FILE_ALLOCATION_INFORMATION
+{
+	LARGE_INTEGER AllocationSize;
+};
+
 struct FILE_ALL_INFORMATION
 {
 	FILE_BASIC_INFORMATION     BasicInformation;
@@ -186,6 +196,7 @@ struct FILE_FS_VOLUME_INFORMATION {
 	FILE_INFO_CLASS(FileDispositionInformation, 13) \
 	FILE_INFO_CLASS(FilePositionInformation, 14) \
 	FILE_INFO_CLASS(FileAllInformation, 18) \
+	FILE_INFO_CLASS(FileAllocationInformation, 19) \
 	FILE_INFO_CLASS(FileEndOfFileInformation, 20) \
 	FILE_INFO_CLASS(FileNetworkOpenInformation, 34) \
 	FILE_INFO_CLASS(FileNormalizedNameInformation, 35) \
diff --git a/Engine/Source/Programs/UnrealBuildAccelerator/Detours/Private/Windows/UbaDetoursFunctionsKernelBase.inl b/Engine/Source/Programs/UnrealBuildAccelerator/Detours/Private/Windows/UbaDetoursFunctionsKernelBase.inl
--- a/Engine/Source/Programs/UnrealBuildAccelerator/Detours/Private/Windows/UbaDetoursFunctionsKernelBase.inl
+++ b/Engine/Source/Programs/UnrealBuildAccelerator/Detours/Private/Windows/UbaDetoursFunctionsKernelBase.inl
@@ -2172,10 +2172,9 @@ BOOL Detoured_SetFileInformationByHandle(HANDLE hFile, FILE_INFO_BY_HANDLE_CLASS
 {
 	DETOURED_CALL(SetFileInformationByHandle);
 
-	if (!isDetouredHandle(hFile))
+	if (g_isRunningWine)
 	{
-		// DOES NOT EXIST IN WINE 10.7
-		if (g_isRunningWine && FileInformationClass == FileRenameInfoEx)
+		if (FileInformationClass == FileRenameInfoEx) // DOES NOT EXIST IN WINE 10.7
 		{
 			auto& info = *(FILE_RENAME_INFO*)lpFileInformation;
 			DWORD flags = info.Flags;
@@ -2187,79 +2186,10 @@ BOOL Detoured_SetFileInformationByHandle(HANDLE hFile, FILE_INFO_BY_HANDLE_CLASS
 			info.Flags = flags;
 			return res;
 		}
-
-		DEBUG_LOG_TRUE(L"SetFileInformationByHandle", L"%llu (%u)", uintptr_t(hFile), FileInformationClass);
-		return True_SetFileInformationByHandle(hFile, FileInformationClass, lpFileInformation, dwBufferSize);
-	}
-
-	auto& dh = asDetouredHandle(hFile);
-	FileObject& fo = *dh.fileObject;
-	FileInfo& fi = *fo.fileInfo;
-
-	if (FileInformationClass == FileRenameInfo)
-	{
-		// We can't call True_SetFileInformationByHandle on remote builds because it resolves relative paths causing them to be relative to remote working dir
-		auto& info = *(FILE_RENAME_INFO*)lpFileInformation;
-		u8 data[2048];
-		UBA_ASSERT(sizeof(FILE_RENAME_INFORMATION) + info.FileNameLength <= sizeof(data));
-		auto& ntInfo = *(FILE_RENAME_INFORMATION*)data;
-		ntInfo.ReplaceIfExists = info.ReplaceIfExists ? TRUE : FALSE;
-		ntInfo.RootDirectory   = info.RootDirectory;
-		ntInfo.FileNameLength  = info.FileNameLength;
-		memcpy(ntInfo.FileName, info.FileName, info.FileNameLength);
-		DEBUG_LOG_TRUE(L"SetFileInformationByHandle", L"%llu (FileRenameInfo)", uintptr_t(hFile));
-		IO_STATUS_BLOCK io;
-		NTSTATUS res = Detoured_NtSetInformationFile(hFile, &io, &ntInfo, dwBufferSize, FileRenameInformation);
-		return SetLastErrorFromNtStatus(res);
-	}
-	else if (FileInformationClass == FileDispositionInfo)
-	{
-		auto& info = *(FILE_DISPOSITION_INFO*)lpFileInformation;
-
-		if (info.DeleteFileW)
+		if (FileInformationClass == FileAllocationInfo) // DOES NOT CALL NtSetFileInformation in wine
 		{
-			DEBUG_LOG_DETOURED(L"SetFileInformationByHandle", L"File is set to be deleted on close (%ls)", HandleToName(hFile));
-			fo.deleteOnClose = true;
-		}
-		else if (fo.deleteOnClose)
-		{
-			DEBUG_LOG_DETOURED(L"SetFileInformationByHandle", L"File is set to NOT be deleted on close (%ls)", HandleToName(hFile));
-			fo.deleteOnClose = false;
-		}
-		else
-		{
-			DEBUG_LOG_DETOURED(L"SetFileInformationByHandle", L"%llu (FileDispositionInfo %u)", uintptr_t(hFile), info.DeleteFileW);
-		}
-
-		if (fo.fileInfo->memoryFile)
-			return true;
-
-		DEBUG_LOG_TRUE(L"SetFileInformationByHandle", L"%llu (FileDispositionInfo)", uintptr_t(hFile));
-		return True_SetFileInformationByHandle(hFile, FileInformationClass, lpFileInformation, dwBufferSize); // In here to be tabbed in log
-	}
-	else if (FileInformationClass == FileAllocationInfo)
-	{
-		if (MemoryFile* mf = fi.memoryFile)
-		{
-			auto& info = *(FILE_ALLOCATION_INFO*)lpFileInformation;
-			mf->EnsureCommitted(dh, info.AllocationSize.QuadPart);
-			DEBUG_LOG_TRUE(L"SetFileInformationByHandle", L"%llu (FileAllocationInfo) Size: %llu", uintptr_t(hFile), info.AllocationSize.QuadPart);
-			return true;
-		}
-	}
-	else if (FileInformationClass == FileEndOfFileInfo)
-	{
-		// TODO: THIS CAN BE REMOVED!!
-
-		if (MemoryFile* mf = fi.memoryFile)
-		{
-			auto& info = *(FILE_END_OF_FILE_INFO*)lpFileInformation;
-			mf->writtenSize = info.EndOfFile.QuadPart;
-			mf->isReported = false;
-			mf->EnsureCommitted(dh, mf->writtenSize);
-			SetLastError(ERROR_SUCCESS);
-			DEBUG_LOG_TRUE(L"SetFileInformationByHandle", L"%llu (FileEndOfFileInfo) End: %llu", uintptr_t(hFile), info.EndOfFile.QuadPart);
-			return true;
+			IO_STATUS_BLOCK io;
+			return SetLastErrorFromNtStatus(Detoured_NtSetInformationFile(hFile, &io, lpFileInformation, dwBufferSize, FileAllocationInformation));
 		}
 	}
 
diff --git a/Engine/Source/Programs/UnrealBuildAccelerator/Detours/Private/Windows/UbaDetoursFunctionsNtDll.inl b/Engine/Source/Programs/UnrealBuildAccelerator/Detours/Private/Windows/UbaDetoursFunctionsNtDll.inl
--- a/Engine/Source/Programs/UnrealBuildAccelerator/Detours/Private/Windows/UbaDetoursFunctionsNtDll.inl
+++ b/Engine/Source/Programs/UnrealBuildAccelerator/Detours/Private/Windows/UbaDetoursFunctionsNtDll.inl
@@ -842,7 +842,7 @@ NTSTATUS Detoured_NtSetInformationFile(HANDLE FileHandle, PIO_STATUS_BLOCK IoSta
 		}
 		return runTrue();
 	}
-	if (FileInformationClass == FileRenameInformation)
+	if (FileInformationClass == FileRenameInformation || FileInformationClass == FileRenameInformationEx)
 	{
 		// We can end up in here through MoveFileEx and SetFileInformationByHandle
 		auto& info = *(FILE_RENAME_INFORMATION*)FileInformation;
@@ -882,6 +882,39 @@ NTSTATUS Detoured_NtSetInformationFile(HANDLE FileHandle, PIO_STATUS_BLOCK IoSta
 		Length = sizeof(FILE_RENAME_INFORMATION) + info2.FileNameLength + 2;
 		return runTrue();
 	}
+	if (FileInformationClass == FileDispositionInformation)
+	{
+		auto& info = *(FILE_DISPOSITION_INFORMATION*)FileInformation;
+		FileObject& fo = *h.fileObject;
+		if (info.DeleteFile)
+		{
+			DEBUG_LOG_DETOURED(L"NtSetInformationFile", L"File is set to be deleted on close (%ls)", HandleToName(FileHandle));
+			fo.deleteOnClose = true;
+		}
+		else if (fo.deleteOnClose)
+		{
+			DEBUG_LOG_DETOURED(L"NtSetInformationFile", L"File is set to NOT be deleted on close (%ls)", HandleToName(FileHandle));
+			fo.deleteOnClose = false;
+		}
+		else
+		{
+			DEBUG_LOG_DETOURED(L"NtSetInformationFile", L"%llu (FileDispositionInfo %u)", uintptr_t(FileHandle), info.DeleteFileW);
+		}
+		if (fo.fileInfo->memoryFile)
+			return true;
+	}
+	if (FileInformationClass == FileAllocationInformation)
+	{
+		FileObject& fo = *h.fileObject;
+		FileInfo& fi = *fo.fileInfo;
+		if (MemoryFile* mf = fi.memoryFile)
+		{
+			auto& info = *(FILE_ALLOCATION_INFORMATION*)FileInformation;
+			mf->EnsureCommitted(h, info.AllocationSize.QuadPart);
+			DEBUG_LOG_TRUE(L"NtSetInformationFile", L"%llu (FileAllocationInfo) Size: %llu", uintptr_t(FileHandle), info.AllocationSize.QuadPart);
+			return STATUS_SUCCESS;
+		}
+	}
 
 	if (trueHandle == INVALID_HANDLE_VALUE)
 	{
