[UBA]
* Improved unit test TestFileMapWrite/TestFileMapWriteRemote

diff --git a/Engine/Source/Programs/UnrealBuildAccelerator/Test/Private/UbaTestSession.cpp b/Engine/Source/Programs/UnrealBuildAccelerator/Test/Private/UbaTestSession.cpp
--- a/Engine/Source/Programs/UnrealBuildAccelerator/Test/Private/UbaTestSession.cpp
+++ b/Engine/Source/Programs/UnrealBuildAccelerator/Test/Private/UbaTestSession.cpp
@@ -98,7 +98,7 @@ namespace uba
 			return true;
 		for (auto& logLine : process.GetLogLines())
 			logger.Error(logLine.text.c_str());
-		return logger.Error(TC("UbaTestApp returned exit code %u"), exitCode);
+		return logger.Error(TC("UbaTestApp returned exit code %i"), int(exitCode));
 	}
 
 	using TestServerSessionFunction = Function<bool(LoggerWithWriter& logger, const StringView& workingDir, SessionServer& sessionServer)>;
@@ -257,13 +257,18 @@ namespace uba
 	{
 		return RunRemote(logger, testRootDir, {}, testFunc, nullptr, deleteAll, serverShouldListen);
 	}
-	bool CreateTextFile(StringBufferBase& outPath, LoggerWithWriter& logger, const tchar* workingDir, const tchar* fileName, const char* text)
+
+	template<typename CharType>
+	bool CreateTextFile(StringBufferBase& outPath, LoggerWithWriter& logger, const tchar* workingDir, const tchar* fileName, const CharType* text)
 	{
 		outPath.Clear().Append(workingDir).EnsureEndsWithSlash().Append(fileName);
 		FileAccessor fr(logger, outPath.data);
 		if (!fr.CreateWrite())
 			return false;
-		fr.Write(text, strlen(text) + 1);
+		u32 len = 0;
+		for (const CharType* it=text; *it; ++it)
+			++len;
+		fr.Write(text, (len + 1)*sizeof(CharType));
 		return fr.Close();
 	}
 
@@ -528,7 +533,7 @@ namespace uba
 				{
 					for (auto& logLine : process.GetLogLines())
 						logger.Error(logLine.text.c_str());
-					return logger.Error(TC("UbaTestApp returned exit code %u"), exitCode);
+					return logger.Error(TC("UbaTestApp returned exit code %i"), int(exitCode));
 				}
 				return true;
 			});
@@ -1005,11 +1010,17 @@ namespace uba
 	bool WriteFileMap(LoggerWithWriter& logger, SessionServer& session, const tchar* workingDir, const RunProcessFunction& runProcess, void* extraData)
 	{
 		StringBuffer<MaxPath> fileR;
-		if (!CreateTextFile(fileR, logger, workingDir, TC("File.h"), "Foo"))
+		if (!CreateTextFile(fileR, logger, workingDir, TC("File.h"), TC("Foo")))
 			return false;
-
-		//auto runProcess2 = [&](ProcessStartInfo& pi) { return session.RunProcess(pi, true, true); };
-		return RunTestApp(logger, workingDir, runProcess, TC("-writeFileMap"));
+		if (!RunTestApp(logger, workingDir, runProcess, TC("-readFileMap=Foo")))
+			return false;
+		if (!RunTestApp(logger, workingDir, runProcess, TC("-writeFileMap=Foo")))
+			return false;
+		if (!RunTestApp(logger, workingDir, runProcess, TC("-writeFileMap=Bar")))
+			return false;
+		if (!RunTestApp(logger, workingDir, runProcess, TC("-readFileMap=Bar")))
+			return false;
+		return true;
 	}
 
 	bool TestFileMapWrite(LoggerWithWriter& logger, const StringBufferBase& testRootDir)
diff --git a/Engine/Source/Programs/UnrealBuildAccelerator/TestApp/Private/Windows/UbaTestAppWin.cpp b/Engine/Source/Programs/UnrealBuildAccelerator/TestApp/Private/Windows/UbaTestAppWin.cpp
--- a/Engine/Source/Programs/UnrealBuildAccelerator/TestApp/Private/Windows/UbaTestAppWin.cpp
+++ b/Engine/Source/Programs/UnrealBuildAccelerator/TestApp/Private/Windows/UbaTestAppWin.cpp
@@ -452,7 +452,22 @@ int wmain(int argc, wchar_t* argv[])
 		if (memcmp(data, L"Bar", 6) != 0)
 			return -1;
 	}
-	else if (wcscmp(argv[1], L"-writeFileMap") == 0)
+	else if (wcsncmp(argv[1], L"-writeFileMap=", 14) == 0)
+	{
+		wchar_t* toWrite = argv[1] + 14;
+		if (WriteTestFile(toWrite, (wcslen(toWrite) + 1)*sizeof(wchar_t), L"File.h") == -1)
+			return -1;
+	}
+	else if (wcsncmp(argv[1], L"-readFileMap=", 13) == 0)
+	{
+		wchar_t* toExpect = argv[1] + 13;
+		wchar_t buf[32];
+		if (ReadTestFile(buf, sizeof(buf), L"File.h") == -1)
+			return -1;
+		if (wcscmp(toExpect, buf) != 0)
+			return LogError(L"-readFileMap did not read what was expected");
+	}
+	else if (wcscmp(argv[1], L"-readWriteFileMap") == 0)
 	{
 		wchar_t buf[32];
 		if (ReadTestFile(buf, 1, L"File.h") == -1)
