[UBA]
* Added a unit test for handling a bad path
* Added some more documentation in comment

diff --git a/Engine/Source/Programs/UnrealBuildAccelerator/Detours/Private/Windows/UbaDetoursFunctionsKernelBase.inl b/Engine/Source/Programs/UnrealBuildAccelerator/Detours/Private/Windows/UbaDetoursFunctionsKernelBase.inl
--- a/Engine/Source/Programs/UnrealBuildAccelerator/Detours/Private/Windows/UbaDetoursFunctionsKernelBase.inl
+++ b/Engine/Source/Programs/UnrealBuildAccelerator/Detours/Private/Windows/UbaDetoursFunctionsKernelBase.inl
@@ -726,7 +726,7 @@ DWORD Detoured_GetFileType(HANDLE hFile)
 	}
 
 	DEBUG_LOG_TRUE(L"GetFileType", L"%llu (%ls)", uintptr_t(hFile), HandleToName(hFile));
-	return True_GetFileType(hFile); // Calling NtQueryVolumeInformationFile
+	return True_GetFileType(hFile); // Calling NtQueryVolumeInformationFile FileFsDeviceInformation
 }
 
 BOOL Shared_GetFileAttributesExW(const tchar* caller, StringView fileName, GET_FILEEX_INFO_LEVELS fInfoLevelId, LPVOID lpFileInformation, LPCWSTR originalName)
diff --git a/Engine/Source/Programs/UnrealBuildAccelerator/Test/Private/UbaTestBasics.cpp b/Engine/Source/Programs/UnrealBuildAccelerator/Test/Private/UbaTestBasics.cpp
--- a/Engine/Source/Programs/UnrealBuildAccelerator/Test/Private/UbaTestBasics.cpp
+++ b/Engine/Source/Programs/UnrealBuildAccelerator/Test/Private/UbaTestBasics.cpp
@@ -188,7 +188,8 @@ namespace uba
 			return logger.Error(TC("FixPath2 (2) failed"));
 		if (!FixPath2(TC("Z:/UEVFS/Root/Engine/Binaries/Win64/UnrealEditor-InterchangeCore-Win64-Debug.pdb::"), 80, workingDir, TStrlen(workingDir), buffer, sizeof_array(buffer), &lengthResult) || lengthResult != 80 || TStrlen(buffer) != 80)
 			return logger.Error(TC("FixPath2 (2) failed"));
-
+		if (FixPath2(TC("'Small vector' optimization: store up to a small number of items on the stack"), ~0u, workingDir, TStrlen(workingDir), buffer, sizeof_array(buffer), &lengthResult))
+			return logger.Error(TC("FixPath2 (2) failed"));
 		#if 0
 		const tchar longPath[] = TC("c:\\Program Files (x86)\\Microsoft Visual Studio\\Shared\\Entity Framework Tools\\NuGet Packages");
 		tchar shortPath[MAX_PATH];
