[UBA]
* Fixed bug where code was assuming mmap returning nullptr was error which is wrong. MAP_FAILED is what should be checked

diff --git a/Engine/Source/Programs/UnrealBuildAccelerator/Core/Private/UbaSharedMemoryView.cpp b/Engine/Source/Programs/UnrealBuildAccelerator/Core/Private/UbaSharedMemoryView.cpp
--- a/Engine/Source/Programs/UnrealBuildAccelerator/Core/Private/UbaSharedMemoryView.cpp
+++ b/Engine/Source/Programs/UnrealBuildAccelerator/Core/Private/UbaSharedMemoryView.cpp
@@ -123,10 +123,12 @@ namespace uba
 		m_memory = (u8*)VirtualAlloc2(GetCurrentProcess(), baseAddress, m_capacity, MEM_RESERVE | MEM_RESERVE_PLACEHOLDER, PAGE_NOACCESS, nullptr, 0);
 		#else
 		m_memory = (u8*)mmap(NULL, m_capacity, PROT_NONE, MAP_PRIVATE | MAP_ANONYMOUS, -1, 0);
+		if (m_memory == MAP_FAILED)
+			m_memory = nullptr;
 		#endif
 
 		UBA_ASSERTF(m_memory, TC("VirtualAlloc2 failed (%s)"), LastErrorToText().data);
-		return m_memory != 0;
+		return m_memory != nullptr;
 	}
 
 	bool SharedMemoryView::AddRequestedMemory(SharedMemoryAllocatorHandle handle, BinaryReader& reader, u32 sliceCount, SharedMemoryMapType mapType, u64 offset, const tchar* hint, u8 init)
@@ -208,6 +210,8 @@ namespace uba
 			protect |= PROT_WRITE;
 		TimerScope ts(kernelStats.mapViewOfFile3);
 		void* viewA = mmap(targetAddress, size, protect, MAP_SHARED | MAP_FIXED, handle.shmFd, offset);
+		if (viewA == MAP_FAILED)
+			viewA = nullptr;
 		#endif
 
 		UBA_ASSERTF(viewA == targetAddress, TC("MapViewOfFile3 failed extending memory with %llu bytes at offset %llu - %s (%s)"), size, offset, LastErrorToText().data, hint);
