[UBA]
* Prepared code to enable directory table overlay for all detoured process. No real changes in this changelist expect some data being sent as u64 instead of u32 from session to detoured process

diff --git a/Engine/Source/Programs/UnrealBuildAccelerator/Common/Private/UbaProcess.cpp b/Engine/Source/Programs/UnrealBuildAccelerator/Common/Private/UbaProcess.cpp
--- a/Engine/Source/Programs/UnrealBuildAccelerator/Common/Private/UbaProcess.cpp
+++ b/Engine/Source/Programs/UnrealBuildAccelerator/Common/Private/UbaProcess.cpp
@@ -301,8 +301,8 @@ namespace uba
 				return false;
 			}
 
-			m_shared.overlayFilesHandle  = memHandle;
-			m_shared.overlayFilesSize = overlaySize;
+			m_shared.overlay.handle = memHandle;
+			m_shared.overlay.size = overlaySize;
 		}
 
 		UBA_ASSERT(m_startInfo.rules);
@@ -843,9 +843,9 @@ namespace uba
 			for (auto& memoryHandle : m_shared.sharedMemoryHandles)
 				m_session.m_sharedMemory.CloseHandle(m_session.m_logger, memoryHandle, TC(""));
 			m_shared.sharedMemoryHandles.clear();
-			if (m_shared.overlayFilesHandle.IsValid())
-				m_session.m_sharedMemory.CloseHandle(m_session.m_logger, m_shared.overlayFilesHandle, TC(""));
-			m_shared.overlayFilesHandle = {};
+			if (m_shared.overlay.handle.IsValid())
+				m_session.m_sharedMemory.CloseHandle(m_session.m_logger, m_shared.overlay.handle, TC(""));
+			m_shared.overlay.handle = {};
 		}
 
 		UBA_ASSERT(m_refCount);
@@ -1023,6 +1023,7 @@ namespace uba
 	{
 		InitMessage msg { *this };
 		InitResponse response;
+
 		CallSession(MessageType_Init, m_session.GetInitResponse(response, msg));
 		writer.WriteU32(m_id);
 		writer.WriteBool(m_parentProcess != nullptr);
@@ -1032,13 +1033,12 @@ namespace uba
 		writer.WriteU64(response.sharedMemoryAllocatorHandle.ToU64());
 		writer.WriteU64(response.permanentFilesHandle.ToU64());
 		writer.WriteU64(response.directoryTableHandle);
-		writer.WriteU32(response.directoryTableSize);
+		writer.WriteU64(ToU64(response.directoryTableSize));
 		writer.WriteU32(response.directoryTableCount);
 		writer.WriteU64(response.mappedFileTableHandle);
 		writer.WriteU32(response.mappedFileTableSize);
 		writer.WriteU32(response.mappedFileTableCount);
-		writer.WriteU64(m_shared.overlayFilesHandle.ToU64());
-		writer.WriteU32(m_shared.overlayFilesSize);
+		writer.WriteU64(m_shared.overlay.handle.ToU64());
 
 		if (!m_startInfo.rootsHandle)
 		{
@@ -1068,7 +1068,7 @@ namespace uba
 		writer.WriteU64(response.size);
 		writer.WriteU32(response.closeId);
 		writer.WriteU32(response.mappedFileTableSize);
-		writer.WriteU32(response.directoryTableSize);
+		writer.WriteU64(ToU64(response.directoryTableSize));
 		return true;
 	}
 
@@ -1114,12 +1114,13 @@ namespace uba
 			reader.ReadString(msg.newName);
 		CloseFileResponse response;
 		CallSession(MessageType_CloseFile, m_session.CloseFile(response, msg));
-		writer.WriteU32(response.directoryTableSize);
+		writer.WriteU64(ToU64(response.directoryTableSize));
 		return true;
 	}
 
 	bool ProcessImpl::HandleDeleteFile(BinaryReader& reader, BinaryWriter& writer)
 	{
+		// TODO: Should be removed and handled by overlay
 		DeleteFileMessage msg { *this };
 		reader.ReadString(msg.fileName);
 		msg.fileNameKey = reader.ReadStringKey();
@@ -1128,12 +1129,13 @@ namespace uba
 		CallSession(MessageType_DeleteFile, m_session.DeleteFile(response, msg));
 		writer.WriteBool(response.result);
 		writer.WriteU32(response.errorCode);
-		writer.WriteU32(response.directoryTableSize);
+		writer.WriteU64(ToU64(response.directoryTableSize));
 		return true;
 	}
 
 	bool ProcessImpl::HandleCopyFile(BinaryReader& reader, BinaryWriter& writer)
 	{
+		// TODO: Should be removed and handled by overlay
 		CopyFileMessage msg{ *this };
 		msg.fromKey = reader.ReadStringKey();
 		reader.ReadString(msg.fromName);
@@ -1145,12 +1147,13 @@ namespace uba
 		writer.WriteString(response.toName);
 		writer.WriteU32(response.closeId);
 		writer.WriteU32(response.errorCode);
-		writer.WriteU32(response.directoryTableSize);
+		writer.WriteU64(ToU64(response.directoryTableSize));
 		return true;
 	}
 
 	bool ProcessImpl::HandleMoveFile(BinaryReader& reader, BinaryWriter& writer)
 	{
+		// TODO: Should be removed and handled by overlay
 		MoveFileMessage msg { *this };
 		msg.fromKey = reader.ReadStringKey();
 		reader.ReadString(msg.fromName);
@@ -1161,12 +1164,13 @@ namespace uba
 		CallSession(MessageType_MoveFile, m_session.MoveFile(response, msg));
 		writer.WriteBool(response.result);
 		writer.WriteU32(response.errorCode);
-		writer.WriteU32(response.directoryTableSize);
+		writer.WriteU64(ToU64(response.directoryTableSize));
 		return true;
 	}
 
 	bool ProcessImpl::HandleChmod(BinaryReader& reader, BinaryWriter& writer)
 	{
+		// TODO: Should be removed and handled by overlay
 		ChmodMessage msg { *this };
 		msg.fileNameKey = reader.ReadStringKey();
 		reader.ReadString(msg.fileName);
@@ -1178,27 +1182,29 @@ namespace uba
 	}
 	bool ProcessImpl::HandleCreateDirectory(BinaryReader& reader, BinaryWriter& writer)
 	{
+		// TODO: Should be removed and handled by overlay
 		CreateDirectoryMessage msg;
 		msg.nameKey = reader.ReadStringKey();
 		reader.ReadString(msg.name);
 		CreateDirectoryResponse response;
-		CallSession(MessageType_CreateDirectory, m_session.CreateDirectory(response, msg));
+		CallSession(MessageType_CreateDirectory, m_session.CreateDirectory(response, msg, nullptr));
 		writer.WriteBool(response.result);
 		writer.WriteU32(response.errorCode);
-		writer.WriteU32(response.directoryTableSize);
+		writer.WriteU64(ToU64(response.directoryTableSize));
 		return true;
 	}
 
 	bool ProcessImpl::HandleRemoveDirectory(BinaryReader& reader, BinaryWriter& writer)
 	{
+		// TODO: Should be removed and handled by overlay
 		RemoveDirectoryMessage msg;
 		msg.nameKey = reader.ReadStringKey();
 		reader.ReadString(msg.name);
 		RemoveDirectoryResponse response;
-		CallSession(MessageType_RemoveDirectory, m_session.RemoveDirectory(response, msg));
+		CallSession(MessageType_RemoveDirectory, m_session.RemoveDirectory(response, msg, &m_shared.overlay));
 		writer.WriteBool(response.result);
 		writer.WriteU32(response.errorCode);
-		writer.WriteU32(response.directoryTableSize);
+		writer.WriteU64(ToU64(response.directoryTableSize));
 		return true;
 	}
 
@@ -1209,7 +1215,7 @@ namespace uba
 		msg.directoryNameKey = reader.ReadStringKey();
 		ListDirectoryResponse response;
 		CallSession(MessageType_ListDirectory, m_session.GetListDirectoryInfo(response, msg.directoryName, msg.directoryNameKey));
-		writer.WriteU32(response.tableSize);
+		writer.WriteU64(ToU64(response.directoryTableSize));
 		writer.WriteU32(response.tableOffset);
 		return true;
 	}
@@ -1218,7 +1224,7 @@ namespace uba
 	{
 		UBA_ASSERT(!m_shared.firstOverflowMessage);
 		bool isInit = reader.ReadBool();
-		writer.WriteU32(m_session.GetDirectoryTableSize());
+		writer.WriteU64(ToU64(m_session.GetDirectoryTableSize(&m_shared.overlay)));
 		writer.WriteU32(m_session.GetFileMappingSize());
 		return HandleGetWrittenFiles(isInit, writer);
 	}
@@ -1673,20 +1679,19 @@ namespace uba
 		// Rewrite overlay.. and let reused process recreate dir table
 		if (!nextProcess.overlayFiles.empty())
 		{
-			UBA_ASSERT(m_shared.overlayFilesHandle.IsValid());
+			UBA_ASSERT(m_shared.overlay.handle.IsValid());
 			auto& sharedMemory = m_session.m_sharedMemory;
-			u8* mem = sharedMemory.MapView(m_shared.overlayFilesHandle, TC("OverlayFiles"), SharedMemoryMapType_ReadWrite);
+			u8* mem = sharedMemory.MapView(m_shared.overlay.handle, TC("OverlayFiles"), SharedMemoryMapType_ReadWrite);
 			memcpy(mem, nextProcess.overlayFiles.data(), nextProcess.overlayFiles.size());
-			sharedMemory.UnmapView(m_shared.overlayFilesHandle, mem);
-			m_shared.overlayFilesSize = u32(nextProcess.overlayFiles.size());
+			sharedMemory.UnmapView(m_shared.overlay.handle, mem);
+			m_shared.overlay.size = u32(nextProcess.overlayFiles.size());
 		}
 
 		writer.WriteString(nextProcess.arguments);
 		writer.WriteString(nextProcess.workingDir);
 		writer.WriteString(nextProcess.description);
 		writer.WriteString(nextProcess.logFile);
-		writer.WriteU32(m_session.GetDirectoryTableSize());
-		writer.WriteU32(m_shared.overlayFilesSize);
+		writer.WriteU64(ToU64(m_session.GetDirectoryTableSize(&m_shared.overlay)));
 		return true;
 	}
 
diff --git a/Engine/Source/Programs/UnrealBuildAccelerator/Common/Private/UbaSession.cpp b/Engine/Source/Programs/UnrealBuildAccelerator/Common/Private/UbaSession.cpp
--- a/Engine/Source/Programs/UnrealBuildAccelerator/Common/Private/UbaSession.cpp
+++ b/Engine/Source/Programs/UnrealBuildAccelerator/Common/Private/UbaSession.cpp
@@ -1511,11 +1511,11 @@ namespace uba
 		return true;
 	}
 
-	u32 Session::RegisterDeleteFile(StringKey fileNameKey, const StringView& fileName)
+	DirTableSize Session::RegisterDeleteFile(StringKey fileNameKey, const StringView& fileName, DirectoryTableOverlay* overlay)
 	{
 		// Remote is not updating its own directory table
 		if (m_runningRemote)
-			return GetDirectoryTableSize();
+			return GetDirectoryTableSize(overlay);
 
 		auto& dirTable = m_directoryTable;
 
@@ -1523,12 +1523,12 @@ namespace uba
 		const tchar* lastSlash;
 		StringBuffer<> dirName;
 		if (!GetDirKey(dirKey, dirName, lastSlash, fileName))
-			return InvalidTableOffset;
+			return {InvalidTableOffset, InvalidTableOffset};
 
 		SCOPED_READ_LOCK(dirTable.m_lookupLock, lookupCs);
 		auto res = dirTable.m_lookup.find(dirKey);
 		if (res == dirTable.m_lookup.end())
-			return 0;
+			return {};
 		DirectoryTable::Directory& dir = res->second;
 		lookupCs.Leave();
 		SCOPED_WRITE_LOCK(dir.lock, dirLock);
@@ -1541,7 +1541,7 @@ namespace uba
 		}
 		
 		if (dir.parseOffset == 2) // Invalid directory, don't need to register this
-			return 0;
+			return {};
 
 		UBA_ASSERTF(dir.parseOffset == 1, TC("Registering deleted file %s with unexpect dir.parseOffset %u "), fileName.data, dir.parseOffset);
 
@@ -1556,7 +1556,7 @@ namespace uba
 		
 		// Does not exist, no need adding to file table
 		if (dir.files.erase(fileNameKey) == 0)
-			return 0;
+			return {};
 
 		u8 temp[1024];
 		u64 written;
@@ -1581,16 +1581,19 @@ namespace uba
 		m_debugLogger->Info(TC("TRACKDEL    %s (Key: %s)"), fileName.data, KeyToString(fileNameKey).data);
 		#endif
 
-		SCOPED_WRITE_LOCK(dirTable.m_memoryLock, memoryLock);
-		u32 writePos = dirTable.m_memorySize;
-		EnsureDirectoryTableMemory(writePos + 8 + written);
-		BinaryWriter writer(m_directoryTableMem + writePos);
-		writer.Write7BitEncoded(written); // Storage size
-		u32 tableOffset = u32(writer.GetPosition()) + sizeof(StringKey);
-		writer.WriteBytes(temp, written);
-		dir.tableOffset = dirTable.m_memorySize + tableOffset;
-		dirTable.m_memorySize += u32(writer.GetPosition());
-		return dirTable.m_memorySize;
+		{
+			SCOPED_WRITE_LOCK(dirTable.m_memoryLock, memoryLock);
+			u32 writePos = dirTable.m_memorySize;
+			EnsureDirectoryTableMemory(writePos + 8 + written);
+			BinaryWriter writer(m_directoryTableMem + writePos);
+			writer.Write7BitEncoded(written); // Storage size
+			u32 tableOffset = u32(writer.GetPosition()) + sizeof(StringKey);
+			writer.WriteBytes(temp, written);
+			dir.tableOffset = dirTable.m_memorySize + tableOffset;
+			dirTable.m_memorySize += u32(writer.GetPosition());
+		}
+
+		return GetDirectoryTableSize(overlay);
 	}
 
 	bool Session::InitOverlayFiles(u8* outMem, u32 capacity, u32& outSize, StringView overlayFiles)
@@ -2346,7 +2349,7 @@ namespace uba
 		StringBuffer<> fixedFilePath;
 		auto key = GetKeyAndFixedName(fixedFilePath, filePath);
 		if (registerToDir)
-			RegisterDeleteFile(key, fixedFilePath);
+			RegisterDeleteFile(key, fixedFilePath, nullptr);
 
 		m_storage.DeleteExternalFileMapping(key, filePath);
 
@@ -2379,7 +2382,7 @@ namespace uba
 		UBA_ASSERT(!m_runningRemote);
 		StringBuffer<> fixedFilePath;
 		auto key = GetKeyAndFixedName(fixedFilePath, filePath);
-		RegisterDeleteFile(key, fixedFilePath);
+		RegisterDeleteFile(key, fixedFilePath, nullptr);
 	}
 
 	bool Session::RegisterNewDirectory(const tchar* directoryPath)
@@ -2867,10 +2870,8 @@ namespace uba
 		out.permanentFilesHandle = m_fileMappingBuffer.GetMemoryHandle();
 
 		out.directoryTableHandle = m_directoryTableHandle.ToU64();
-		{
-			SCOPED_READ_LOCK(m_directoryTable.m_memoryLock, l);
-			out.directoryTableSize = (u32)m_directoryTable.m_memorySize;
-		}
+		out.directoryTableSize = GetDirectoryTableSize(&msg.process.m_shared.overlay);
+
 		{
 			SCOPED_READ_LOCK(m_directoryTable.m_lookupLock, l);
 			out.directoryTableCount = (u32)m_directoryTable.m_lookup.size();
@@ -2887,10 +2888,16 @@ namespace uba
 		return true;
 	}
 
-	u32 Session::GetDirectoryTableSize()
+	DirTableSize Session::GetDirectoryTableSize(DirectoryTableOverlay* overlay)
 	{
+		u32 overlaySize = 0u;
+		if (overlay)
+		{
+			SCOPED_READ_LOCK(overlay->lock, lock);
+			overlaySize = overlay->size;
+		}
 		SCOPED_READ_LOCK(m_directoryTable.m_memoryLock, lock);
-		return m_directoryTable.m_memorySize;
+		return {m_directoryTable.m_memorySize, overlaySize};
 	}
 
 	u32 Session::GetFileMappingSize()
@@ -3180,7 +3187,7 @@ namespace uba
 		auto tableSizeGuard = MakeGuard([&]()
 			{
 				out.mappedFileTableSize = GetFileMappingSize();
-				out.directoryTableSize = GetDirectoryTableSize();
+				out.directoryTableSize = GetDirectoryTableSize(&msg.process.m_shared.overlay);
 			});
 
 		// if ((message.Access & FileAccess.Write) != 0)
@@ -3224,7 +3231,7 @@ namespace uba
 		auto tableSizeGuard = MakeGuard([&]()
 			{
 				out.mappedFileTableSize = GetFileMappingSize();
-				out.directoryTableSize = GetDirectoryTableSize();
+				out.directoryTableSize = GetDirectoryTableSize(&process.m_shared.overlay);
 			});
 
 		if constexpr (!IsWindows)
@@ -3333,12 +3340,14 @@ namespace uba
 		u64 fileSize = 0;
 		u64 lastWriteTime = 0;
 
+		ProcessImpl::Shared& shared = msg.process.m_shared;
+
 		ProcessStartInfo& startInfo = msg.process.m_startInfo;
 		auto& rules = *startInfo.rules;
 
 		if (msg.deleteOnClose)
 		{
-			auto& writtenFiles = msg.process.m_shared.writtenFiles;
+			auto& writtenFiles = shared.writtenFiles;
 			auto findIt2 = writtenFiles.find(activeFile.nameKey);
 			if (findIt2 != writtenFiles.end())
 				writtenFiles.erase(findIt2);
@@ -3358,9 +3367,9 @@ namespace uba
 					msgName = msg.newName;
 			}
 			UBA_ASSERT(key != StringKeyZero);
-			auto& writtenFiles = msg.process.m_shared.writtenFiles;
+			auto& writtenFiles = shared.writtenFiles;
 
-			SCOPED_FUTEX(msg.process.m_shared.writtenFilesLock, writtenLock);
+			SCOPED_FUTEX(shared.writtenFilesLock, writtenLock);
 			auto insres = writtenFiles.try_emplace(key);
 			WrittenFile& writtenFile = insres.first->second;
 
@@ -3491,7 +3500,7 @@ namespace uba
 
 			if (addMapping)
 			{
-				msg.process.m_shared.sharedMemoryHandles.erase(msg.memoryHandle); // memory handle transferred to written files
+				shared.sharedMemoryHandles.erase(msg.memoryHandle); // memory handle transferred to written files
 				writtenFile.memoryHandle = msg.memoryHandle;
 				writtenFile.originalMemoryHandle = msg.memoryHandle;
 				writtenFile.memoryWritten = msg.memoryWritten;
@@ -3518,12 +3527,12 @@ namespace uba
 		{
 			if (!msg.newName.IsEmpty())
 			{
-				RegisterDeleteFile(activeFile.nameKey, activeFile.name);
+				RegisterDeleteFile(activeFile.nameKey, activeFile.name, &shared.overlay);
 				if (RegisterCreateFileForWrite(msg.newNameKey, msg.newName, registerRealFile, fileSize, lastWriteTime) && registerRealFile)
 					TraceWrittenFile(msg.process.m_id, msg.newName, fileSize);
 			}
 			else if (msg.deleteOnClose)
-				RegisterDeleteFile(activeFile.nameKey, activeFile.name);
+				RegisterDeleteFile(activeFile.nameKey, activeFile.name, &shared.overlay);
 			else
 			{
 				if (RegisterCreateFileForWrite(activeFile.nameKey, activeFile.name, registerRealFile, fileSize, lastWriteTime) && registerRealFile)
@@ -3531,7 +3540,7 @@ namespace uba
 			}
 		}
 
-		out.directoryTableSize = GetDirectoryTableSize();
+		out.directoryTableSize = GetDirectoryTableSize(&shared.overlay);
 		return true;
 	}
 
@@ -3551,8 +3560,9 @@ namespace uba
 			m_outputFiles.erase(msg.fileName.data);
 		}
 
+		auto& shared = msg.process.m_shared;
+
 		{
-			auto& shared = msg.process.m_shared;
 			SCOPED_FUTEX(shared.writtenFilesLock, lock);
 			auto findIt = shared.writtenFiles.find(msg.fileNameKey);
 			if (findIt != shared.writtenFiles.end())
@@ -3572,7 +3582,7 @@ namespace uba
 		if (deleteRealFile)
 			out.result = uba::DeleteFileW(msg.fileName.data);
 		out.errorCode = GetLastError();
-		out.directoryTableSize = RegisterDeleteFile(msg.fileNameKey, msg.fileName);
+		out.directoryTableSize = RegisterDeleteFile(msg.fileNameKey, msg.fileName, &shared.overlay);
 		return true;
 	}
 
@@ -3596,11 +3606,11 @@ namespace uba
 	{
 		auto& process = msg.process;
 		bool isMoved = false;
+		auto& shared = process.m_shared;
 		{
-			auto& fs = process.m_shared;
-			SCOPED_FUTEX(fs.writtenFilesLock, writtenLock);
-			auto findIt = fs.writtenFiles.find(msg.fromKey);
-			if (findIt != fs.writtenFiles.end())
+			SCOPED_FUTEX(shared.writtenFilesLock, writtenLock);
+			auto findIt = shared.writtenFiles.find(msg.fromKey);
+			if (findIt != shared.writtenFiles.end())
 			{
 				auto& oldFile = findIt->second;
 				bool isMapping = oldFile.memoryHandle.IsValid();
@@ -3616,13 +3626,13 @@ namespace uba
 				}
 
 				UBA_ASSERT(msg.toKey != StringKeyZero);
-				auto insres = fs.writtenFiles.try_emplace(msg.toKey);
+				auto insres = shared.writtenFiles.try_emplace(msg.toKey);
 				UBA_ASSERTF(insres.second, TC("Failed to add new file %s when moving from %s (destination already exists)"), msg.fromName.data, msg.toName.data);
 				WrittenFile& newFile = insres.first->second;
 				newFile = oldFile;
 				newFile.key = msg.toKey;
 				newFile.name = msg.toName.data;
-				fs.writtenFiles.erase(findIt);
+				shared.writtenFiles.erase(findIt);
 
 				if (isMapping)
 				{
@@ -3648,7 +3658,7 @@ namespace uba
 				TraceWrittenFile(process.m_id, msg.toName, 0);
 
 		out.errorCode = ERROR_SUCCESS;
-		out.directoryTableSize = RegisterDeleteFile(msg.fromKey, msg.fromName);
+		out.directoryTableSize = RegisterDeleteFile(msg.fromKey, msg.fromName, &shared.overlay);
 		return true;
 	}
 
@@ -3668,7 +3678,7 @@ namespace uba
 		return true;
 	}
 
-	bool Session::CreateDirectory(CreateDirectoryResponse& out, const CreateDirectoryMessage& msg)
+	bool Session::CreateDirectory(CreateDirectoryResponse& out, const CreateDirectoryMessage& msg, DirectoryTableOverlay* overlay)
 	{
 		out.result = uba::CreateDirectoryW(msg.name.data);
 
@@ -3690,20 +3700,20 @@ namespace uba
 			WriteDirectoryEntries(dirKey, dirName);
 		}
 
-		out.directoryTableSize = GetDirectoryTableSize();
+		out.directoryTableSize = GetDirectoryTableSize(overlay);
 		return true;
 	}
 
-	bool Session::RemoveDirectory(RemoveDirectoryResponse& out, const RemoveDirectoryMessage& msg)
+	bool Session::RemoveDirectory(RemoveDirectoryResponse& out, const RemoveDirectoryMessage& msg, DirectoryTableOverlay* overlay)
 	{
 		out.result = uba::RemoveDirectoryW(msg.name.data);
 		if (out.result)
-			RegisterDeleteFile(msg.nameKey, msg.name);
+			RegisterDeleteFile(msg.nameKey, msg.name, overlay);
 		else
 			out.errorCode = GetLastError();
 		// This has a race condition. If same directory is removed at the same time the failing one
 		// might send back a directoryTableSize that does not include the delete
-		out.directoryTableSize = GetDirectoryTableSize();
+		out.directoryTableSize = GetDirectoryTableSize(overlay);
 		return true;
 	}
 
@@ -3724,7 +3734,7 @@ namespace uba
 		u32 tableOffset;
 		u32 tableSize = WriteDirectoryEntries(dirKey, dirName, &tableOffset);
 		out.tableOffset = tableOffset;
-		out.tableSize = tableSize;
+		out.directoryTableSize = {tableSize, 0u};
 		return true;
 	}
 
@@ -3941,7 +3951,7 @@ namespace uba
 
 	void Session::PrintSessionStats(Logger& logger)
 	{
-		logger.Info(TC("  DirectoryTable      %7u %9s"), u32(m_directoryTable.m_lookup.size()), BytesToText(GetDirectoryTableSize()).str);
+		logger.Info(TC("  DirectoryTable      %7u %9s"), u32(m_directoryTable.m_lookup.size()), BytesToText(GetDirectoryTableSize(nullptr).main).str);
 		logger.Info(TC("  MappingTable        %7u %9s"), u32(m_fileMappingTableLookup.size()), BytesToText(GetFileMappingSize()).str);
 		logger.Info(TC("  MappingBuffer       %7u %9s"), 1, BytesToText(m_fileMappingBuffer.GetUsed()).str);
 		logger.Info(TC("  ProcessReuseCount   %7u"), m_reuseCount.load());
diff --git a/Engine/Source/Programs/UnrealBuildAccelerator/Common/Private/UbaSessionClient.cpp b/Engine/Source/Programs/UnrealBuildAccelerator/Common/Private/UbaSessionClient.cpp
--- a/Engine/Source/Programs/UnrealBuildAccelerator/Common/Private/UbaSessionClient.cpp
+++ b/Engine/Source/Programs/UnrealBuildAccelerator/Common/Private/UbaSessionClient.cpp
@@ -622,7 +622,7 @@ namespace uba
 		if (casKey == CasKeyZero)
 		{
 			//m_logger.Warning(TC("No casfile found for %s"), fileName.data);
-			out.directoryTableSize = GetDirectoryTableSize();
+			out.directoryTableSize = GetDirectoryTableSize(&process.m_shared.overlay);
 			out.mappedFileTableSize = GetFileMappingSize();
 			out.fileName.Append(fileName);
 			return true;
@@ -762,7 +762,7 @@ namespace uba
 				}
 			}
 
-			out.directoryTableSize = GetDirectoryTableSize();
+			out.directoryTableSize = GetDirectoryTableSize(&process.m_shared.overlay);
 			out.mappedFileTableSize = GetFileMappingSize();
 			out.fileName.Append(newName);
 			out.size = fileSize;
@@ -887,7 +887,7 @@ namespace uba
 		if (out.result)
 			if (!SendUpdateDirectoryTable(reader.Reset()))
 				return false;
-		out.directoryTableSize = GetDirectoryTableSize();
+		out.directoryTableSize = GetDirectoryTableSize(&msg.process.m_shared.overlay);
 		return true;
 	}
 
@@ -915,7 +915,7 @@ namespace uba
 			if (!out.errorCode)
 				if (!SendUpdateDirectoryTable(reader.Reset()))
 					return false;
-			out.directoryTableSize = GetDirectoryTableSize();
+			out.directoryTableSize = GetDirectoryTableSize(&msg.process.m_shared.overlay);
 			return true;
 		}
 		lock.Leave();
@@ -1020,7 +1020,7 @@ namespace uba
 		return true;
 	}
 
-	bool SessionClient::CreateDirectory(CreateDirectoryResponse& out, const CreateDirectoryMessage& msg)
+	bool SessionClient::CreateDirectory(CreateDirectoryResponse& out, const CreateDirectoryMessage& msg, DirectoryTableOverlay* overlay)
 	{
 		// TODO: Cache this if it becomes noisy
 
@@ -1037,12 +1037,13 @@ namespace uba
 			if (!SendUpdateDirectoryTable(reader.Reset()))
 				return false;
 
-		out.directoryTableSize = GetDirectoryTableSize();
+		out.directoryTableSize = GetDirectoryTableSize(overlay);
 		return true;
 	}
 
-	bool SessionClient::RemoveDirectory(RemoveDirectoryResponse& out, const RemoveDirectoryMessage& msg)
+	bool SessionClient::RemoveDirectory(RemoveDirectoryResponse& out, const RemoveDirectoryMessage& msg, DirectoryTableOverlay* overlay)
 	{
+		// TODO: Should update overlay directory
 		StackBinaryWriter<1024> writer;
 		NetworkMessage networkMsg(m_client, ServiceId, SessionMessageType_RemoveDirectory, writer);
 		writer.WriteString(msg.name);
@@ -1056,7 +1057,7 @@ namespace uba
 			if (!SendUpdateDirectoryTable(reader.Reset()))
 				return false;
 
-		out.directoryTableSize = GetDirectoryTableSize();
+		out.directoryTableSize = GetDirectoryTableSize(overlay);
 		return true;
 	}
 
@@ -1128,17 +1129,17 @@ namespace uba
 
 		u32 tableOffset = reader.ReadU32();
 
-		u32 oldTableSize = GetDirectoryTableSize();
+		u32 oldTableSize = GetDirectoryTableSize(nullptr).main;
 		if (!UpdateDirectoryTableFromServer(reader))
 			return false;
-		u32 newTableSize = GetDirectoryTableSize();
+		u32 newTableSize = GetDirectoryTableSize(nullptr).main;
 
 		// Ask for a refresh of hashes straight away since they will likely be asked for by the process doing this query
 		if (oldTableSize != newTableSize)
 			m_waitToSendEvent.Set();
 
 		out.tableOffset = tableOffset;
-		out.tableSize = newTableSize;
+		out.directoryTableSize = {newTableSize, 0u};
 
 		return true;
 	}
@@ -1625,7 +1626,7 @@ namespace uba
 		}
 
 		if (!outStartInfos.empty())
-			if (neededDirectoryTableSize > GetDirectoryTableSize())
+			if (neededDirectoryTableSize > GetDirectoryTableSize(nullptr).main)
 				if (!SendUpdateDirectoryTable(reader.Reset()))
 					return false;
 
@@ -2308,8 +2309,8 @@ namespace uba
 						u8* mem = m_sharedMemory.MapView(memHandle, TC("OverlayFiles"), SharedMemoryMapType_ReadWrite);
 						memcpy(mem, startInfo.overlayFilesData.data(), startInfo.overlayFilesData.size());
 						m_sharedMemory.UnmapView(memHandle, mem);
-						process->m_shared.overlayFilesHandle = memHandle;
-						process->m_shared.overlayFilesSize = u32(startInfo.overlayFilesData.size());
+						process->m_shared.overlay.handle = memHandle;
+						process->m_shared.overlay.size = u32(startInfo.overlayFilesData.size());
 					}
 
 					if (!process->Start(startInfo, true, env, true, exitedRec)) // If false, exitFunction is not called
diff --git a/Engine/Source/Programs/UnrealBuildAccelerator/Common/Private/UbaSessionServer.cpp b/Engine/Source/Programs/UnrealBuildAccelerator/Common/Private/UbaSessionServer.cpp
--- a/Engine/Source/Programs/UnrealBuildAccelerator/Common/Private/UbaSessionServer.cpp
+++ b/Engine/Source/Programs/UnrealBuildAccelerator/Common/Private/UbaSessionServer.cpp
@@ -1590,13 +1590,14 @@ namespace uba
 
 	bool SessionServer::HandleDeleteFile(const ConnectionInfo& connectionInfo, const WorkContext& workContext, BinaryReader& reader, BinaryWriter& writer)
 	{
+		// TODO: This should be removed and replaced with overlay
 		StringKey fileNameKey = reader.ReadStringKey();
 		StringBuffer<> fileName;
 		reader.ReadString(fileName);
 		bool result = uba::DeleteFileW(fileName.data);
 		u32 errorCode = GetLastError();
 		if (result)
-			RegisterDeleteFile(fileNameKey, fileName);
+			RegisterDeleteFile(fileNameKey, fileName, nullptr);
 		writer.WriteBool(result);
 		writer.WriteU32(errorCode);
 		return true;
@@ -1604,6 +1605,7 @@ namespace uba
 
 	bool SessionServer::HandleCopyFile(const ConnectionInfo& connectionInfo, const WorkContext& workContext, BinaryReader& reader, BinaryWriter& writer)
 	{
+		// TODO: This should be removed and replaced with overlay
 		StringKey fromNameKey = reader.ReadStringKey(); (void)fromNameKey;
 		StringBuffer<> fromName;
 		reader.ReadString(fromName);
@@ -1620,10 +1622,11 @@ namespace uba
 
 	bool SessionServer::HandleCreateDirectory(const ConnectionInfo& connectionInfo, const WorkContext& workContext, BinaryReader& reader, BinaryWriter& writer)
 	{
+		// TODO: This should be removed and replaced with overlay
 		CreateDirectoryMessage msg;
 		reader.ReadString(msg.name);
 		CreateDirectoryResponse response;
-		if (!Session::CreateDirectory(response, msg))
+		if (!Session::CreateDirectory(response, msg, nullptr))
 			return false;
 		writer.WriteBool(response.result);
 		writer.WriteU32(response.errorCode);
@@ -1632,10 +1635,11 @@ namespace uba
 
 	bool SessionServer::HandleRemoveDirectory(const ConnectionInfo& connectionInfo, const WorkContext& workContext, BinaryReader& reader, BinaryWriter& writer)
 	{
+		// TODO: This should be removed and replaced with overlay
 		RemoveDirectoryMessage msg;
 		reader.ReadString(msg.name);
 		RemoveDirectoryResponse response;
-		if (!Session::RemoveDirectory(response, msg))
+		if (!Session::RemoveDirectory(response, msg, nullptr))
 			return false;
 		writer.WriteBool(response.result);
 		writer.WriteU32(response.errorCode);
@@ -1818,7 +1822,7 @@ namespace uba
 		}
 		fillLock.Leave();
 
-		u32 neededDirectoryTableSize = GetDirectoryTableSize();
+		u32 neededDirectoryTableSize = GetDirectoryTableSize(nullptr).main;
 		u32 neededHashTableSize;
 		{
 			SCOPED_FUTEX_READ(m_nameToHashLookupLock, l);
@@ -2396,7 +2400,7 @@ namespace uba
 		//m_logger.Info(TC("WritePos: %llu"), session.dirTablePos);
 		writer.WriteU32(session.dirTablePos); // We can figure out on the other side if everything was written based on if the message is full or not.
 
-		u32 toSend = GetDirectoryTableSize() - session.dirTablePos;
+		u32 toSend = GetDirectoryTableSize(nullptr).main - session.dirTablePos;
 		if (toSend == 0)
 			return true;
 
@@ -2584,7 +2588,7 @@ namespace uba
 				MemoryMap map;
 				if (!GetOrCreateMemoryMapFromStorage(map, fileNameKey, fileName.data, findIt->second, memoryMapAlignment))
 					return false;
-				out.directoryTableSize = GetDirectoryTableSize();
+				out.directoryTableSize = GetDirectoryTableSize(&process.m_shared.overlay);
 				out.mappedFileTableSize = GetFileMappingSize();
 				out.fileName.Append(map.name);
 				out.size = map.size;
diff --git a/Engine/Source/Programs/UnrealBuildAccelerator/Common/Public/UbaProcess.h b/Engine/Source/Programs/UnrealBuildAccelerator/Common/Public/UbaProcess.h
--- a/Engine/Source/Programs/UnrealBuildAccelerator/Common/Public/UbaProcess.h
+++ b/Engine/Source/Programs/UnrealBuildAccelerator/Common/Public/UbaProcess.h
@@ -151,8 +151,7 @@ namespace uba
 			UnorderedMap<StringKey, WrittenFile> writtenFiles;
 			OverflowedWrittenFilesMessage* firstOverflowMessage = nullptr;
 			UnorderedSet<SharedMemoryHandle> sharedMemoryHandles; // Track created memory handles. entries can be transferred to WrittenFiles
-			SharedMemoryHandle overlayFilesHandle;
-			u32 overlayFilesSize = 0;
+			DirectoryTableOverlay overlay;
 		};
 
 		u32 m_nativeProcessId = 0;
diff --git a/Engine/Source/Programs/UnrealBuildAccelerator/Common/Public/UbaSession.h b/Engine/Source/Programs/UnrealBuildAccelerator/Common/Public/UbaSession.h
--- a/Engine/Source/Programs/UnrealBuildAccelerator/Common/Public/UbaSession.h
+++ b/Engine/Source/Programs/UnrealBuildAccelerator/Common/Public/UbaSession.h
@@ -54,6 +54,7 @@ namespace uba
 	struct ListDirectoryMessage;
 	struct ListDirectoryResponse;
 	struct WrittenFile;
+	struct DirectoryTableOverlay;
 	struct NextProcessInfo;
 	using RootsHandle = u64;
 
@@ -118,7 +119,7 @@ namespace uba
 
 		// For testing
 		const u8* GetDirectoryTableMemory() { return m_directoryTableMem; }
-		u32 GetDirectoryTableSize();
+		DirTableSize GetDirectoryTableSize(DirectoryTableOverlay* overlay);
 
 		void FlushDeadProcesses();
 
@@ -157,7 +158,7 @@ namespace uba
 		void ThreadGcFileMappings();
 
 		bool RegisterCreateFileForWrite(StringKey fileNameKey, const StringView& fileName, bool registerRealFile, u64 fileSize = 0, u64 lastWriteTime = 0, bool invalidateStorage = true);
-		u32 RegisterDeleteFile(StringKey fileNameKey, const StringView& fileName);
+		DirTableSize RegisterDeleteFile(StringKey fileNameKey, const StringView& fileName, DirectoryTableOverlay* overlay);
 
 		virtual bool PrepareProcess(ProcessImpl& process, bool isChild, StringBufferBase& outRealApplication, const tchar*& outRealWorkingDir);
 		virtual void* GetProcessEnvironmentVariables();
@@ -182,8 +183,8 @@ namespace uba
 		virtual bool CopyFile(CopyFileResponse& out, const CopyFileMessage& msg);
 		virtual bool MoveFile(MoveFileResponse& out, const MoveFileMessage& msg);
 		virtual bool Chmod(ChmodResponse& out, const ChmodMessage& msg);
-		virtual bool CreateDirectory(CreateDirectoryResponse& out, const CreateDirectoryMessage& msg);
-		virtual bool RemoveDirectory(RemoveDirectoryResponse& out, const RemoveDirectoryMessage& msg);
+		virtual bool CreateDirectory(CreateDirectoryResponse& out, const CreateDirectoryMessage& msg, DirectoryTableOverlay* overlay);
+		virtual bool RemoveDirectory(RemoveDirectoryResponse& out, const RemoveDirectoryMessage& msg, DirectoryTableOverlay* overlay);
 		virtual bool GetFullFileName(GetFullFileNameResponse& out, const GetFullFileNameMessage& msg);
 		virtual bool GetLongPathName(GetLongPathNameResponse& out, const GetLongPathNameMessage& msg);
 		virtual bool GetListDirectoryInfo(ListDirectoryResponse& out, const StringView& dirName, const StringKey& dirKey);
@@ -429,7 +430,7 @@ namespace uba
 		SharedMemoryHandle permanentFilesHandle;
 		u64 permanentFilesOffset = 0;
 		u64 directoryTableHandle = 0;
-		u32 directoryTableSize = 0;
+		DirTableSize directoryTableSize;
 		u32 directoryTableCount = 0;
 		u64 mappedFileTableHandle = 0;
 		u32 mappedFileTableSize = 0;
@@ -451,7 +452,7 @@ namespace uba
 		u64 size = InvalidValue;
 		u32 closeId = 0;
 		u32 mappedFileTableSize = 0;
-		u32 directoryTableSize = 0;
+		DirTableSize directoryTableSize;
 	};
 
 	struct CloseFileMessage
@@ -470,7 +471,7 @@ namespace uba
 
 	struct CloseFileResponse
 	{
-		u32 directoryTableSize = 0;
+		DirTableSize directoryTableSize;
 	};
 
 	struct DeleteFileMessage
@@ -485,7 +486,7 @@ namespace uba
 	{
 		bool result = false;
 		u32 errorCode = ~0u;
-		u32 directoryTableSize = 0;
+		DirTableSize directoryTableSize;
 	};
 
 	struct CopyFileMessage
@@ -503,7 +504,7 @@ namespace uba
 		StringBuffer<> toName;
 		u32 closeId = 0;
 		u32 errorCode = ~0u;
-		u32 directoryTableSize = 0;
+		DirTableSize directoryTableSize;
 	};
 
 	struct MoveFileMessage
@@ -522,7 +523,7 @@ namespace uba
 	{
 		bool result = false;
 		u32 errorCode = ~0u;
-		u32 directoryTableSize = 0;
+		DirTableSize directoryTableSize;
 	};
 
 	struct ChmodMessage
@@ -576,7 +577,7 @@ namespace uba
 	{
 		bool result = false;
 		u32 errorCode = 0;
-		u32 directoryTableSize = 0;
+		DirTableSize directoryTableSize;
 	};
 
 	struct RemoveDirectoryMessage
@@ -589,7 +590,7 @@ namespace uba
 	{
 		bool result = false;
 		u32 errorCode = 0;
-		u32 directoryTableSize = 0;
+		DirTableSize directoryTableSize;
 	};
 
 	struct ListDirectoryMessage
@@ -601,7 +602,7 @@ namespace uba
 	struct ListDirectoryResponse
 	{
 		u32 tableOffset = 0;
-		u32 tableSize = 0;
+		DirTableSize directoryTableSize;
 	};
 
 	struct WrittenFile
@@ -632,6 +633,13 @@ namespace uba
 		TString breadcrumbs;
 		Vector<u8> overlayFiles;
 	};
+
+	struct DirectoryTableOverlay
+	{
+		ReaderWriterLock lock;
+		SharedMemoryHandle handle;
+		u32 size = 0;
+	};
 }
 
 template<> struct std::hash<uba::ProcessHandle> { size_t operator()(const uba::ProcessHandle& g) const { return g.GetHash(); } };
diff --git a/Engine/Source/Programs/UnrealBuildAccelerator/Common/Public/UbaSessionClient.h b/Engine/Source/Programs/UnrealBuildAccelerator/Common/Public/UbaSessionClient.h
--- a/Engine/Source/Programs/UnrealBuildAccelerator/Common/Public/UbaSessionClient.h
+++ b/Engine/Source/Programs/UnrealBuildAccelerator/Common/Public/UbaSessionClient.h
@@ -65,8 +65,8 @@ namespace uba
 		virtual bool CopyFile(CopyFileResponse& out, const CopyFileMessage& msg) override;
 		virtual bool MoveFile(MoveFileResponse& out, const MoveFileMessage& msg) override;
 		virtual bool Chmod(ChmodResponse& out, const ChmodMessage& msg) override;
-		virtual bool CreateDirectory(CreateDirectoryResponse& out, const CreateDirectoryMessage& msg) override;
-		virtual bool RemoveDirectory(RemoveDirectoryResponse& out, const RemoveDirectoryMessage& msg) override;
+		virtual bool CreateDirectory(CreateDirectoryResponse& out, const CreateDirectoryMessage& msg, DirectoryTableOverlay* overlay) override;
+		virtual bool RemoveDirectory(RemoveDirectoryResponse& out, const RemoveDirectoryMessage& msg, DirectoryTableOverlay* overlay) override;
 		virtual bool GetFullFileName(GetFullFileNameResponse& out, const GetFullFileNameMessage& msg) override;
 		virtual bool GetLongPathName(GetLongPathNameResponse& out, const GetLongPathNameMessage& msg) override;
 		virtual bool GetListDirectoryInfo(ListDirectoryResponse& out, const StringView& dirName, const StringKey& dirKey) override;
diff --git a/Engine/Source/Programs/UnrealBuildAccelerator/Core/Public/UbaDirectoryTable.h b/Engine/Source/Programs/UnrealBuildAccelerator/Core/Public/UbaDirectoryTable.h
--- a/Engine/Source/Programs/UnrealBuildAccelerator/Core/Public/UbaDirectoryTable.h
+++ b/Engine/Source/Programs/UnrealBuildAccelerator/Core/Public/UbaDirectoryTable.h
@@ -7,503 +7,555 @@
 
 namespace uba
 {
+	////////////////////////////////////////////////////////////////////////////////////////////////////
+
 	enum : u64 { InvalidTableOffset = 0 };
 	enum : u64 { OverlayTableOffset = 256*1024*1024 };
 
+	struct DirTableSize { u32 main = 0; u32 overlay = 0; };
+	inline DirTableSize FromU64(u64 v) { return {u32(v), u32(v >> 32u)}; }
+	inline u64 ToU64(DirTableSize size) { return size.main + (u64(size.overlay) << 32llu); }
+
+	////////////////////////////////////////////////////////////////////////////////////////////////////
+
 	class DirectoryTable
 	{
 	public:
 		using EntryLookup = GrowingUnorderedMap<StringKey, u32>;
-		struct Directory { Directory(MemoryBlock& block) : files(block) {} u32 tableOffset = InvalidTableOffset; u32 parseOffset = InvalidTableOffset; EntryLookup files; ReaderWriterLock lock; };
 
-		void Init(const u8* mem, u32 tableCount, u32 tableSize)
+		struct Directory
 		{
-			m_memory = mem;
-			m_lookup.reserve(tableCount + 100);
-			m_memoryBlock.CommitNoLock(tableCount*(sizeof(GrowingUnorderedMap<StringKey, Directory>::value_type)+16), TC(""));
-			ParseDirectoryTable(tableSize);
-		}
+			Directory(MemoryBlock& block) : files(block) {}
+			u32 tableOffset = InvalidTableOffset;
+			u32 parseOffset = InvalidTableOffset;
+			EntryLookup files;
+			ReaderWriterLock lock;
+		};
+
 
-		void InitOverlay(const u8* mem, u32 tableSize)
+		inline void Init(const u8* mem, u32 tableCount, u32 tableSize);
+		inline void InitOverlay(const u8* mem, u32 tableSize);
+		inline void DeinitOverlay(const u8* oldMem, u32 size);
+		inline void ParseDirectoryTable(DirTableSize size);
+		inline void ParseDirectoryTableNoLock(DirTableSize size);
+		inline void ParseDirectoryTableNoLock(const u8* memory, u32 from, u32 to, u32 additionalOffset = 0);
+		inline void PopulateDirectory(const StringKeyHasher& hasher, Directory& dir);
+		inline bool IsDirectoryUpToDate(Directory& dir);
+		inline void PopulateDirectoryNoLock(const StringKeyHasher& hasher, Directory& dir);
+		inline BinaryReader GetReader(u32 tableOffset);
+		inline void PopulateDirectoryRecursive(const StringKeyHasher& hasher, u32 tableOffset, u32 parseOffset, EntryLookup& files);
+		inline void PopulateDirectoryWithFiles(BinaryReader& reader, const StringKeyHasher& hasher, EntryLookup& files);
+
+		enum Exists
 		{
-			m_overlayMemory = mem;
-			ParseDirectoryTableNoLock(mem, 0, tableSize, OverlayTableOffset);
-		}
+			Exists_Yes,
+			Exists_No,
+			Exists_Maybe,
+		};
+
+		inline Exists EntryExists(StringKey entryKey, const StringView& entryName, bool checkIfDir = false, u32* tableOffset = nullptr);
+		inline Exists EntryExistsNoLock(StringKey entryKey, const StringView& entryName, bool checkIfDir = false, u32* tableOffset = nullptr);
+		inline Exists EntryExists(const StringView& str, bool checkIfDir = false);
 
-		void DeinitOverlay(const u8* oldMem, u32 size)
+		struct EntryInformation
 		{
-			if (!size)
-				return;
+			u32 attributes;
+			u32 volumeSerial;
+			u64 fileIndex;
+			u64 size;
+			u64 lastWrite;
+		};
 
-			m_overlayMemory = oldMem;
-			for (auto i=m_lookup.begin(), e=m_lookup.end(); i!=e;)
-			{
-				auto& dir = i->second;
-				if (dir.tableOffset < OverlayTableOffset)
-				{
-					++i;
-					continue;
-				}
+		inline u32 GetAttributes(u32 tableOffset);
+		inline u32 GetEntryInformation(EntryInformation& outInfo, u32 tableOffset, tchar* outFileName = nullptr, u32 fileNameCapacity = 0);
+		inline void GetFinalPath(StringBufferBase& out, const tchar* path);
 
-				while (true)
-				{
-					BinaryReader reader(GetReader(dir.tableOffset));
-					dir.tableOffset = u32(reader.Read7BitEncoded());
-					if (dir.tableOffset < OverlayTableOffset)
-						break;
-				}
-				if (dir.tableOffset == InvalidTableOffset)
-				{
-					i = m_lookup.erase(i);
-					continue;
-				}
+		#if PLATFORM_WINDOWS
+		template<typename Func> inline void TraverseFilesRecursiveNoLock(const StringBufferBase& path, const Func& func);
+		template<typename Func> inline void TraverseAllFilesNoLock(const Func& func);
+		#endif
 
-				dir.parseOffset = dir.tableOffset;
+		inline DirectoryTable(MemoryBlock& block);
+		DirectoryTable(const DirectoryTable&) = delete;
+		void operator=(const DirectoryTable&) = delete;
 
-				for (auto j=dir.files.begin(); j!=dir.files.end();)
-				{
-					if (j->second < OverlayTableOffset)
-						++j;
-					else
-						j = dir.files.erase(j);
-				}
+		MemoryBlock& m_memoryBlock;
+		ReaderWriterLock m_lookupLock;
+		GrowingUnorderedMap<StringKey, Directory> m_lookup;
 
-				++i;
-			}
-			m_overlayMemory = nullptr;
-		}
+		ReaderWriterLock m_memoryLock;
+		const u8* m_memory = nullptr;
+		u32 m_memorySize = 0;
 
-		void ParseDirectoryTable(u32 size)
-		{
-			SCOPED_WRITE_LOCK(m_lookupLock, lock);
-			ParseDirectoryTableNoLock(size);
-		}
+		const u8* m_overlayMemory = nullptr;
+	};
 
-		void ParseDirectoryTableNoLock(u32 size)
-		{
-			if (size <= m_memorySize)
-				return;
-			ParseDirectoryTableNoLock(m_memory, m_memorySize, size);
-			m_memorySize = size;
-		}
+	////////////////////////////////////////////////////////////////////////////////////////////////////
+
+	void DirectoryTable::Init(const u8* mem, u32 tableCount, u32 tableSize)
+	{
+		m_memory = mem;
+		m_lookup.reserve(tableCount + 100);
+		m_memoryBlock.CommitNoLock(tableCount*(sizeof(GrowingUnorderedMap<StringKey, Directory>::value_type)+16), TC(""));
+		ParseDirectoryTable({tableSize, 0u});
+	}
+
+	void DirectoryTable::InitOverlay(const u8* mem, u32 tableSize)
+	{
+		m_overlayMemory = mem;
+		ParseDirectoryTableNoLock(mem, 0, tableSize, OverlayTableOffset);
+	}
+
+	void DirectoryTable::DeinitOverlay(const u8* oldMem, u32 size)
+	{
+		if (!size)
+			return;
 
-		void ParseDirectoryTableNoLock(const u8* memory, u32 from, u32 to, u32 additionalOffset = 0)
+		m_overlayMemory = oldMem;
+		for (auto i=m_lookup.begin(), e=m_lookup.end(); i!=e;)
 		{
-			BinaryReader reader(memory, from, to);
+			auto& dir = i->second;
+			if (dir.tableOffset < OverlayTableOffset)
+			{
+				++i;
+				continue;
+			}
+
 			while (true)
 			{
-				u64 pos = reader.GetPosition();
-				if (pos == to)
+				BinaryReader reader(GetReader(dir.tableOffset));
+				dir.tableOffset = u32(reader.Read7BitEncoded());
+				if (dir.tableOffset < OverlayTableOffset)
 					break;
-				UBA_ASSERTF(pos < to, TC("Should never read past size (pos: %llu, size: %u)"), pos, to);
-				u64 storageSize = reader.Read7BitEncoded();
-				StringKey dirKey = reader.ReadStringKey();
-				auto insres = m_lookup.try_emplace(dirKey, m_memoryBlock); // Note that this is allowed to overwrite
-				u32& tableOffset = insres.first->second.tableOffset;
-				UBA_ASSERTF(tableOffset < OverlayTableOffset || additionalOffset == OverlayTableOffset, TC("Can't add non-overlay directory items after overlay items"));
-				tableOffset = u32(reader.GetPosition()) + additionalOffset;
-				reader.Skip(storageSize - sizeof(dirKey));
 			}
-		}
+			if (dir.tableOffset == InvalidTableOffset)
+			{
+				i = m_lookup.erase(i);
+				continue;
+			}
 
-		void PopulateDirectory(const StringKeyHasher& hasher, Directory& dir)
-		{
-			SCOPED_WRITE_LOCK(dir.lock, lock);
-			PopulateDirectoryNoLock(hasher, dir);
-		}
+			dir.parseOffset = dir.tableOffset;
 
-		bool IsDirectoryUpToDate(Directory& dir)
-		{
-			return dir.parseOffset == dir.tableOffset;
-		}
+			for (auto j=dir.files.begin(); j!=dir.files.end();)
+			{
+				if (j->second < OverlayTableOffset)
+					++j;
+				else
+					j = dir.files.erase(j);
+			}
 
-		void PopulateDirectoryNoLock(const StringKeyHasher& hasher, Directory& dir)
-		{
-			if (dir.parseOffset == dir.tableOffset)
-				return;
-			PopulateDirectoryRecursive(hasher, dir.tableOffset, dir.parseOffset, dir.files);
-			dir.parseOffset = dir.tableOffset;
+			++i;
 		}
+		m_overlayMemory = nullptr;
+	}
+
+	void DirectoryTable::ParseDirectoryTable(DirTableSize size)
+	{
+		SCOPED_WRITE_LOCK(m_lookupLock, lock);
+		ParseDirectoryTableNoLock(size);
+	}
+
+	void DirectoryTable::ParseDirectoryTableNoLock(DirTableSize size)
+	{
+		if (size.main <= m_memorySize)
+			return;
+		ParseDirectoryTableNoLock(m_memory, m_memorySize, size.main);
+		m_memorySize = size.main;
+	}
 
-		BinaryReader GetReader(u32 tableOffset)
+	void DirectoryTable::ParseDirectoryTableNoLock(const u8* memory, u32 from, u32 to, u32 additionalOffset)
+	{
+		BinaryReader reader(memory, from, to);
+		while (true)
 		{
-			if (tableOffset < OverlayTableOffset)
-				return BinaryReader(m_memory, tableOffset);
-			else
-				return { m_overlayMemory - OverlayTableOffset, tableOffset};
+			u64 pos = reader.GetPosition();
+			if (pos == to)
+				break;
+			UBA_ASSERTF(pos < to, TC("Should never read past size (pos: %llu, size: %u)"), pos, to);
+			u64 storageSize = reader.Read7BitEncoded();
+			StringKey dirKey = reader.ReadStringKey();
+			auto insres = m_lookup.try_emplace(dirKey, m_memoryBlock); // Note that this is allowed to overwrite
+			u32& tableOffset = insres.first->second.tableOffset;
+			UBA_ASSERTF(tableOffset < OverlayTableOffset || additionalOffset == OverlayTableOffset, TC("Can't add non-overlay directory items after overlay items"));
+			tableOffset = u32(reader.GetPosition()) + additionalOffset;
+			reader.Skip(storageSize - sizeof(dirKey));
 		}
+	}
 
-		void PopulateDirectoryRecursive(const StringKeyHasher& hasher, u32 tableOffset, u32 parseOffset, EntryLookup& files)
-		{
-			BinaryReader reader(GetReader(tableOffset));
+	void DirectoryTable::PopulateDirectory(const StringKeyHasher& hasher, Directory& dir)
+	{
+		SCOPED_WRITE_LOCK(dir.lock, lock);
+		PopulateDirectoryNoLock(hasher, dir);
+	}
+
+	bool DirectoryTable::IsDirectoryUpToDate(Directory& dir)
+	{
+		return dir.parseOffset == dir.tableOffset;
+	}
+
+	void DirectoryTable::PopulateDirectoryNoLock(const StringKeyHasher& hasher, Directory& dir)
+	{
+		if (dir.parseOffset == dir.tableOffset)
+			return;
+		PopulateDirectoryRecursive(hasher, dir.tableOffset, dir.parseOffset, dir.files);
+		dir.parseOffset = dir.tableOffset;
+	}
+
+	BinaryReader DirectoryTable::GetReader(u32 tableOffset)
+	{
+		if (tableOffset < OverlayTableOffset)
+			return BinaryReader(m_memory, tableOffset);
+		else
+			return { m_overlayMemory - OverlayTableOffset, tableOffset};
+	}
+
+	void DirectoryTable::PopulateDirectoryRecursive(const StringKeyHasher& hasher, u32 tableOffset, u32 parseOffset, EntryLookup& files)
+	{
+		BinaryReader reader(GetReader(tableOffset));
 
-			u32 prevTableOffset = u32(reader.Read7BitEncoded());
+		u32 prevTableOffset = u32(reader.Read7BitEncoded());
 
-			u32 buffer[48*1024];
-			u32 count = 0;
-			u32* readerOffsets = buffer;
+		u32 buffer[48*1024];
+		u32 count = 0;
+		u32* readerOffsets = buffer;
+		readerOffsets[count++] = u32(reader.GetPosition());
+		bool firstIsRoot = true;
+		while (true)
+		{
+			if (prevTableOffset == InvalidTableOffset || prevTableOffset == parseOffset)
+			{
+				firstIsRoot = prevTableOffset == InvalidTableOffset;
+				break;
+			}
+			reader = GetReader(prevTableOffset);
+			prevTableOffset = u32(reader.Read7BitEncoded());
 			readerOffsets[count++] = u32(reader.GetPosition());
-			bool firstIsRoot = true;
-			while (true)
+			if (count == sizeof_array(buffer))
 			{
-				if (prevTableOffset == InvalidTableOffset || prevTableOffset == parseOffset)
-				{
-					firstIsRoot = prevTableOffset == InvalidTableOffset;
-					break;
-				}
-				reader = GetReader(prevTableOffset);
-				prevTableOffset = u32(reader.Read7BitEncoded());
-				readerOffsets[count++] = u32(reader.GetPosition());
-				if (count == sizeof_array(buffer))
-				{
-					readerOffsets = new u32[1024*1024]; // This sucks, but somethings the directory is huuuge. Ideally these files should be spread out over multiple directories
-					memcpy(readerOffsets, buffer, sizeof(buffer));
-				}
+				readerOffsets = new u32[1024*1024]; // This sucks, but somethings the directory is huuuge. Ideally these files should be spread out over multiple directories
+				memcpy(readerOffsets, buffer, sizeof(buffer));
 			}
+		}
 
-			for (u32 i=count; i>0; --i)
+		for (u32 i=count; i>0; --i)
+		{
+			reader = GetReader(readerOffsets[i-1]);
+			if (firstIsRoot)
 			{
-				reader = GetReader(readerOffsets[i-1]);
-				if (firstIsRoot)
-				{
-					firstIsRoot = false;
-					u32 attr = reader.ReadFileAttributes();
-					if (!attr)
-						continue;
-					reader.ReadVolumeSerial(); // Directory volume serial
-					reader.ReadFileIndex(); // Directory file index
-				}
-
-				PopulateDirectoryWithFiles(reader, hasher, files);
+				firstIsRoot = false;
+				u32 attr = reader.ReadFileAttributes();
+				if (!attr)
+					continue;
+				reader.ReadVolumeSerial(); // Directory volume serial
+				reader.ReadFileIndex(); // Directory file index
 			}
 
-			if (readerOffsets != buffer)
-				delete[] readerOffsets;
+			PopulateDirectoryWithFiles(reader, hasher, files);
 		}
 
-		void PopulateDirectoryWithFiles(BinaryReader& reader, const StringKeyHasher& hasher, EntryLookup& files)
-		{
-			u64 itemCount = reader.Read7BitEncoded();
+		if (readerOffsets != buffer)
+			delete[] readerOffsets;
+	}
 
-			files.reserve(files.size() + itemCount);
-	
-			StringBuffer<> filename;
-			filename.Append(PathSeparator);
+	void DirectoryTable::PopulateDirectoryWithFiles(BinaryReader& reader, const StringKeyHasher& hasher, EntryLookup& files)
+	{
+		u64 itemCount = reader.Read7BitEncoded();
 
-			while (itemCount--)
-			{
-				u32 offset = u32(reader.GetPosition());
-				filename.Resize(1);
-				reader.ReadString(filename);
-				if (CaseInsensitiveFs)
-					filename.MakeLower();
-				u32 attr = reader.ReadFileAttributes();
-				reader.ReadVolumeSerial();
-				reader.ReadFileIndex();
-				if (!IsDirectory(attr))
-				{
-					reader.ReadFileTime();
-					reader.ReadFileSize();
-				}
+		files.reserve(files.size() + itemCount);
 
-				StringKey filenameKey = ToStringKey(hasher, filename.data, filename.count);
-				files[filenameKey] = offset; // Always write, since same file might have been added with new info
-			}
-		}
+		StringBuffer<> filename;
+		filename.Append(PathSeparator);
 
-		enum Exists
+		while (itemCount--)
 		{
-			Exists_Yes,
-			Exists_No,
-			Exists_Maybe,
-		};
+			u32 offset = u32(reader.GetPosition());
+			filename.Resize(1);
+			reader.ReadString(filename);
+			if (CaseInsensitiveFs)
+				filename.MakeLower();
+			u32 attr = reader.ReadFileAttributes();
+			reader.ReadVolumeSerial();
+			reader.ReadFileIndex();
+			if (!IsDirectory(attr))
+			{
+				reader.ReadFileTime();
+				reader.ReadFileSize();
+			}
 
-		Exists EntryExists(StringKey entryKey, const StringView& entryName, bool checkIfDir = false, u32* tableOffset = nullptr)
-		{
-			SCOPED_READ_LOCK(m_lookupLock, lock);
-			return EntryExistsNoLock(entryKey, entryName, checkIfDir, tableOffset);
+			StringKey filenameKey = ToStringKey(hasher, filename.data, filename.count);
+			files[filenameKey] = offset; // Always write, since same file might have been added with new info
 		}
+	}
 
-		Exists EntryExistsNoLock(StringKey entryKey, const StringView& entryName, bool checkIfDir = false, u32* tableOffset = nullptr)
+	DirectoryTable::Exists DirectoryTable::EntryExists(StringKey entryKey, const StringView& entryName, bool checkIfDir, u32* tableOffset)
+	{
+		SCOPED_READ_LOCK(m_lookupLock, lock);
+		return EntryExistsNoLock(entryKey, entryName, checkIfDir, tableOffset);
+	}
+
+	DirectoryTable::Exists DirectoryTable::EntryExistsNoLock(StringKey entryKey, const StringView& entryName, bool checkIfDir, u32* tableOffset)
+	{
+		u32 startSkip = 2;
+		if (checkIfDir)
 		{
-			u32 startSkip = 2;
-			if (checkIfDir)
+			auto findIt = m_lookup.find(entryKey);
+			if (findIt != m_lookup.end())
 			{
-				auto findIt = m_lookup.find(entryKey);
-				if (findIt != m_lookup.end())
-				{
-					if (tableOffset)
-						*tableOffset = u32(findIt->second.tableOffset) | 0x80000000; // Use significant bit to say that this is a dir
-					return Exists_Yes;
-				}
-				startSkip = 1;
+				if (tableOffset)
+					*tableOffset = u32(findIt->second.tableOffset) | 0x80000000; // Use significant bit to say that this is a dir
+				return Exists_Yes;
 			}
+			startSkip = 1;
+		}
 
-			// Scan backwards
-			const tchar* rend = entryName.data;
-			const tchar* rit = rend + entryName.count - startSkip;
+		// Scan backwards
+		const tchar* rend = entryName.data;
+		const tchar* rit = rend + entryName.count - startSkip;
 
-			bool inAncestor = false;
-			while (rit > rend)
+		bool inAncestor = false;
+		while (rit > rend)
+		{
+			if (*rit != PathSeparator)
 			{
-				if (*rit != PathSeparator)
-				{
-					--rit;
+				--rit;
 
-					#if !PLATFORM_WINDOWS
-					if (rit != rend) // We want to test empty for non-windows
-					#endif
-						continue;
-				}
+#if !PLATFORM_WINDOWS
+				if (rit != rend) // We want to test empty for non-windows
+#endif
+					continue;
+			}
 
-				u64 sublen = u64(rit - rend);
+			u64 sublen = u64(rit - rend);
 
-				StringKeyHasher ancestorHasher;
-				ancestorHasher.Update(rend, sublen);
-				StringKey ancestorKey = ToStringKey(ancestorHasher);
-				auto dirIt = m_lookup.find(ancestorKey);
-				if (dirIt != m_lookup.end())
+			StringKeyHasher ancestorHasher;
+			ancestorHasher.Update(rend, sublen);
+			StringKey ancestorKey = ToStringKey(ancestorHasher);
+			auto dirIt = m_lookup.find(ancestorKey);
+			if (dirIt != m_lookup.end())
+			{
+				DirectoryTable::Directory& parentDir = dirIt->second;
+				if (parentDir.tableOffset == -1)
+					return Exists_No;
+				if (parentDir.parseOffset != parentDir.tableOffset)
 				{
-					DirectoryTable::Directory& parentDir = dirIt->second;
-					if (parentDir.tableOffset == -1)
-						return Exists_No;
+					SCOPED_WRITE_LOCK(parentDir.lock, lock);
 					if (parentDir.parseOffset != parentDir.tableOffset)
 					{
-						SCOPED_WRITE_LOCK(parentDir.lock, lock);
-						if (parentDir.parseOffset != parentDir.tableOffset)
-						{
-							PopulateDirectoryRecursive(ancestorHasher, parentDir.tableOffset, parentDir.parseOffset, parentDir.files);
-							parentDir.parseOffset = parentDir.tableOffset;
-						}
+						PopulateDirectoryRecursive(ancestorHasher, parentDir.tableOffset, parentDir.parseOffset, parentDir.files);
+						parentDir.parseOffset = parentDir.tableOffset;
 					}
+				}
 
-					SCOPED_READ_LOCK(parentDir.lock, lock);
-					auto entryIt = parentDir.files.find(entryKey);
-					if (entryIt == parentDir.files.end())
+				SCOPED_READ_LOCK(parentDir.lock, lock);
+				auto entryIt = parentDir.files.find(entryKey);
+				if (entryIt == parentDir.files.end())
+					return Exists_No;
+				if (inAncestor)
+				{
+					BinaryReader reader(m_memory, entryIt->second, m_memorySize);
+					reader.SkipString();
+					if (!IsDirectory(reader.ReadFileAttributes()))
 						return Exists_No;
-					if (inAncestor)
-					{
-						BinaryReader reader(m_memory, entryIt->second, m_memorySize);
-						reader.SkipString();
-						if (!IsDirectory(reader.ReadFileAttributes()))
-							return Exists_No;
-						return Exists_Maybe;
-					}
-					if (tableOffset)
-						*tableOffset = entryIt->second;
-					return Exists_Yes;
+					return Exists_Maybe;
 				}
-
-				entryKey = ancestorKey;
-				--rit;
-				inAncestor = true;
+				if (tableOffset)
+					*tableOffset = entryIt->second;
+				return Exists_Yes;
 			}
-			return Exists_Maybe;
-		}
 
-		Exists EntryExists(const StringView& str, bool checkIfDir = false)
-		{
-			StringBuffer<> str2(str);
-			if (str2[str2.count-1] == PathSeparator)
-				str2.Resize(str2.count-1);
-			if (CaseInsensitiveFs)
-				str2.MakeLower();
-			return EntryExists(ToStringKey(str2), str2, checkIfDir, nullptr);
+			entryKey = ancestorKey;
+			--rit;
+			inAncestor = true;
 		}
+		return Exists_Maybe;
+	}
 
-		struct EntryInformation
-		{
-			u32 attributes;
-			u32 volumeSerial;
-			u64 fileIndex;
-			u64 size;
-			u64 lastWrite;
-		};
-
-		u32 GetAttributes(u32 tableOffset)
+	DirectoryTable::Exists DirectoryTable::EntryExists(const StringView& str, bool checkIfDir)
+	{
+		StringBuffer<> str2(str);
+		if (str2[str2.count-1] == PathSeparator)
+			str2.Resize(str2.count-1);
+		if (CaseInsensitiveFs)
+			str2.MakeLower();
+		return EntryExists(ToStringKey(str2), str2, checkIfDir, nullptr);
+	}
+
+	u32 DirectoryTable::GetAttributes(u32 tableOffset)
+	{
+		if (tableOffset & 0x80000000)
 		{
-			if (tableOffset & 0x80000000)
+			tableOffset = tableOffset & ~0x80000000;
+			BinaryReader reader(m_memory, tableOffset);
+			u64 prevTableOffset = reader.Read7BitEncoded();
+			while (prevTableOffset != InvalidTableOffset)
 			{
-				tableOffset = tableOffset & ~0x80000000;
-				BinaryReader reader(m_memory, tableOffset);
-				u64 prevTableOffset = reader.Read7BitEncoded();
-				while (prevTableOffset != InvalidTableOffset)
-				{
-					reader.SetPosition(prevTableOffset);
-					prevTableOffset = reader.Read7BitEncoded();
-				}
-				return reader.ReadFileAttributes();
+				reader.SetPosition(prevTableOffset);
+				prevTableOffset = reader.Read7BitEncoded();
 			}
-			BinaryReader reader(m_memory, tableOffset);
-			reader.SkipString();
 			return reader.ReadFileAttributes();
 		}
+		BinaryReader reader(m_memory, tableOffset);
+		reader.SkipString();
+		return reader.ReadFileAttributes();
+	}
 
-		u32 GetEntryInformation(EntryInformation& outInfo, u32 tableOffset, tchar* outFileName = nullptr, u32 fileNameCapacity = 0)
+	u32 DirectoryTable::GetEntryInformation(EntryInformation& outInfo, u32 tableOffset, tchar* outFileName, u32 fileNameCapacity)
+	{
+		if (tableOffset & 0x80000000)
 		{
-			if (tableOffset & 0x80000000)
-			{
-				tableOffset = tableOffset & ~0x80000000;
-				BinaryReader reader(GetReader(tableOffset));
-				u64 prevTableOffset = reader.Read7BitEncoded();
-				while (prevTableOffset != InvalidTableOffset)
-				{
-					reader.SetPosition(prevTableOffset);
-					prevTableOffset = reader.Read7BitEncoded();
-				}
-				outInfo.attributes = reader.ReadFileAttributes();
-				if (outInfo.attributes)
-				{
-					outInfo.volumeSerial = reader.ReadVolumeSerial();
-					outInfo.fileIndex = reader.ReadFileIndex();
-				}
-				outInfo.size = 0;
-				outInfo.lastWrite = 0;
-				UBA_ASSERT(!outFileName);
-				return ~u32(0);
-			}
-
+			tableOffset = tableOffset & ~0x80000000;
 			BinaryReader reader(GetReader(tableOffset));
-			if (outFileName)
-				reader.ReadString(outFileName, fileNameCapacity);
-			else
-				reader.SkipString();
-			outInfo.attributes = reader.ReadFileAttributes();
-			outInfo.volumeSerial = reader.ReadVolumeSerial();
-			outInfo.fileIndex = reader.ReadFileIndex();
-			if (IsDirectory(outInfo.attributes))
+			u64 prevTableOffset = reader.Read7BitEncoded();
+			while (prevTableOffset != InvalidTableOffset)
 			{
-				outInfo.size = 0;
-				outInfo.lastWrite = 0;
+				reader.SetPosition(prevTableOffset);
+				prevTableOffset = reader.Read7BitEncoded();
 			}
-			else
+			outInfo.attributes = reader.ReadFileAttributes();
+			if (outInfo.attributes)
 			{
-				outInfo.lastWrite = reader.ReadFileTime();
-				outInfo.size = reader.ReadFileSize();
+				outInfo.volumeSerial = reader.ReadVolumeSerial();
+				outInfo.fileIndex = reader.ReadFileIndex();
 			}
-			return u32(reader.GetPosition());
+			outInfo.size = 0;
+			outInfo.lastWrite = 0;
+			UBA_ASSERT(!outFileName);
+			return ~u32(0);
 		}
 
-		void GetFinalPath(StringBufferBase& out, const tchar* path)
+		BinaryReader reader(GetReader(tableOffset));
+		if (outFileName)
+			reader.ReadString(outFileName, fileNameCapacity);
+		else
+			reader.SkipString();
+		outInfo.attributes = reader.ReadFileAttributes();
+		outInfo.volumeSerial = reader.ReadVolumeSerial();
+		outInfo.fileIndex = reader.ReadFileIndex();
+		if (IsDirectory(outInfo.attributes))
+		{
+			outInfo.size = 0;
+			outInfo.lastWrite = 0;
+		}
+		else
 		{
-			UBA_ASSERT(IsAbsolutePath(path));
+			outInfo.lastWrite = reader.ReadFileTime();
+			outInfo.size = reader.ReadFileSize();
+		}
+		return u32(reader.GetPosition());
+	}
 
-			Directory* directory = nullptr;
-			const tchar* prevSlash = TStrchr(path+3, PathSeparator);
-			if (!prevSlash)
-			{
-				// return root directory as-is
-				out.Append(path);
-				return;
-			}
-			out.Append(path, u64(prevSlash - path));
-			const tchar* end = path + TStrlen(path);
+	void DirectoryTable::GetFinalPath(StringBufferBase& out, const tchar* path)
+	{
+		UBA_ASSERT(IsAbsolutePath(path));
 
-			StringBuffer<> forHash;
-			forHash.Append(path, u64(prevSlash - path));
+		Directory* directory = nullptr;
+		const tchar* prevSlash = TStrchr(path+3, PathSeparator);
+		if (!prevSlash)
+		{
+			// return root directory as-is
+			out.Append(path);
+			return;
+		}
+		out.Append(path, u64(prevSlash - path));
+		const tchar* end = path + TStrlen(path);
+
+		StringBuffer<> forHash;
+		forHash.Append(path, u64(prevSlash - path));
+		if (CaseInsensitiveFs)
+			forHash.MakeLower();
+
+		StringKeyHasher hasher;
+		hasher.Update(forHash.data, forHash.count);
+
+		SCOPED_READ_LOCK(m_lookupLock, lock);
+		while (true)
+		{
+			const tchar* slash = TStrchr(prevSlash + 1, PathSeparator);
+			if (!slash)
+				slash = end;
+
+			forHash.Clear().Append(prevSlash, u64(slash - prevSlash));
 			if (CaseInsensitiveFs)
 				forHash.MakeLower();
-
-			StringKeyHasher hasher;
 			hasher.Update(forHash.data, forHash.count);
+			StringKey fileNameKey = ToStringKey(hasher);
 
-			SCOPED_READ_LOCK(m_lookupLock, lock);
-			while (true)
+			if (directory)
 			{
-				const tchar* slash = TStrchr(prevSlash + 1, PathSeparator);
-				if (!slash)
-					slash = end;
-
-				forHash.Clear().Append(prevSlash, u64(slash - prevSlash));
-				if (CaseInsensitiveFs)
-					forHash.MakeLower();
-				hasher.Update(forHash.data, forHash.count);
-				StringKey fileNameKey = ToStringKey(hasher);
-
-				if (directory)
+				SCOPED_READ_LOCK(directory->lock, lock2);
+				auto fileIt = directory->files.find(fileNameKey);
+				if (fileIt != directory->files.end())
 				{
-					SCOPED_READ_LOCK(directory->lock, lock2);
-					auto fileIt = directory->files.find(fileNameKey);
-					if (fileIt != directory->files.end())
-					{
-						UBA_ASSERT(fileIt->second != ~0u);
-						BinaryReader reader(GetReader(fileIt->second));
-						StringBuffer<> fileName;
-						reader.ReadString(fileName);
-						out.Append(PathSeparator).Append(fileName);
-					}
-					else
-						out.Append(prevSlash, u64(slash - prevSlash));
+					UBA_ASSERT(fileIt->second != ~0u);
+					BinaryReader reader(GetReader(fileIt->second));
+					StringBuffer<> fileName;
+					reader.ReadString(fileName);
+					out.Append(PathSeparator).Append(fileName);
 				}
 				else
 					out.Append(prevSlash, u64(slash - prevSlash));
+			}
+			else
+				out.Append(prevSlash, u64(slash - prevSlash));
 
-				if (slash == end)
-					return;
+			if (slash == end)
+				return;
 
-				prevSlash = slash;
+			prevSlash = slash;
 
-				auto findIt = m_lookup.find(fileNameKey);
-				if (findIt == m_lookup.end())
-				{
-					directory = nullptr;
-					continue;
-				}
-				directory = &findIt->second;
-				PopulateDirectory(hasher, *directory); // This is needed to make sure files lookup is populated for query above
+			auto findIt = m_lookup.find(fileNameKey);
+			if (findIt == m_lookup.end())
+			{
+				directory = nullptr;
+				continue;
 			}
+			directory = &findIt->second;
+			PopulateDirectory(hasher, *directory); // This is needed to make sure files lookup is populated for query above
 		}
+	}
 
 #if PLATFORM_WINDOWS
-		template<typename Func>
-		void TraverseFilesRecursiveNoLock(const StringBufferBase& path, const Func& func)
+	template<typename Func>
+	void DirectoryTable::TraverseFilesRecursiveNoLock(const StringBufferBase& path, const Func& func)
+	{
+		auto findIt = m_lookup.find(ToStringKey(path));
+		if (findIt == m_lookup.end())
+			return;
+		StringKeyHasher hasher;
+		hasher.Update(path.data, path.count);
+		PopulateDirectory(hasher, findIt->second);
+		for (auto& fileKv : findIt->second.files)
 		{
-			auto findIt = m_lookup.find(ToStringKey(path));
-			if (findIt == m_lookup.end())
-				return;
-			StringKeyHasher hasher;
-			hasher.Update(path.data, path.count);
-			PopulateDirectory(hasher, findIt->second);
-			for (auto& fileKv : findIt->second.files)
-			{
-				DirectoryTable::EntryInformation info;
-				StringBuffer<> fileName(path);
-				fileName.Append(PathSeparator);
-				u32 fileOffset = fileKv.second;
-				GetEntryInformation(info, fileOffset, fileName.data + fileName.count, fileName.capacity - fileName.count);
-				fileName.count = TStrlen(fileName.data);
-				if (CaseInsensitiveFs)
-					fileName.MakeLower();
-				func(info, fileName, fileOffset);
-				TraverseFilesRecursiveNoLock(fileName, func);
-			}
+			DirectoryTable::EntryInformation info;
+			StringBuffer<> fileName(path);
+			fileName.Append(PathSeparator);
+			u32 fileOffset = fileKv.second;
+			GetEntryInformation(info, fileOffset, fileName.data + fileName.count, fileName.capacity - fileName.count);
+			fileName.count = TStrlen(fileName.data);
+			if (CaseInsensitiveFs)
+				fileName.MakeLower();
+			func(info, fileName, fileOffset);
+			TraverseFilesRecursiveNoLock(fileName, func);
 		}
+	}
 
-		template<typename Func>
-		void TraverseAllFilesNoLock(const Func& func)
+	template<typename Func>
+	void DirectoryTable::TraverseAllFilesNoLock(const Func& func)
+	{
+		for (tchar l='a';l!='z'; ++l)
 		{
-			for (tchar l='a';l!='z'; ++l)
-			{
-				StringBuffer<4> drive;
-				drive.Append(l).Append(':');
-				TraverseFilesRecursiveNoLock(drive, func);
-			}
+			StringBuffer<4> drive;
+			drive.Append(l).Append(':');
+			TraverseFilesRecursiveNoLock(drive, func);
 		}
+	}
 #endif
 
-		DirectoryTable(MemoryBlock& block) : m_memoryBlock(block), m_lookup(block) {}
-		DirectoryTable(const DirectoryTable&) = delete;
-		void operator=(const DirectoryTable&) = delete;
-
-		MemoryBlock& m_memoryBlock;
-		ReaderWriterLock m_lookupLock;
-		GrowingUnorderedMap<StringKey, Directory> m_lookup;
-
-		ReaderWriterLock m_memoryLock;
-		const u8* m_memory = nullptr;
-		u32 m_memorySize = 0;
-
-		const u8* m_overlayMemory = nullptr;
-	};
+	DirectoryTable::DirectoryTable(MemoryBlock& block)
+	:	m_memoryBlock(block)
+	,	m_lookup(block)
+	{
+	}
 
+	////////////////////////////////////////////////////////////////////////////////////////////////////
 }
diff --git a/Engine/Source/Programs/UnrealBuildAccelerator/Detours/Private/PosixOS/UbaDetoursFunctionsPosix.cpp b/Engine/Source/Programs/UnrealBuildAccelerator/Detours/Private/PosixOS/UbaDetoursFunctionsPosix.cpp
--- a/Engine/Source/Programs/UnrealBuildAccelerator/Detours/Private/PosixOS/UbaDetoursFunctionsPosix.cpp
+++ b/Engine/Source/Programs/UnrealBuildAccelerator/Detours/Private/PosixOS/UbaDetoursFunctionsPosix.cpp
@@ -961,7 +961,7 @@ UBA_EXPORT int UBA_WRAPPER(mkdir)(const char* path, mode_t mode)
 	}
 	DevirtualizePath(pathName);
 
-	u32 directoryTableSize;
+	DirTableSize directoryTableSize;
 	int res;
 	u32 errorCode = 0;
 	StringKey pathNameKey = ToFilenameKey(pathName);
@@ -973,7 +973,7 @@ UBA_EXPORT int UBA_WRAPPER(mkdir)(const char* path, mode_t mode)
 		BinaryReader reader = writer.Flush();
 		res = reader.ReadBool() ? 0 : -1;
 		errorCode = reader.ReadU32();
-		directoryTableSize = reader.ReadU32();
+		directoryTableSize = FromU64(reader.ReadU64());
 	}
 
 	g_directoryTable.ParseDirectoryTable(directoryTableSize);
@@ -997,7 +997,7 @@ UBA_EXPORT int UBA_WRAPPER(rmdir)(const char* path)
 
 	DevirtualizePath(pathName);
 
-	u32 directoryTableSize;
+	DirTableSize directoryTableSize;
 	bool res;
 	u32 errorCode = 0;
 	StringKey pathNameKey = ToFilenameKey(pathName);
@@ -1009,7 +1009,7 @@ UBA_EXPORT int UBA_WRAPPER(rmdir)(const char* path)
 		BinaryReader reader = writer.Flush();;
 		res = reader.ReadBool();
 		errorCode = reader.ReadU32();
-		directoryTableSize = reader.ReadU32();
+		directoryTableSize = FromU64(reader.ReadU64());
 	}
 
 	g_directoryTable.ParseDirectoryTable(directoryTableSize);
@@ -1788,7 +1788,7 @@ UBA_EXPORT int UBA_WRAPPER(rename)(const char* oldpath, const char* newpath)
 	}
 
 
-	u32 directoryTableSize;
+	DirTableSize directoryTableSize;
 	u32 errorCode;
 	bool result;
 	{
@@ -1801,7 +1801,7 @@ UBA_EXPORT int UBA_WRAPPER(rename)(const char* oldpath, const char* newpath)
 		BinaryReader reader = writer.Flush();
 		result = reader.ReadBool();
 		errorCode = reader.ReadU32();
-		directoryTableSize = reader.ReadU32();
+		directoryTableSize = FromU64(reader.ReadU64());
 		//DEBUG_LOG_PIPE(L"MoveFile", L"%ls to %ls", lpExistingFileName, lpNewFileName);
 	}
 
@@ -1908,7 +1908,7 @@ UBA_EXPORT int UBA_WRAPPER(symlink)(const char* path1, const char* path2)
 	StringBuffer<> newToName;
 	u32 closeId;
 	u32 lastError;
-	u32 directoryTableSize;
+	DirTableSize directoryTableSize;
 	{
 		RPC_MESSAGE(CopyFile, copyFile)
 		writer.WriteStringKey(fromKey);
@@ -1920,7 +1920,7 @@ UBA_EXPORT int UBA_WRAPPER(symlink)(const char* path1, const char* path2)
 		reader.ReadString(newToName);
 		closeId = reader.ReadU32();
 		lastError = reader.ReadU32();
-		directoryTableSize = reader.ReadU32();
+		directoryTableSize = FromU64(reader.ReadU64());
 	}
 
 	UBA_ASSERT(closeId != ~0u);
@@ -1977,7 +1977,7 @@ int Shared_DeleteFile(const char* funcName, const char* pathname)
 
 	StringKey fileNameKey = ToFilenameKey(fixedName);
 
-	u32 directoryTableSize;
+	DirTableSize directoryTableSize;
 	bool result;
 	u32 errorCode;
 	{
@@ -1989,7 +1989,7 @@ int Shared_DeleteFile(const char* funcName, const char* pathname)
 		BinaryReader reader = writer.Flush();
 		result = reader.ReadBool();
 		errorCode = reader.ReadU32();
-		directoryTableSize = reader.ReadU32();
+		directoryTableSize = FromU64(reader.ReadU64());
 		pcs.Leave();
 		DEBUG_LOG_PIPE(L"DeleteFile", L"%s", lpFileName);
 	}
@@ -2893,13 +2893,12 @@ namespace uba
 
 		bool trackInputs = false;
 		u64 directoryTableHandle;
-		u32 directoryTableSize;
+		DirTableSize directoryTableSize;
 		u32 directoryTableCount;
 		u64 mappedFileTableHandle;
 		u32 mappedFileTableSize;
 		u32 mappedFileTableCount;
 		SharedMemoryHandle overlayFilesHandle;
-		u32 overlayFilesSize;
 
 		{
 			RPC_MESSAGE(Init, init);
@@ -2915,14 +2914,12 @@ namespace uba
 			g_sharedMemoryAllocator = SharedMemoryAllocatorHandle::FromU64(reader.ReadU64());
 			reader.ReadU64(); // permanent files memory handle
 			directoryTableHandle = reader.ReadU64();
-			directoryTableSize = reader.ReadU32();
+			directoryTableSize = FromU64(reader.ReadU64());
 			directoryTableCount = reader.ReadU32();
 			mappedFileTableHandle = reader.ReadU64();
 			mappedFileTableSize = reader.ReadU32();
 			mappedFileTableCount = reader.ReadU32();
-
 			overlayFilesHandle = SharedMemoryHandle::FromU64(reader.ReadU64());
-			overlayFilesSize = reader.ReadU32();
 
 			if (u16 vfsSize = reader.ReadU16())
 			{
@@ -2962,14 +2959,14 @@ namespace uba
 		UBA_ASSERTF(dirTableFd != -1, "shm_open failed (%s)", strerror(errno));
 		u8* dirTableMem = (u8*)mmap(NULL, DirTableMemSize, PROT_READ, MAP_SHARED, dirTableFd, 0);
 		UBA_ASSERTF(dirTableMem != MAP_FAILED, "mmap for dirtable mem failed (%s)", strerror(errno));
-		g_directoryTable.Init(dirTableMem, directoryTableCount, directoryTableSize);
+		g_directoryTable.Init(dirTableMem, directoryTableCount, directoryTableSize.main);
 
-		if (overlayFilesSize)
+		if (u32 overlaySize = directoryTableSize.overlay)
 		{
 			SharedMemoryView overlayFilesView;
 			Rpc_GetSharedMemory(overlayFilesView, g_sharedMemoryAllocator, nullptr, overlayFilesHandle, 0, 0, TC("PermanentFiles"));
-			g_directoryTable.InitOverlay(overlayFilesView.DetachMemory(), overlayFilesSize);
-			g_overlayFilesSize = overlayFilesSize;
+			g_directoryTable.InitOverlay(overlayFilesView.DetachMemory(), overlaySize);
+			g_overlayFilesSize = overlaySize;
 		}
 
 		if (trackInputs)
@@ -3130,8 +3127,7 @@ extern "C"
 
 		*outArguments = 0;
 		bool newProcess;
-		u32 overlayFilesSize = 0;
-		u32 directoryTableSize = 0;
+		DirTableSize directoryTableSize;
 		{
 			RPC_MESSAGE(GetNextProcess, log);
 			writer.WriteU32(prevExitCode);
@@ -3148,8 +3144,7 @@ extern "C"
 				reader.SkipString(); // workingDir
 				reader.SkipString(); // description
 				reader.ReadString(g_logName.Clear());
-				directoryTableSize = reader.ReadU32();
-				overlayFilesSize = reader.ReadU32();
+				directoryTableSize = FromU64(reader.ReadU64());
 			}
 		}
 
@@ -3168,13 +3163,13 @@ extern "C"
 			LogHeader(ToView(outArguments));
 			#endif
 
-			if (overlayFilesSize)
+			if (u32 overlaySize = directoryTableSize.overlay)
 			{
 				const u8* overlayMem = g_directoryTable.m_overlayMemory;
 				g_directoryTable.DeinitOverlay(oldOverlayFiles.data(), g_overlayFilesSize);
 				g_directoryTable.ParseDirectoryTableNoLock(directoryTableSize);
-				g_directoryTable.InitOverlay(overlayMem, overlayFilesSize);
-				g_overlayFilesSize = overlayFilesSize;
+				g_directoryTable.InitOverlay(overlayMem, overlaySize);
+				g_overlayFilesSize = overlaySize;
 			}
 		}
 
diff --git a/Engine/Source/Programs/UnrealBuildAccelerator/Detours/Private/UbaDetoursFileMappingTable.cpp b/Engine/Source/Programs/UnrealBuildAccelerator/Detours/Private/UbaDetoursFileMappingTable.cpp
--- a/Engine/Source/Programs/UnrealBuildAccelerator/Detours/Private/UbaDetoursFileMappingTable.cpp
+++ b/Engine/Source/Programs/UnrealBuildAccelerator/Detours/Private/UbaDetoursFileMappingTable.cpp
@@ -96,7 +96,7 @@ namespace uba
 		outSize = reader.ReadU64();
 		outCloseId = reader.ReadU32();
 		u32 mappedFileTableSize = reader.ReadU32();
-		u32 directoryTableSize = u32(reader.ReadU32());
+		DirTableSize directoryTableSize = FromU64(reader.ReadU64());
 		pcs.Leave();
 		DEBUG_LOG_PIPE(L"CreateFile", L"%ls (%ls)", (access == 0 ? L"ATTRIB" : ((access & AccessFlag_Write) ? L"WRITE" : L"READ")), fileName.data);
 
@@ -120,14 +120,14 @@ namespace uba
 
 	u32 Rpc_UpdateDirectory(const StringKey& dirKey, const tchar* dirName, u64 dirNameLen, bool lockDirTable)
 	{
-		u32 directoryTableSize;
+		DirTableSize directoryTableSize;
 		u32 tableOffset;
 		{
 			RPC_MESSAGE(ListDirectory, listDirectory)
 			writer.WriteString(dirName, dirNameLen);
 			writer.WriteStringKey(dirKey);
 			BinaryReader reader = writer.Flush();
-			directoryTableSize = u32(reader.ReadU32());
+			directoryTableSize = FromU64(reader.ReadU64());
 			u64 to = reader.ReadU32();
 			tableOffset = to == InvalidTableOffset ? ~u32(0) : (u32)to;
 			pcs.Leave();
@@ -142,7 +142,7 @@ namespace uba
 
 	void Rpc_UpdateCloseHandle(const tchar* handleName, u32 closeId, bool deleteOnClose, const tchar* newName, SharedMemoryHandle memoryHandle, u64 memoryWritten, bool success)
 	{
-		u32 directoryTableSize;
+		DirTableSize directoryTableSize;
 		{
 			RPC_MESSAGE(CloseFile, closeFile)
 			writer.WriteString(handleName);
@@ -166,7 +166,7 @@ namespace uba
 			else
 				writer.WriteStringKey(StringKeyZero);
 			BinaryReader reader = writer.Flush();
-			directoryTableSize = reader.ReadU32();
+			directoryTableSize = FromU64(reader.ReadU64());
 			pcs.Leave();
 			DEBUG_LOG_PIPE(L"CloseFile", L"");
 		}
@@ -249,14 +249,14 @@ namespace uba
 	void Rpc_UpdateTables()
 	{
 		SCOPED_FUTEX(g_communicationLock, pcs);
-		u32 directoryTableSize;
+		DirTableSize directoryTableSize;
 		u32 fileMappingTableSize;
 		bool done;
 		{
 			RPC_MESSAGE_NO_LOCK(UpdateTables, updateTables)
 			writer.WriteBool(false);
 			BinaryReader reader = writer.Flush();
-			directoryTableSize = reader.ReadU32();
+			directoryTableSize = FromU64(reader.ReadU64());
 			fileMappingTableSize = reader.ReadU32();
 			done = UpdateWrittenFiles(reader);
 
diff --git a/Engine/Source/Programs/UnrealBuildAccelerator/Detours/Private/Windows/UbaDetoursFunctionsKernelBase.inl b/Engine/Source/Programs/UnrealBuildAccelerator/Detours/Private/Windows/UbaDetoursFunctionsKernelBase.inl
--- a/Engine/Source/Programs/UnrealBuildAccelerator/Detours/Private/Windows/UbaDetoursFunctionsKernelBase.inl
+++ b/Engine/Source/Programs/UnrealBuildAccelerator/Detours/Private/Windows/UbaDetoursFunctionsKernelBase.inl
@@ -281,7 +281,7 @@ BOOL Detoured_CreateDirectoryW(LPCWSTR lpPathName, LPSECURITY_ATTRIBUTES lpSecur
 		return false;
 	}
 
-	u32 directoryTableSize;
+	DirTableSize directoryTableSize;
 	BOOL res;
 	u32 errorCode = 0;
 	{
@@ -291,7 +291,7 @@ BOOL Detoured_CreateDirectoryW(LPCWSTR lpPathName, LPSECURITY_ATTRIBUTES lpSecur
 		BinaryReader reader = writer.Flush();
 		res = reader.ReadBool();
 		errorCode = reader.ReadU32();
-		directoryTableSize = reader.ReadU32();
+		directoryTableSize = FromU64(reader.ReadU64());
 	}
 
 	g_directoryTable.ParseDirectoryTable(directoryTableSize);
@@ -325,7 +325,7 @@ BOOL Detoured_RemoveDirectoryW(LPCWSTR lpPathName)
 
 	DevirtualizePath(pathName);
 
-	u32 directoryTableSize;
+	DirTableSize directoryTableSize;
 	BOOL res;
 	u32 errorCode = 0;
 	StringKey pathNameKey = ToStringKeyLower(pathName);
@@ -337,7 +337,7 @@ BOOL Detoured_RemoveDirectoryW(LPCWSTR lpPathName)
 		BinaryReader reader = writer.Flush();
 		res = reader.ReadBool();
 		errorCode = reader.ReadU32();
-		directoryTableSize = reader.ReadU32();
+		directoryTableSize = FromU64(reader.ReadU64());
 	}
 
 	g_directoryTable.ParseDirectoryTable(directoryTableSize);
@@ -1169,7 +1169,7 @@ BOOL Shared_CopyFileExW(const tchar* func, LPCWSTR lpExistingFileName, LPCWSTR l
 	StringBuffer<> newToName;
 	u32 closeId;
 	u32 lastError;
-	u32 directoryTableSize;
+	DirTableSize directoryTableSize;
 	{
 		RPC_MESSAGE(CopyFile, copyFile)
 		writer.WriteStringKey(fromKey);
@@ -1181,7 +1181,7 @@ BOOL Shared_CopyFileExW(const tchar* func, LPCWSTR lpExistingFileName, LPCWSTR l
 		reader.ReadString(newToName);
 		closeId = reader.ReadU32();
 		lastError = reader.ReadU32();
-		directoryTableSize = reader.ReadU32();
+		directoryTableSize = FromU64(reader.ReadU64());
 	}
 
 	if (closeId == ~0u) // Copy was made server side
@@ -1309,7 +1309,7 @@ BOOL Detoured_DeleteFileW(LPCWSTR lpFileName)
 
 	StringKey fileNameKey = ToStringKey(fixedNameLower);
 
-	u32 directoryTableSize;
+	DirTableSize directoryTableSize;
 	bool result;
 	u32 errorCode;
 	{
@@ -1321,7 +1321,7 @@ BOOL Detoured_DeleteFileW(LPCWSTR lpFileName)
 		BinaryReader reader = writer.Flush();
 		result = reader.ReadBool();
 		errorCode = reader.ReadU32();
-		directoryTableSize = reader.ReadU32();
+		directoryTableSize = FromU64(reader.ReadU64());
 		pcs.Leave();
 		DEBUG_LOG_PIPE(L"DeleteFile", L"%ls", lpFileName);
 	}
@@ -1414,7 +1414,7 @@ bool Shared_MoveFile(const tchar* func, LPCWSTR lpExistingFileName, LPCWSTR lpNe
 
 	StringKey destKey = ToStringKeyLower(dest);
 
-	u32 directoryTableSize;
+	DirTableSize directoryTableSize;
 	u32 errorCode;
 	bool result;
 	{
@@ -1427,7 +1427,7 @@ bool Shared_MoveFile(const tchar* func, LPCWSTR lpExistingFileName, LPCWSTR lpNe
 		BinaryReader reader = writer.Flush();
 		result = reader.ReadBool();
 		errorCode = reader.ReadU32();
-		directoryTableSize = reader.ReadU32();
+		directoryTableSize = FromU64(reader.ReadU64());
 		pcs.Leave();
 		DEBUG_LOG_PIPE(func, L"%ls to %ls", lpExistingFileName, lpNewFileName);
 	}
diff --git a/Engine/Source/Programs/UnrealBuildAccelerator/Detours/Private/Windows/UbaDetoursFunctionsWin.cpp b/Engine/Source/Programs/UnrealBuildAccelerator/Detours/Private/Windows/UbaDetoursFunctionsWin.cpp
--- a/Engine/Source/Programs/UnrealBuildAccelerator/Detours/Private/Windows/UbaDetoursFunctionsWin.cpp
+++ b/Engine/Source/Programs/UnrealBuildAccelerator/Detours/Private/Windows/UbaDetoursFunctionsWin.cpp
@@ -1087,13 +1087,12 @@ void Init(const DetoursPayload& payload, u64 startTime)
 	HANDLE sharedMemoryAllocatorHandle;
 	SharedMemoryHandle permanentFilesHandle;
 	HANDLE directoryTableHandle;
-	u32 directoryTableSize;
+	DirTableSize directoryTableSize;
 	u32 directoryTableCount;
 	HANDLE mappedFileTableHandle;
 	u32 mappedFileTableSize;
 	u32 mappedFileTableCount;
 	SharedMemoryHandle overlayFilesHandle;
-	u32 overlayFilesSize;
 
 	{
 		RPC_MESSAGE(Init, init)
@@ -1109,14 +1108,13 @@ void Init(const DetoursPayload& payload, u64 startTime)
 		sharedMemoryAllocatorHandle = (HANDLE)reader.ReadU64();
 		permanentFilesHandle = SharedMemoryHandle::FromU64(reader.ReadU64());
 		directoryTableHandle = FileMappingHandle::FromU64(reader.ReadU64()).mh;
-		directoryTableSize = reader.ReadU32();
+		directoryTableSize = FromU64(reader.ReadU64());
 		directoryTableCount = reader.ReadU32();
 		mappedFileTableHandle = FileMappingHandle::FromU64(reader.ReadU64()).mh;
 		mappedFileTableSize = reader.ReadU32();
 		mappedFileTableCount = reader.ReadU32();
 
 		overlayFilesHandle = SharedMemoryHandle::FromU64(reader.ReadU64());
-		overlayFilesSize = reader.ReadU32();
 
 		if (u16 vfsSize = reader.ReadU16())
 		{
@@ -1183,14 +1181,14 @@ void Init(const DetoursPayload& payload, u64 startTime)
 	}
 	{
 		TimerScope ts2(g_stats.dirTable);
-		g_directoryTable.Init(directoryTableMem, directoryTableCount, directoryTableSize);
+		g_directoryTable.Init(directoryTableMem, directoryTableCount, directoryTableSize.main);
 
-		if (overlayFilesSize)
+		if (directoryTableSize.overlay)
 		{
 			SharedMemoryView overlayFilesView;
 			Rpc_GetSharedMemory(overlayFilesView, g_sharedMemoryAllocator, nullptr, overlayFilesHandle, 0, 0, TC("PermanentFiles"));
-			g_directoryTable.InitOverlay(overlayFilesView.DetachMemory(), overlayFilesSize);
-			g_overlayFilesSize = overlayFilesSize;
+			g_directoryTable.InitOverlay(overlayFilesView.DetachMemory(), directoryTableSize.overlay);
+			g_overlayFilesSize = directoryTableSize.overlay;
 		}
 	}
 
@@ -1310,8 +1308,7 @@ extern "C"
 
 		*outArguments = 0;
 		bool newProcess;
-		u32 overlayFilesSize = 0;
-		u32 directoryTableSize = 0;
+		DirTableSize directoryTableSize;
 		{
 			RPC_MESSAGE(GetNextProcess, log)
 			writer.WriteU32(prevExitCode);
@@ -1328,8 +1325,7 @@ extern "C"
 				reader.SkipString(); // workingDir
 				reader.SkipString(); // description
 				reader.ReadString(g_logName.Clear());
-				directoryTableSize = reader.ReadU32();
-				overlayFilesSize = reader.ReadU32();
+				directoryTableSize = FromU64(reader.ReadU64());
 			}
 		}
 
@@ -1348,13 +1344,13 @@ extern "C"
 			LogHeader(ToView(outArguments));
 			#endif
 
-			if (overlayFilesSize)
+			if (u32 overlaySize = directoryTableSize.overlay)
 			{
 				const u8* overlayMem = g_directoryTable.m_overlayMemory;
 				g_directoryTable.DeinitOverlay(oldOverlayFiles.data(), g_overlayFilesSize);
 				g_directoryTable.ParseDirectoryTableNoLock(directoryTableSize);
-				g_directoryTable.InitOverlay(overlayMem, overlayFilesSize);
-				g_overlayFilesSize = overlayFilesSize;
+				g_directoryTable.InitOverlay(overlayMem, overlaySize);
+				g_overlayFilesSize = overlaySize;
 			}
 		}
 
diff --git a/Engine/Source/Programs/UnrealBuildAccelerator/Test/Private/UbaTestSession.cpp b/Engine/Source/Programs/UnrealBuildAccelerator/Test/Private/UbaTestSession.cpp
--- a/Engine/Source/Programs/UnrealBuildAccelerator/Test/Private/UbaTestSession.cpp
+++ b/Engine/Source/Programs/UnrealBuildAccelerator/Test/Private/UbaTestSession.cpp
@@ -669,7 +669,7 @@ namespace uba
 		CHECK_TRUE(session.RefreshDirectory(workingDir, true));
 		CHECK_TRUE(session.RefreshDirectory(workingDir));
 		CHECK_TRUE(dirTable.EntryExists(ToView(workingDir)) == DirectoryTable::Exists_Maybe);
-		dirTable.ParseDirectoryTable(session.GetDirectoryTableSize());
+		dirTable.ParseDirectoryTable(session.GetDirectoryTableSize(nullptr));
 		CHECK_TRUE(dirTable.EntryExists(ToView(workingDir), true) == DirectoryTable::Exists_Yes);
 
 		StringBuffer<MaxPath> sourceFile;
@@ -683,7 +683,7 @@ namespace uba
 		CHECK_TRUE(GetAttributes(sourceFile) != INVALID_FILE_ATTRIBUTES);
 
 		CHECK_TRUE(dirTable.EntryExists(sourceFile) == DirectoryTable::Exists_No);
-		dirTable.ParseDirectoryTable(session.GetDirectoryTableSize());
+		dirTable.ParseDirectoryTable(session.GetDirectoryTableSize(nullptr));
 		CHECK_TRUE(dirTable.EntryExists(sourceFile) == DirectoryTable::Exists_Yes);
 
 		StringBuffer<MaxPath> newDir;
@@ -695,7 +695,7 @@ namespace uba
 		CHECK_TRUE(CreateDirectoryW(newDir.data));
 		CHECK_TRUE(session.RegisterNewFile(newDir.data));
 		CHECK_TRUE(dirTable.EntryExists(newDir) == DirectoryTable::Exists_No);
-		dirTable.ParseDirectoryTable(session.GetDirectoryTableSize());
+		dirTable.ParseDirectoryTable(session.GetDirectoryTableSize(nullptr));
 		CHECK_TRUE(dirTable.EntryExists(newDir) == DirectoryTable::Exists_Yes);
 		CHECK_TRUE(GetAttributes(newDir) != INVALID_FILE_ATTRIBUTES);
 		CHECK_TRUE(GetAttributes(newDirAndSlash) != INVALID_FILE_ATTRIBUTES);
@@ -706,7 +706,7 @@ namespace uba
 		CHECK_TRUE(GetAttributes(newDir2) == INVALID_FILE_ATTRIBUTES);
 		CHECK_TRUE(session.RefreshDirectory(workingDir))
 		CHECK_TRUE(GetAttributes(newDir2) != INVALID_FILE_ATTRIBUTES);
-		dirTable.ParseDirectoryTable(session.GetDirectoryTableSize());
+		dirTable.ParseDirectoryTable(session.GetDirectoryTableSize(nullptr));
 		CHECK_TRUE(dirTable.EntryExists(newDir2) == DirectoryTable::Exists_Yes);
 
 		return true;
