[UBA]
* Renamed variables in DirectoryTable to fit better what they are

diff --git a/Engine/Source/Programs/UnrealBuildAccelerator/Common/Private/UbaSession.cpp b/Engine/Source/Programs/UnrealBuildAccelerator/Common/Private/UbaSession.cpp
--- a/Engine/Source/Programs/UnrealBuildAccelerator/Common/Private/UbaSession.cpp
+++ b/Engine/Source/Programs/UnrealBuildAccelerator/Common/Private/UbaSession.cpp
@@ -579,9 +579,9 @@ namespace uba
 
 	bool Session::WriteDirectoryEntriesInternal(DirectoryTable::Directory& dir, const StringKey& dirKey, StringView dirPath, bool isRefresh, u32& outTableOffset)
 	{
-		UBA_ASSERT(dir.parseOffset != 3);
+		UBA_ASSERT(dir.populatedOffset != 3);
 
-		if (dir.tableOffset != InvalidTableOffset && !isRefresh)
+		if (dir.latestOffset != InvalidTableOffset && !isRefresh)
 		{
 			isRefresh = true;
 		}
@@ -712,7 +712,7 @@ namespace uba
 
 		written = memoryWriter.GetPosition();
 
-		u64 storageSize = sizeof(StringKey) + Get7BitEncodedCount(dir.tableOffset) + Get7BitEncodedCount(itemCount) + written;
+		u64 storageSize = sizeof(StringKey) + Get7BitEncodedCount(dir.latestOffset) + Get7BitEncodedCount(itemCount) + written;
 
 		u32 tableOffset;
 
@@ -726,7 +726,7 @@ namespace uba
 			tableWriter.Write7BitEncoded(storageSize);
 			tableWriter.WriteStringKey(dirKey);
 			tableOffset = writePos + u32(tableWriter.GetPosition());
-			tableWriter.Write7BitEncoded(dir.tableOffset);
+			tableWriter.Write7BitEncoded(dir.latestOffset);
 		}
 		else
 		{
@@ -734,7 +734,7 @@ namespace uba
 			tableWriter.Write7BitEncoded(storageSize);
 			tableWriter.WriteStringKey(dirKey);
 			tableOffset = writePos + u32(tableWriter.GetPosition());
-			tableWriter.Write7BitEncoded(dir.tableOffset);
+			tableWriter.Write7BitEncoded(dir.latestOffset);
 			tableWriter.WriteFileAttributes(dirAttributes);
 			tableWriter.WriteVolumeSerial(volumeSerialIndex);
 			tableWriter.WriteFileIndex(fileIndex);
@@ -753,7 +753,7 @@ namespace uba
 
 
 		outTableOffset = tableOffset;
-		dir.tableOffset = tableOffset;
+		dir.latestOffset = tableOffset;
 		return true;
 	}
 
@@ -767,20 +767,20 @@ namespace uba
 
 		SCOPED_WRITE_LOCK(dir.lock, dirLock);
 
-		if (dir.parseOffset == 1)
+		if (dir.populatedOffset == 1)
 		{
-			outTableOffset = dir.tableOffset;
+			outTableOffset = dir.latestOffset;
 			return;
 		}
 
 		if (!WriteDirectoryEntriesInternal(dir, dirKey, dirPath, false, outTableOffset))
 		{
 			outTableOffset = InvalidTableOffset;
-			dir.parseOffset = 2;
+			dir.populatedOffset = 2;
 		}
 		else
 		{
-			dir.parseOffset = 1;
+			dir.populatedOffset = 1;
 		}
 
 		u64 dirlen = dirPath.count;
@@ -1378,23 +1378,23 @@ namespace uba
 		SCOPED_WRITE_LOCK(dir.lock, dirLock);
 
 		// To prevent race where code creating dir manage to add to lookup but then got here later than this thread.
-		while (dir.parseOffset == 0)
+		while (dir.populatedOffset == 0)
 		{
 			dirLock.Leave();
 			Sleep(1);
 			dirLock.Enter();
 		}
 
-		// Directory was attempted to be added when it didn't exist. It is still added to dirtable lookup but we set parseOffset to 2.
+		// Directory was attempted to be added when it didn't exist. It is still added to dirtable lookup but we set populatedOffset to 2.
 		// If adding a file, clearly it does exist.. so let's reparse it.
-		if (dir.parseOffset == 2)
+		if (dir.populatedOffset == 2)
 		{
 			dirLock.Leave();
 			u32 res = WriteDirectoryEntries(dirKey, dirName);
 			UBA_ASSERT(res); (void)res;
 			dirLock.Enter();
 		}
-		UBA_ASSERTF(dir.parseOffset == 1 || (!m_shouldWriteToDisk && dir.parseOffset != 3), TC("Registering create file for write %s with unexpect dir.parseOffset %u "), fileName.data, dir.parseOffset);
+		UBA_ASSERTF(dir.populatedOffset == 1 || (!m_shouldWriteToDisk && dir.populatedOffset != 3), TC("Registering create file for write %s with unexpect dir.populatedOffset %u "), fileName.data, dir.populatedOffset);
 
 		if (fileNameKey == StringKeyZero)
 		{
@@ -1478,7 +1478,7 @@ namespace uba
 		{
 			BinaryWriter writer(temp, 0, sizeof(temp));
 			writer.WriteStringKey(dirKey);
-			writer.Write7BitEncoded(dir.tableOffset); // Previous entry for same directory
+			writer.Write7BitEncoded(dir.latestOffset); // Previous entry for same directory
 			writer.Write7BitEncoded(1); // Count
 			writer.WriteString(lastSlash + 1);
 			entryPos = writer.GetPosition();
@@ -1506,7 +1506,7 @@ namespace uba
 		insres.first->second = dirTable.m_memorySize + u32(writer.GetPosition() + entryPos); // Storing position to last write time
 		u32 tableOffset = u32(writer.GetPosition()) + sizeof(StringKey);
 		writer.WriteBytes(temp, written);
-		dir.tableOffset = dirTable.m_memorySize + tableOffset;
+		dir.latestOffset = dirTable.m_memorySize + tableOffset;
 		dirTable.m_memorySize += u32(writer.GetPosition());
 		return true;
 	}
@@ -1533,17 +1533,17 @@ namespace uba
 		lookupCs.Leave();
 		SCOPED_WRITE_LOCK(dir.lock, dirLock);
 
-		while (dir.parseOffset == 0)
+		while (dir.populatedOffset == 0)
 		{
 			dirLock.Leave();
 			Sleep(1);
 			dirLock.Enter();
 		}
 		
-		if (dir.parseOffset == 2) // Invalid directory, don't need to register this
+		if (dir.populatedOffset == 2) // Invalid directory, don't need to register this
 			return {};
 
-		UBA_ASSERTF(dir.parseOffset == 1, TC("Registering deleted file %s with unexpect dir.parseOffset %u "), fileName.data, dir.parseOffset);
+		UBA_ASSERTF(dir.populatedOffset == 1, TC("Registering deleted file %s with unexpect dir.populatedOffset %u "), fileName.data, dir.populatedOffset);
 
 		if (fileNameKey == StringKeyZero)
 		{
@@ -1563,7 +1563,7 @@ namespace uba
 		{
 			BinaryWriter writer(temp, 0, sizeof(temp));
 			writer.WriteStringKey(dirKey);
-			writer.Write7BitEncoded(dir.tableOffset); // Previous entry for same directory
+			writer.Write7BitEncoded(dir.latestOffset); // Previous entry for same directory
 			writer.Write7BitEncoded(1); // Count
 			writer.WriteString(lastSlash + 1);
 			writer.WriteFileAttributes(0);
@@ -1589,7 +1589,7 @@ namespace uba
 			writer.Write7BitEncoded(written); // Storage size
 			u32 tableOffset = u32(writer.GetPosition()) + sizeof(StringKey);
 			writer.WriteBytes(temp, written);
-			dir.tableOffset = dirTable.m_memorySize + tableOffset;
+			dir.latestOffset = dirTable.m_memorySize + tableOffset;
 			dirTable.m_memorySize += u32(writer.GetPosition());
 		}
 
@@ -1653,21 +1653,21 @@ namespace uba
 				lookupLock.Leave();
 
 				SCOPED_WRITE_LOCK(dir.lock, dirLock);
-				if (dir.parseOffset != 3)
+				if (dir.populatedOffset != 3)
 				{
-					while (dir.parseOffset == 0)
+					while (dir.populatedOffset == 0)
 					{
 						dirLock.Leave();
 						Sleep(1);
 						dirLock.Enter();
 					}
-					dir.parseOffset = 3; // 3 means it is locked, it is now a directory owned by overlays
+					dir.populatedOffset = 3; // 3 means it is locked, it is now a directory owned by overlays
 
 					#if UBA_DEBUG_TRACK_DIR
 					m_debugLogger->Info(TC("TRACKOVERL %.*s"), dirName.count, dirName.data);
 					#endif
 				}
-				dirOffsetIns.first->second = dir.tableOffset;
+				dirOffsetIns.first->second = dir.latestOffset;
 			}
 
 
@@ -2299,13 +2299,13 @@ namespace uba
 		lookupLock.Leave();
 		SCOPED_WRITE_LOCK(dir.lock, dirLock);
 
-		while (dir.parseOffset == 0)
+		while (dir.populatedOffset == 0)
 		{
 			dirLock.Leave();
 			Sleep(1);
 			dirLock.Enter();
 		}
-		UBA_ASSERT(dir.parseOffset == 1);
+		UBA_ASSERT(dir.populatedOffset == 1);
 
 		//m_directoryTable.PopulateDirectoryNoLock(hasher, dir);
 
diff --git a/Engine/Source/Programs/UnrealBuildAccelerator/Core/Public/UbaDirectoryTable.h b/Engine/Source/Programs/UnrealBuildAccelerator/Core/Public/UbaDirectoryTable.h
--- a/Engine/Source/Programs/UnrealBuildAccelerator/Core/Public/UbaDirectoryTable.h
+++ b/Engine/Source/Programs/UnrealBuildAccelerator/Core/Public/UbaDirectoryTable.h
@@ -26,8 +26,8 @@ namespace uba
 		struct Directory
 		{
 			Directory(MemoryBlock& block) : files(block) {}
-			u32 tableOffset = InvalidTableOffset;
-			u32 parseOffset = InvalidTableOffset;
+			u32 latestOffset = InvalidTableOffset;
+			u32 populatedOffset = InvalidTableOffset;
 			EntryLookup files;
 			ReaderWriterLock lock;
 		};
@@ -43,7 +43,7 @@ namespace uba
 		inline bool IsDirectoryUpToDate(Directory& dir);
 		inline void PopulateDirectoryNoLock(const StringKeyHasher& hasher, Directory& dir);
 		inline BinaryReader GetReader(u32 tableOffset);
-		inline void PopulateDirectoryRecursive(const StringKeyHasher& hasher, u32 tableOffset, u32 parseOffset, EntryLookup& files);
+		inline void PopulateDirectoryRecursive(const StringKeyHasher& hasher, u32 latestOffset, u32 populatedOffset, EntryLookup& files);
 		inline void PopulateDirectoryWithFiles(BinaryReader& reader, const StringKeyHasher& hasher, EntryLookup& files);
 
 		enum Exists
@@ -53,8 +53,8 @@ namespace uba
 			Exists_Maybe,
 		};
 
-		inline Exists EntryExists(StringKey entryKey, const StringView& entryName, bool checkIfDir = false, u32* tableOffset = nullptr);
-		inline Exists EntryExistsNoLock(StringKey entryKey, const StringView& entryName, bool checkIfDir = false, u32* tableOffset = nullptr);
+		inline Exists EntryExists(StringKey entryKey, const StringView& entryName, bool checkIfDir = false, u32* outTableOffset = nullptr);
+		inline Exists EntryExistsNoLock(StringKey entryKey, const StringView& entryName, bool checkIfDir = false, u32* outTableOffset = nullptr);
 		inline Exists EntryExists(const StringView& str, bool checkIfDir = false);
 
 		struct EntryInformation
@@ -115,7 +115,7 @@ namespace uba
 		for (auto i=m_lookup.begin(), e=m_lookup.end(); i!=e;)
 		{
 			auto& dir = i->second;
-			if (dir.tableOffset < OverlayTableOffset)
+			if (dir.latestOffset < OverlayTableOffset)
 			{
 				++i;
 				continue;
@@ -123,18 +123,18 @@ namespace uba
 
 			while (true)
 			{
-				BinaryReader reader(GetReader(dir.tableOffset));
-				dir.tableOffset = u32(reader.Read7BitEncoded());
-				if (dir.tableOffset < OverlayTableOffset)
+				BinaryReader reader(GetReader(dir.latestOffset));
+				dir.latestOffset = u32(reader.Read7BitEncoded());
+				if (dir.latestOffset < OverlayTableOffset)
 					break;
 			}
-			if (dir.tableOffset == InvalidTableOffset)
+			if (dir.latestOffset == InvalidTableOffset)
 			{
 				i = m_lookup.erase(i);
 				continue;
 			}
 
-			dir.parseOffset = dir.tableOffset;
+			dir.populatedOffset = dir.latestOffset;
 
 			for (auto j=dir.files.begin(); j!=dir.files.end();)
 			{
@@ -175,7 +175,7 @@ namespace uba
 			u64 storageSize = reader.Read7BitEncoded();
 			StringKey dirKey = reader.ReadStringKey();
 			auto insres = m_lookup.try_emplace(dirKey, m_memoryBlock); // Note that this is allowed to overwrite
-			u32& tableOffset = insres.first->second.tableOffset;
+			u32& tableOffset = insres.first->second.latestOffset;
 			UBA_ASSERTF(tableOffset < OverlayTableOffset || additionalOffset == OverlayTableOffset, TC("Can't add non-overlay directory items after overlay items"));
 			tableOffset = u32(reader.GetPosition()) + additionalOffset;
 			reader.Skip(storageSize - sizeof(dirKey));
@@ -190,15 +190,15 @@ namespace uba
 
 	bool DirectoryTable::IsDirectoryUpToDate(Directory& dir)
 	{
-		return dir.parseOffset == dir.tableOffset;
+		return dir.populatedOffset == dir.latestOffset;
 	}
 
 	void DirectoryTable::PopulateDirectoryNoLock(const StringKeyHasher& hasher, Directory& dir)
 	{
-		if (dir.parseOffset == dir.tableOffset)
+		if (dir.populatedOffset == dir.latestOffset)
 			return;
-		PopulateDirectoryRecursive(hasher, dir.tableOffset, dir.parseOffset, dir.files);
-		dir.parseOffset = dir.tableOffset;
+		PopulateDirectoryRecursive(hasher, dir.latestOffset, dir.populatedOffset, dir.files);
+		dir.populatedOffset = dir.latestOffset;
 	}
 
 	BinaryReader DirectoryTable::GetReader(u32 tableOffset)
@@ -209,9 +209,9 @@ namespace uba
 			return { m_overlayMemory - OverlayTableOffset, tableOffset};
 	}
 
-	void DirectoryTable::PopulateDirectoryRecursive(const StringKeyHasher& hasher, u32 tableOffset, u32 parseOffset, EntryLookup& files)
+	void DirectoryTable::PopulateDirectoryRecursive(const StringKeyHasher& hasher, u32 latestOffset, u32 populatedOffset, EntryLookup& files)
 	{
-		BinaryReader reader(GetReader(tableOffset));
+		BinaryReader reader(GetReader(latestOffset));
 
 		u32 prevTableOffset = u32(reader.Read7BitEncoded());
 
@@ -222,7 +222,7 @@ namespace uba
 		bool firstIsRoot = true;
 		while (true)
 		{
-			if (prevTableOffset == InvalidTableOffset || prevTableOffset == parseOffset)
+			if (prevTableOffset == InvalidTableOffset || prevTableOffset == populatedOffset)
 			{
 				firstIsRoot = prevTableOffset == InvalidTableOffset;
 				break;
@@ -287,13 +287,13 @@ namespace uba
 		}
 	}
 
-	DirectoryTable::Exists DirectoryTable::EntryExists(StringKey entryKey, const StringView& entryName, bool checkIfDir, u32* tableOffset)
+	DirectoryTable::Exists DirectoryTable::EntryExists(StringKey entryKey, const StringView& entryName, bool checkIfDir, u32* outTableOffset)
 	{
 		SCOPED_READ_LOCK(m_lookupLock, lock);
-		return EntryExistsNoLock(entryKey, entryName, checkIfDir, tableOffset);
+		return EntryExistsNoLock(entryKey, entryName, checkIfDir, outTableOffset);
 	}
 
-	DirectoryTable::Exists DirectoryTable::EntryExistsNoLock(StringKey entryKey, const StringView& entryName, bool checkIfDir, u32* tableOffset)
+	DirectoryTable::Exists DirectoryTable::EntryExistsNoLock(StringKey entryKey, const StringView& entryName, bool checkIfDir, u32* outTableOffset)
 	{
 		u32 startSkip = 2;
 		if (checkIfDir)
@@ -301,8 +301,8 @@ namespace uba
 			auto findIt = m_lookup.find(entryKey);
 			if (findIt != m_lookup.end())
 			{
-				if (tableOffset)
-					*tableOffset = u32(findIt->second.tableOffset) | 0x80000000; // Use significant bit to say that this is a dir
+				if (outTableOffset)
+					*outTableOffset = u32(findIt->second.latestOffset) | 0x80000000; // Use significant bit to say that this is a dir
 				return Exists_Yes;
 			}
 			startSkip = 1;
@@ -334,15 +334,15 @@ namespace uba
 			if (dirIt != m_lookup.end())
 			{
 				DirectoryTable::Directory& parentDir = dirIt->second;
-				if (parentDir.tableOffset == -1)
+				if (parentDir.latestOffset == -1)
 					return Exists_No;
-				if (parentDir.parseOffset != parentDir.tableOffset)
+				if (parentDir.populatedOffset != parentDir.latestOffset)
 				{
 					SCOPED_WRITE_LOCK(parentDir.lock, lock);
-					if (parentDir.parseOffset != parentDir.tableOffset)
+					if (parentDir.populatedOffset != parentDir.latestOffset)
 					{
-						PopulateDirectoryRecursive(ancestorHasher, parentDir.tableOffset, parentDir.parseOffset, parentDir.files);
-						parentDir.parseOffset = parentDir.tableOffset;
+						PopulateDirectoryRecursive(ancestorHasher, parentDir.latestOffset, parentDir.populatedOffset, parentDir.files);
+						parentDir.populatedOffset = parentDir.latestOffset;
 					}
 				}
 
@@ -358,8 +358,8 @@ namespace uba
 						return Exists_No;
 					return Exists_Maybe;
 				}
-				if (tableOffset)
-					*tableOffset = entryIt->second;
+				if (outTableOffset)
+					*outTableOffset = entryIt->second;
 				return Exists_Yes;
 			}
 
diff --git a/Engine/Source/Programs/UnrealBuildAccelerator/Detours/Private/UbaDetoursFileMappingTable.cpp b/Engine/Source/Programs/UnrealBuildAccelerator/Detours/Private/UbaDetoursFileMappingTable.cpp
--- a/Engine/Source/Programs/UnrealBuildAccelerator/Detours/Private/UbaDetoursFileMappingTable.cpp
+++ b/Engine/Source/Programs/UnrealBuildAccelerator/Detours/Private/UbaDetoursFileMappingTable.cpp
@@ -372,7 +372,7 @@ namespace uba
 		auto& dir = dirFindIt->second;
 
 		if (checkIfDir)
-			return u32(dir.tableOffset) | 0x80000000; // Use significant bit to say that this is a dir
+			return u32(dir.latestOffset) | 0x80000000; // Use significant bit to say that this is a dir
 
 		g_directoryTable.PopulateDirectory(hash.open, dir);
 
diff --git a/Engine/Source/Programs/UnrealBuildAccelerator/Detours/Private/Windows/UbaDetoursFunctionsKernelBase.inl b/Engine/Source/Programs/UnrealBuildAccelerator/Detours/Private/Windows/UbaDetoursFunctionsKernelBase.inl
--- a/Engine/Source/Programs/UnrealBuildAccelerator/Detours/Private/Windows/UbaDetoursFunctionsKernelBase.inl
+++ b/Engine/Source/Programs/UnrealBuildAccelerator/Detours/Private/Windows/UbaDetoursFunctionsKernelBase.inl
@@ -503,8 +503,8 @@ BOOL Detoured_GetVolumeInformationByHandleW(HANDLE hFile, LPWSTR lpVolumeNameBuf
 	else if (isListDirectoryHandle(hFile))
 	{
 		auto& listHandle = asListDirectoryHandle(hFile);
-		if (listHandle.dir.tableOffset != InvalidTableOffset)
-			entryOffset = listHandle.dir.tableOffset | 0x80000000;
+		if (listHandle.dir.latestOffset != InvalidTableOffset)
+			entryOffset = listHandle.dir.latestOffset | 0x80000000;
 		else
 			UBA_ASSERT(false);
 		trueHandle = INVALID_HANDLE_VALUE;
@@ -1573,9 +1573,9 @@ __forceinline HANDLE Shared_FindFirstFileExW(LPCWSTR lpFileName, FINDEX_INFO_LEV
 			Rpc_UpdateDirectory(hash.key, buf, bufChars, false);
 	}
 	bool exists = false;
-	if (dir.tableOffset != InvalidTableOffset)
+	if (dir.latestOffset != InvalidTableOffset)
 	{
-		u32 entryOffset = dir.tableOffset | 0x80000000;
+		u32 entryOffset = dir.latestOffset | 0x80000000;
 		DirectoryTable::EntryInformation entryInfo;
 		g_directoryTable.GetEntryInformation(entryInfo, entryOffset);
 		exists = entryInfo.attributes != 0;
@@ -1843,8 +1843,8 @@ BOOL Detoured_GetFileInformationByHandleEx(HANDLE hFile, FILE_INFO_BY_HANDLE_CLA
 		}
 
 		auto& listHandle = asListDirectoryHandle(hFile);
-		if (listHandle.dir.tableOffset != InvalidTableOffset)
-			entryOffset = listHandle.dir.tableOffset | 0x80000000;
+		if (listHandle.dir.latestOffset != InvalidTableOffset)
+			entryOffset = listHandle.dir.latestOffset | 0x80000000;
 		else
 			UBA_ASSERT(false);
 		trueHandle = INVALID_HANDLE_VALUE;
diff --git a/Engine/Source/Programs/UnrealBuildAccelerator/Detours/Private/Windows/UbaDetoursFunctionsNtDll.inl b/Engine/Source/Programs/UnrealBuildAccelerator/Detours/Private/Windows/UbaDetoursFunctionsNtDll.inl
--- a/Engine/Source/Programs/UnrealBuildAccelerator/Detours/Private/Windows/UbaDetoursFunctionsNtDll.inl
+++ b/Engine/Source/Programs/UnrealBuildAccelerator/Detours/Private/Windows/UbaDetoursFunctionsNtDll.inl
@@ -233,7 +233,7 @@ NTSTATUS Detoured_NtQueryVolumeInformationFile(HANDLE FileHandle, PIO_STATUS_BLO
 		if (FsInformationClass == FileFsDeviceInformation)
 			return runDefaultFsDeviceInfo();
 		if (FsInformationClass == FileFsVolumeInformation)
-			return runDefaultFsVolumeInfo(listHandle.dir.tableOffset | 0x80000000);
+			return runDefaultFsVolumeInfo(listHandle.dir.latestOffset | 0x80000000);
 		UBA_ASSERTF(false, L"NtQueryVolumeInformationFile called in ListDirectoryHandle using class %u which is not implemented (%ls)", FsInformationClass, HandleToName(FileHandle));
 	}
 	auto res = True_NtQueryVolumeInformationFile(TrueHandle, IoStatusBlock, FsInformation, Length, FsInformationClass);
@@ -260,9 +260,9 @@ NTSTATUS Detoured_NtQueryInformationFile(HANDLE FileHandle, PIO_STATUS_BLOCK IoS
 		else if (FileInformationClass == FileIdInformation)
 		{
 			auto& info = *(FILE_ID_INFORMATION*)FileInformation;
-			if (listHandle.dir.tableOffset != InvalidTableOffset)
+			if (listHandle.dir.latestOffset != InvalidTableOffset)
 			{
-				u32 entryOffset = listHandle.dir.tableOffset | 0x80000000;
+				u32 entryOffset = listHandle.dir.latestOffset | 0x80000000;
 				DirectoryTable::EntryInformation entryInfo;
 				g_directoryTable.GetEntryInformation(entryInfo, entryOffset);
 				info.VolumeSerialNumber = entryInfo.volumeSerial;
@@ -1536,9 +1536,9 @@ NTSTATUS NTAPI Shared_NtCreateFile(bool IsCreateFunc, PHANDLE hFileHandle, ACCES
 		}
 
 		bool exists = false;
-		if (dir.tableOffset != InvalidTableOffset)
+		if (dir.latestOffset != InvalidTableOffset)
 		{
-			u32 entryOffset = dir.tableOffset | 0x80000000;
+			u32 entryOffset = dir.latestOffset | 0x80000000;
 			DirectoryTable::EntryInformation entryInfo;
 			g_directoryTable.GetEntryInformation(entryInfo, entryOffset);
 			exists = entryInfo.attributes != 0;
