[UBA]
* Moved most of directory table logic out of UbaSession to more easily be able to unit test

diff --git a/Engine/Source/Programs/UnrealBuildAccelerator/Common/Private/UbaDirectoryTableHolder.cpp b/Engine/Source/Programs/UnrealBuildAccelerator/Common/Private/UbaDirectoryTableHolder.cpp
new file mode 100644
--- /dev/null
+++ b/Engine/Source/Programs/UnrealBuildAccelerator/Common/Private/UbaDirectoryTableHolder.cpp
@@ -0,0 +1,746 @@
+// Copyright Epic Games, Inc. All Rights Reserved.
+
+#include "UbaDirectoryTableHolder.h"
+#include "UbaDirectoryIterator.h"
+#include "UbaFileMapping.h"
+#include "UbaProtocol.h"
+
+#define UBA_DEBUG_TRACK_DIR 0 // UBA_DEBUG_LOGGER
+
+namespace uba
+{
+	////////////////////////////////////////////////////////////////////////////////////////////////////
+
+	bool GetDirKey(StringKey& outDirKey, StringBufferBase& outDirName, const tchar*& outLastSlash, const StringView& fileName)
+	{
+		outLastSlash = TStrrchr(fileName.data, PathSeparator);
+		UBA_ASSERTF(outLastSlash, TC("Can't get dir key for path %s"), fileName.data);
+		if (!outLastSlash)
+			return false;
+
+		u64 dirLen = u64(outLastSlash - fileName.data);
+		outDirName.Append(fileName.data, dirLen);
+		outDirKey = CaseInsensitiveFs ? ToStringKeyLower(outDirName) : ToStringKey(outDirName);
+		return true;
+	}
+
+	StringKey GetKeyAndFixedName(StringBuffer<>& outFixedFilePath, StringKeyHasher& outHasher, const tchar* filePath, StringKey* outDirKey, const tchar** outBaseFileName)
+	{
+		StringBuffer<> workingDir;
+		if (!IsAbsolutePath(filePath))
+		{
+			GetCurrentDirectoryW(workingDir);
+			workingDir.EnsureEndsWithSlash();
+		}
+		FixPath(filePath, workingDir.data, workingDir.count, outFixedFilePath);
+
+		StringKey dirKey;
+		StringBuffer<> dirNameForHash;
+		const tchar* baseFileName = nullptr;
+		GetDirKey(dirKey, dirNameForHash, baseFileName, outFixedFilePath);
+		if (outDirKey)
+			*outDirKey = dirKey;
+		if (outBaseFileName)
+			*outBaseFileName = baseFileName;
+
+		if (CaseInsensitiveFs)
+			dirNameForHash.MakeLower();
+
+		outHasher.Update(dirNameForHash);
+
+		if (baseFileName)
+		{
+			StringBuffer<256> baseFileNameForHash;
+			baseFileNameForHash.Append(baseFileName);
+			if (CaseInsensitiveFs)
+				baseFileNameForHash.MakeLower();
+
+			outHasher.Update(baseFileNameForHash);
+		}
+
+		StringKey result = ToStringKey(outHasher);
+
+#if UBA_DEBUG
+		StringBuffer<> testPath(outFixedFilePath);
+		if (CaseInsensitiveFs)
+			testPath.MakeLower();
+		StringKey testKey = ToStringKey(testPath);
+		UBA_ASSERTF(testKey == result, TC("Key mismatch for %s"), outFixedFilePath.data);
+#endif
+
+		return result;
+	}
+
+	StringKey GetKeyAndFixedName(StringBuffer<>& outFixedFilePath, const tchar* filePath, StringKey* outDirKey, const tchar** outBaseFileName)
+	{
+		StringKeyHasher hasher;
+		return GetKeyAndFixedName(outFixedFilePath, hasher, filePath, outDirKey, outBaseFileName);
+	}
+
+	////////////////////////////////////////////////////////////////////////////////////////////////////
+
+	DirectoryTableHolder::DirectoryTableHolder(LogWriter& logWriter, const tchar* logPrefix, bool runningRemote)
+	:	m_logger(logWriter, logPrefix)
+	,	m_runningRemote(runningRemote)
+	,	m_directoryTableMemory(TC("SessionDirectoryTable"))
+	,	m_directoryTable(m_directoryTableMemory)
+	{
+		m_directoryTableHandle = uba::CreateMemoryMappingW(m_logger, PAGE_READWRITE|SEC_RESERVE, DirTableMemSize, nullptr, TC("DirMappings"));
+		UBA_ASSERT(m_directoryTableHandle.IsValid());
+		m_directoryTableMem = MapViewOfFile(m_logger, m_directoryTableHandle, FILE_MAP_WRITE, 0, DirTableMemSize);
+		UBA_ASSERT(m_directoryTableMem);
+
+		m_directoryTable.m_memory = m_directoryTableMem;
+		m_directoryTable.m_lookup.reserve(30000);
+	}
+
+	DirectoryTableHolder::~DirectoryTableHolder()
+	{
+		UnmapViewOfFile(m_logger, m_directoryTableMem, DirTableMemSize, TC("DirectoryTable"), false); // Can't discard since only parts might be mapped
+		CloseFileMapping(m_logger, m_directoryTableHandle, TC("DirectoryTable"));
+	}
+
+	void DirectoryTableHolder::SetTreatTempDirAsEmpty()
+	{
+#if PLATFORM_WINDOWS
+		wchar_t systemTemp[256];
+		GetEnvironmentVariableW(L"TEMP", systemTemp, 256);
+		StringBuffer<> temp;
+		FixPath(systemTemp, nullptr, 0, temp);
+		temp.MakeLower();
+		m_directoryForcedEmpty = ToStringKey(temp);
+#endif
+	}
+
+	bool DirectoryTableHolder::RefreshDirectory(const tchar* dirName, bool forceRegister)
+	{
+		UBA_ASSERT(!m_runningRemote);
+
+		StringBuffer<> dirPath;
+		StringKeyHasher hasher;
+		GetKeyAndFixedName(dirPath, hasher, dirName);
+
+		StringKey dirKey = ToStringKey(hasher);
+
+		auto& dirTable = m_directoryTable;
+		SCOPED_READ_LOCK(dirTable.m_lookupLock, lookupLock);
+		auto res = dirTable.m_lookup.find(dirKey);
+		if (res == dirTable.m_lookup.end())
+		{
+			lookupLock.Leave();
+			if (forceRegister)
+				WriteDirectoryEntries(dirKey, dirPath);
+			return true;
+		}
+		DirectoryTable::Directory& dir = res->second;
+		lookupLock.Leave();
+		SCOPED_WRITE_LOCK(dir.lock, dirLock);
+
+		while (dir.populatedOffset == 0)
+		{
+			dirLock.Leave();
+			Sleep(1);
+			dirLock.Enter();
+		}
+		UBA_ASSERT(dir.populatedOffset == 1);
+
+		u32 tableOffset;
+		return WriteDirectoryEntriesInternal(dir, dirKey, dirPath, true, tableOffset);
+	}
+
+	bool DirectoryTableHolder::RegisterNewFile(const tchar* filePath)
+	{
+		UBA_ASSERT(!m_runningRemote);
+		StringBuffer<> fixedFilePath;
+		auto key = GetKeyAndFixedName(fixedFilePath, filePath);
+		return RegisterCreateFileForWrite(key, fixedFilePath, true);
+	}
+
+	void DirectoryTableHolder::RegisterDeleteFile(const tchar* filePath)
+	{
+		UBA_ASSERT(!m_runningRemote);
+		StringBuffer<> fixedFilePath;
+		auto key = GetKeyAndFixedName(fixedFilePath, filePath);
+		RegisterDeleteFile(key, fixedFilePath, nullptr);
+	}
+
+	bool DirectoryTableHolder::RegisterNewDirectory(const tchar* directoryPath)
+	{
+		UBA_ASSERT(!m_runningRemote);
+		StringBuffer<> fixedDirPath;
+		auto dirKey = GetKeyAndFixedName(fixedDirPath, directoryPath);
+		if (!RegisterCreateFileForWrite(dirKey, fixedDirPath, true))
+			return false;
+		WriteDirectoryEntries(dirKey, fixedDirPath);
+		return true;
+	}
+
+	const u8* DirectoryTableHolder::GetDirectoryTableMemory()
+	{
+		return m_directoryTableMem;
+	}
+
+	DirTableSize DirectoryTableHolder::GetDirectoryTableSize(DirectoryTableOverlay* overlay)
+	{
+		u32 overlaySize = 0u;
+		if (overlay)
+		{
+			SCOPED_READ_LOCK(overlay->lock, lock);
+			overlaySize = overlay->size;
+		}
+		SCOPED_READ_LOCK(m_directoryTable.m_memoryLock, lock);
+		return {m_directoryTable.m_memorySize, overlaySize};
+	}
+
+	bool DirectoryTableHolder::RegisterCreateFileForWrite(StringKey fileNameKey, StringView fileName, bool registerRealFile, u64 fileSize, u64 lastWriteTime, bool invalidateStorage)
+	{
+		// Remote is not updating its own directory table
+		if (m_runningRemote)
+			return true;
+
+		auto& dirTable = m_directoryTable;
+
+		StringKey dirKey;
+		const tchar* lastSlash;
+		StringBuffer<> dirName;
+		if (!GetDirKey(dirKey, dirName, lastSlash, fileName))
+			return true;
+
+#if 0//_DEBUG  // Bring this back, turned off right now because a few lines above the call to this method we add a mapping
+		{
+			SCOPED_FUTEX(m_fileMappingTableLookupLock, lookupLock);
+			auto findIt = m_fileMappingTableLookup.find(fileNameKey);
+			if (findIt != m_fileMappingTableLookup.end())
+			{
+				FileMappingEntry& entry = findIt->second;
+				lookupLock.Leave();
+				SCOPED_FUTEX(entry.lock, entryCs);
+				UBA_ASSERT(!entry.mapping);
+			}
+		}
+#endif
+
+		bool shouldWriteToDisk = registerRealFile && ShouldWriteToDisk(fileName);
+
+		// When not writing to disk we need to populate lookup before adding non-written files.. otherwise they will be lost once lookup is actually populated
+		if (!shouldWriteToDisk)
+		{
+			u32 res = WriteDirectoryEntries(dirKey, dirName);
+			UBA_ASSERTF(res, TC("Failed to write directory entries for %s"), dirName.data); (void)res;
+		}
+
+		SCOPED_READ_LOCK(dirTable.m_lookupLock, lookupCs);
+		auto findIt = dirTable.m_lookup.find(dirKey);
+		if (findIt == dirTable.m_lookup.end())
+			return true;
+
+		DirectoryTable::Directory& dir = findIt->second;
+		lookupCs.Leave();
+
+		SCOPED_WRITE_LOCK(dir.lock, dirLock);
+
+		// To prevent race where code creating dir manage to add to lookup but then got here later than this thread.
+		while (dir.populatedOffset == 0)
+		{
+			dirLock.Leave();
+			Sleep(1);
+			dirLock.Enter();
+		}
+
+		// Directory was attempted to be added when it didn't exist. It is still added to dirtable lookup but we set populatedOffset to 2.
+		// If adding a file, clearly it does exist.. so let's reparse it.
+		if (dir.populatedOffset == 2)
+		{
+			dirLock.Leave();
+			u32 res = WriteDirectoryEntries(dirKey, dirName);
+			UBA_ASSERT(res); (void)res;
+			dirLock.Enter();
+		}
+		UBA_ASSERTF(dir.populatedOffset == 1 || (!ShouldWriteToDisk(fileName) && dir.populatedOffset != 3), TC("Registering create file for write %s with unexpect dir.populatedOffset %u "), fileName.data, dir.populatedOffset);
+
+		if (fileNameKey == StringKeyZero)
+		{
+			StringBuffer<> forKey;
+			forKey.Append(fileName);
+			if (CaseInsensitiveFs)
+				forKey.MakeLower();
+			fileNameKey = ToStringKey(forKey);
+		}
+		auto insres = dir.files.try_emplace(fileNameKey, ~u32(0));
+
+		u64 fileIndex = InvalidValue;
+		u32 attributes = 0;
+		u32 volumeSerial = 0;
+		bool isDirectory = false;
+
+		if (shouldWriteToDisk)
+		{
+			FileInformation info;
+			if (!GetFileInformation(info, m_logger, fileName.data))
+				return m_logger.Error(TC("Failed to get file information for %s while checking file added for write. This should not happen! (%s)"), fileName.data, LastErrorToText().data);
+
+			attributes = info.attributes;
+			volumeSerial = info.volumeSerialNumber;
+			lastWriteTime = info.lastWriteTime;
+			isDirectory = IsDirectory(attributes);
+			if (isDirectory)
+				fileSize = 0;
+			else
+				fileSize = info.size;
+			fileIndex = info.index;
+		}
+		else
+		{
+			// TODO: Do we need more code here?
+			attributes = DefaultAttributes();
+			volumeSerial = 1;
+			fileIndex = ++m_fileIndexCounter;
+		}
+
+		// Check if new write is actually a write. The file might just have been open with write permissions and then actually never written to.
+		// We check this by using lastWriteTime. If it hasn't change, directory table is already up-to-date
+		if (!insres.second && insres.first->second != ~u32(0))
+		{
+			BinaryReader reader(m_directoryTableMem + insres.first->second);
+			u32 oldAttr = reader.ReadFileAttributes();
+
+			if (IsDirectory(oldAttr))  // Ignore updating directories.. they should always be the same regardless
+			{
+				UBA_ASSERT(isDirectory || !shouldWriteToDisk);
+				return true;
+			}
+			reader.ReadVolumeSerial();
+			u64 oldFileIndex = reader.ReadFileIndex();(void)oldFileIndex;
+
+			u64 oldLastWriteTime = reader.ReadFileTime();
+			if (lastWriteTime == oldLastWriteTime)
+			{
+#if !PLATFORM_WINDOWS
+				UBA_ASSERT(oldFileIndex == fileIndex); // Checking so it is really the same file
+#endif
+				u64 oldSize = reader.ReadFileSize();
+				if (oldSize == fileSize && oldAttr == attributes) // Only attributes could change from a chmod
+					return true;
+				// TODO: Somehow this can happen and I have no idea how. last written time should be set on close file so it shouldnt be possible.
+				//else
+				//	m_logger.Warning(TC("Somehow file %s has same last written time at two points in time but different size (old %llu new %llu)"), fileName.data, oldSize, fileSize);
+			}
+		}
+
+		// There are directory crawlers happening in parallel so we need to really make sure to invalidate this one since a crawler can actually
+		// hit this file with information from a query before it was written.. and then it will turn it back to "verified" using old info
+		if (registerRealFile && invalidateStorage)
+			InvalidateCachedFileInfo(fileNameKey);
+
+		FileEntryAdded(fileNameKey, lastWriteTime, fileSize);
+
+		u8 temp[1024];
+		u64 written;
+		u64 entryPos;
+		{
+			BinaryWriter writer(temp, 0, sizeof(temp));
+			writer.WriteStringKey(dirKey);
+			writer.Write7BitEncoded(dir.latestOffset); // Previous entry for same directory
+			writer.Write7BitEncoded(1); // Count
+			writer.WriteString(lastSlash + 1);
+			entryPos = writer.GetPosition();
+			writer.WriteFileAttributes(attributes);
+			writer.WriteVolumeSerial(m_volumeCache.GetSerialIndex(volumeSerial));
+			writer.WriteFileIndex(fileIndex);
+			if (!isDirectory)
+			{
+				writer.WriteFileTime(lastWriteTime);
+				writer.WriteFileSize(fileSize);
+			}
+			written = writer.GetPosition();
+		}
+
+#if UBA_DEBUG_TRACK_DIR
+		m_debugLogger->Info(TC("TRACKADD    %s (Size: %llu, Attr: %u, Key: %s, Id: %llu)"), fileName.data, fileSize, attributes, KeyToString(fileNameKey).data, fileIndex);
+#endif
+
+
+		SCOPED_WRITE_LOCK(dirTable.m_memoryLock, memoryLock);
+		u32 writePos = dirTable.m_memorySize;
+		EnsureDirectoryTableMemory(writePos + 8 + written);
+		BinaryWriter writer(m_directoryTableMem + writePos);
+		writer.Write7BitEncoded(written); // Storage size
+		insres.first->second = dirTable.m_memorySize + u32(writer.GetPosition() + entryPos); // Storing position to last write time
+		u32 tableOffset = u32(writer.GetPosition()) + sizeof(StringKey);
+		writer.WriteBytes(temp, written);
+		dir.latestOffset = dirTable.m_memorySize + tableOffset;
+		dirTable.m_memorySize += u32(writer.GetPosition());
+		return true;
+	}
+
+	DirTableSize DirectoryTableHolder::RegisterDeleteFile(StringKey fileNameKey, const StringView& fileName, DirectoryTableOverlay* overlay)
+	{
+		// Remote is not updating its own directory table
+		if (m_runningRemote)
+			return GetDirectoryTableSize(overlay);
+
+		auto& dirTable = m_directoryTable;
+
+		StringKey dirKey;
+		const tchar* lastSlash;
+		StringBuffer<> dirName;
+		if (!GetDirKey(dirKey, dirName, lastSlash, fileName))
+			return {InvalidTableOffset, InvalidTableOffset};
+
+		SCOPED_READ_LOCK(dirTable.m_lookupLock, lookupCs);
+		auto res = dirTable.m_lookup.find(dirKey);
+		if (res == dirTable.m_lookup.end())
+			return {};
+		DirectoryTable::Directory& dir = res->second;
+		lookupCs.Leave();
+		SCOPED_WRITE_LOCK(dir.lock, dirLock);
+
+		while (dir.populatedOffset == 0)
+		{
+			dirLock.Leave();
+			Sleep(1);
+			dirLock.Enter();
+		}
+
+		if (dir.populatedOffset == 2) // Invalid directory, don't need to register this
+			return {};
+
+		UBA_ASSERTF(dir.populatedOffset == 1, TC("Registering deleted file %s with unexpect dir.populatedOffset %u "), fileName.data, dir.populatedOffset);
+
+		if (fileNameKey == StringKeyZero)
+		{
+			StringBuffer<> forKey;
+			forKey.Append(fileName);
+			if (CaseInsensitiveFs)
+				forKey.MakeLower();
+			fileNameKey = ToStringKey(forKey);
+		}
+
+		// Does not exist, no need adding to file table
+		if (dir.files.erase(fileNameKey) == 0)
+			return {};
+
+		u8 temp[1024];
+		u64 written;
+		{
+			BinaryWriter writer(temp, 0, sizeof(temp));
+			writer.WriteStringKey(dirKey);
+			writer.Write7BitEncoded(dir.latestOffset); // Previous entry for same directory
+			writer.Write7BitEncoded(1); // Count
+			writer.WriteString(lastSlash + 1);
+			writer.WriteFileAttributes(0);
+			writer.WriteVolumeSerial(0);
+			writer.WriteFileIndex(0);
+			if (true) // !IsDirectory()
+			{
+				writer.WriteFileTime(0);
+				writer.WriteFileSize(0);
+			}
+			written = writer.GetPosition();
+		}
+
+#if UBA_DEBUG_TRACK_DIR
+		m_debugLogger->Info(TC("TRACKDEL    %s (Key: %s)"), fileName.data, KeyToString(fileNameKey).data);
+#endif
+
+		{
+			SCOPED_WRITE_LOCK(dirTable.m_memoryLock, memoryLock);
+			u32 writePos = dirTable.m_memorySize;
+			EnsureDirectoryTableMemory(writePos + 8 + written);
+			BinaryWriter writer(m_directoryTableMem + writePos);
+			writer.Write7BitEncoded(written); // Storage size
+			u32 tableOffset = u32(writer.GetPosition()) + sizeof(StringKey);
+			writer.WriteBytes(temp, written);
+			dir.latestOffset = dirTable.m_memorySize + tableOffset;
+			dirTable.m_memorySize += u32(writer.GetPosition());
+		}
+
+		return GetDirectoryTableSize(overlay);
+	}
+
+	u32 DirectoryTableHolder::WriteDirectoryEntries(const StringKey& dirKey, StringView dirPath, u32* outTableOffset)
+	{
+		auto& dirTable = m_directoryTable;
+		u32 temp;
+		if (!outTableOffset)
+			outTableOffset = &temp;
+		WriteDirectoryEntriesRecursive(dirKey, dirPath, *outTableOffset);
+		SCOPED_READ_LOCK(dirTable.m_memoryLock, memoryLock);
+		return dirTable.m_memorySize;
+	}
+
+	void DirectoryTableHolder::WriteDirectoryEntriesRecursive(const StringKey& dirKey, StringView dirPath, u32& outTableOffset)
+	{
+		auto& dirTable = m_directoryTable;
+		SCOPED_WRITE_LOCK(dirTable.m_lookupLock, lookupLock);
+		auto res = dirTable.m_lookup.try_emplace(dirKey, dirTable.m_memoryBlock);
+		DirectoryTable::Directory& dir = res.first->second;
+		lookupLock.Leave();
+
+		SCOPED_WRITE_LOCK(dir.lock, dirLock);
+
+		if (dir.populatedOffset == 1)
+		{
+			outTableOffset = dir.latestOffset;
+			return;
+		}
+
+		if (!WriteDirectoryEntriesInternal(dir, dirKey, dirPath, false, outTableOffset))
+		{
+			outTableOffset = InvalidTableOffset;
+			dir.populatedOffset = 2;
+		}
+		else
+		{
+			dir.populatedOffset = 1;
+		}
+
+		u64 dirlen = dirPath.count;
+
+		if (!dirlen) // This is for non-windows.. '/' is actually empty to get hashes correct
+			return;
+
+		// scan backwards first
+		const tchar* rit = (tchar*)dirPath.data + dirlen - 2;
+		while (rit > dirPath.data)
+		{
+			if (*rit != PathSeparator)
+			{
+				--rit;
+				continue;
+			}
+			break;
+		}
+
+		if (IsWindows && rit <= dirPath.data) // There are no path separators left, this is the drive
+		{
+			dirPath.count = 0;
+			return;
+		}
+
+		dirPath.count = u32(rit - dirPath.data);
+
+		StringBuffer<> parentDirForHash;
+		parentDirForHash.Append(dirPath);
+		if (CaseInsensitiveFs)
+			parentDirForHash.MakeLower();
+		StringKey parentKey = ToStringKey(parentDirForHash);
+
+		// Traverse through ancestors and populate them, this is an optimization
+		u32 parentOffset;
+		WriteDirectoryEntriesRecursive(parentKey, dirPath, parentOffset);
+	}
+
+	bool DirectoryTableHolder::WriteDirectoryEntriesInternal(DirectoryTable::Directory& dir, const StringKey& dirKey, StringView dirPath, bool isRefresh, u32& outTableOffset)
+	{
+		UBA_ASSERT(dir.populatedOffset != 3);
+
+		if (dir.latestOffset != InvalidTableOffset && !isRefresh)
+		{
+			isRefresh = true;
+		}
+
+		auto& dirTable = m_directoryTable;
+
+		u32 volumeSerial = 0;
+		u32 volumeSerialIndex = 0;
+		u32 dirAttributes = 0;
+		u64 fileIndex = 0;
+
+		u64 written = 0;
+		u32 itemCount = 0;
+
+		StringKeyHasher hasher;
+		if (dirPath.count)
+		{
+			StringBuffer<> forHash;
+			forHash.Append(dirPath);
+			if (CaseInsensitiveFs)
+				forHash.MakeLower();
+			hasher.Update(forHash.data, forHash.count);
+		}
+
+#if UBA_DEBUG_TRACK_DIR
+		m_debugLogger->BeginScope();
+		auto dg = MakeGuard([&]() { m_debugLogger->EndScope(); });
+		StringBuffer<> str;
+		str.Append(TCV("TRACKDIR "));
+		if (isRefresh)
+			str.Append(TCV("(Refresh) "));
+		str.Append(dirPath).Append('\n');
+		m_debugLogger->Log(LogEntryType_Info, str);
+#endif
+
+
+		Vector<u8> memoryBlock;
+		memoryBlock.resize(4096);
+		BinaryWriter memoryWriter(memoryBlock.data(), 0, memoryBlock.size());
+
+		if (dirKey != m_directoryForcedEmpty)
+		{
+			StringBuffer<4> realPath;
+			if constexpr (IsWindows)
+			{
+				if (dirPath.count == 2)
+					dirPath = realPath.Append(dirPath).Append(PathSeparator);
+			}
+			else
+			{
+				if (!dirPath.count)
+					dirPath = realPath.Append(PathSeparator);
+			}
+
+			bool res = TraverseDir(m_logger, dirPath,
+				[&](const DirectoryEntry& e)
+				{
+					StringBuffer<256> fileNameForHash;
+					fileNameForHash.Append(PathSeparator).Append(e.name, e.nameLen);
+					if (CaseInsensitiveFs)
+						fileNameForHash.MakeLower();
+
+					StringKey fileKey = ToStringKey(hasher, fileNameForHash.data, fileNameForHash.count);
+					auto res = dir.files.try_emplace(fileKey, ~0u);
+					if (!res.second)
+						return;
+					UBA_ASSERT(e.attributes);
+					UBA_ASSERT(e.nameLen < 512);
+					memoryWriter.WriteString(e.name, e.nameLen);
+
+#if UBA_DEBUG_TRACK_DIR
+					m_debugLogger->Info(TC("    %s (Size: %llu, Attr: %u, Key: %s, Id: %llu)"), e.name, e.size, e.attributes, KeyToString(fileKey).data, e.id);
+#endif
+
+					u64 id = e.id;
+					if (id == 0xffffffffffffffffllu) // When using projfs we might not have the file yet and in that case we need to make this up.
+						id = ++m_fileIndexCounter;
+
+					res.first->second = u32(memoryWriter.GetPosition()); // Temporary offset that will be used further down to calculate the real offset
+					memoryWriter.WriteFileAttributes(e.attributes);
+					memoryWriter.WriteVolumeSerial(e.volumeSerial == volumeSerial ? volumeSerialIndex : m_volumeCache.GetSerialIndex(e.volumeSerial));
+					memoryWriter.WriteFileIndex(id);
+					if (!IsDirectory(e.attributes))
+					{
+						memoryWriter.WriteFileTime(e.lastWritten);
+						memoryWriter.WriteFileSize(e.size);
+					}
+
+					FileEntryAdded(res.first->first, e.lastWritten, e.size);
+
+
+					++itemCount;
+					if (memoryWriter.GetPosition() > memoryBlock.size() - MaxPath)
+					{
+						memoryBlock.resize(memoryBlock.size() * 2);
+						memoryWriter.ChangeData(memoryBlock.data(), memoryBlock.size());
+					}
+				}, true,
+				[&](const DirectoryInfo& e)
+					{
+						volumeSerial = e.volumeSerial;
+						volumeSerialIndex = m_volumeCache.GetSerialIndex(volumeSerial);
+						dirAttributes = e.attributes;
+						fileIndex = e.id;
+					});
+				if (!res)
+				{
+#if UBA_DEBUG_TRACK_DIR
+					m_debugLogger->Info(TC("    FAILED (not existing?)"));
+#endif
+
+					if (!IsWindows || dirPath.count > 3)
+						return false;
+				}
+		}
+		else
+		{
+#if UBA_DEBUG_TRACK_DIR
+			m_debugLogger->Info(TC("    FORCED EMPTY"));
+#endif
+
+#if PLATFORM_WINDOWS
+			dirAttributes = FILE_ATTRIBUTE_DIRECTORY;
+#else
+			UBA_ASSERTF(false, TC("Not implemented"));
+#endif
+		}
+
+		written = memoryWriter.GetPosition();
+
+		u64 storageSize = sizeof(StringKey) + Get7BitEncodedCount(dir.latestOffset) + Get7BitEncodedCount(itemCount) + written;
+
+		u32 tableOffset;
+
+		SCOPED_WRITE_LOCK(dirTable.m_memoryLock, memoryLock);
+		u32 writePos = dirTable.m_memorySize;
+		EnsureDirectoryTableMemory(writePos + 128 + storageSize);
+		BinaryWriter tableWriter(m_directoryTableMem + dirTable.m_memorySize);
+
+		if (isRefresh)
+		{
+			tableWriter.Write7BitEncoded(storageSize);
+			tableWriter.WriteStringKey(dirKey);
+			tableOffset = writePos + u32(tableWriter.GetPosition());
+			tableWriter.Write7BitEncoded(dir.latestOffset);
+		}
+		else
+		{
+			storageSize += Get7BitEncodedCount(dirAttributes) + Get7BitEncodedCount(volumeSerialIndex) + sizeof(fileIndex);
+			tableWriter.Write7BitEncoded(storageSize);
+			tableWriter.WriteStringKey(dirKey);
+			tableOffset = writePos + u32(tableWriter.GetPosition());
+			tableWriter.Write7BitEncoded(dir.latestOffset);
+			tableWriter.WriteFileAttributes(dirAttributes);
+			tableWriter.WriteVolumeSerial(volumeSerialIndex);
+			tableWriter.WriteFileIndex(fileIndex);
+		}
+
+		tableWriter.Write7BitEncoded(itemCount);
+		u32 filesOffset = writePos + u32(tableWriter.GetPosition());
+		tableWriter.WriteBytes(memoryBlock.data(), written);
+		dirTable.m_memorySize += u32(tableWriter.GetPosition());
+
+		memoryLock.Leave();
+
+		// Update offsets to be relative to full memory
+		for (auto& kv : dir.files)
+			kv.second = filesOffset + kv.second;
+
+
+		outTableOffset = tableOffset;
+		dir.latestOffset = tableOffset;
+		return true;
+	}
+
+	void DirectoryTableHolder::EnsureDirectoryTableMemory(u64 neededSize)
+	{
+		auto& dirTable = m_directoryTable;
+		if (neededSize <= m_directoryTableMemCommitted)
+			return;
+
+		u64 newSize = AlignUp(neededSize, u64(1024*1024));
+		if (newSize > DirTableMemSize)
+		{
+			static bool called = m_logger.Error(TC("Directory table overflow. DirTableMemSize need to be increased (Size: %llu)"), dirTable.m_memorySize);
+		}
+
+		u64 toCommit = newSize - m_directoryTableMemCommitted;
+		u8* address = m_directoryTableMem + m_directoryTableMemCommitted;
+		if (!MapViewCommit(m_logger, address, toCommit, false))
+			m_logger.Error(TC("Failed to commit memory for directory table (Committed: %llu, ToCommit: %llu) (%s)"), m_directoryTableMemCommitted, toCommit, LastErrorToText().data);
+		m_directoryTableMemCommitted += toCommit;
+		LockMemory(m_logger, address, toCommit, TC("DirectoryTable"));
+	}
+
+	void DirectoryTableHolder::FileEntryAdded(StringKey fileNameKey, u64 lastWritten, u64 size)
+	{
+	}
+
+	bool DirectoryTableHolder::ShouldWriteToDisk(StringView fileName)
+	{
+		return true;
+	}
+
+	void DirectoryTableHolder::InvalidateCachedFileInfo(const StringKey& fileNameKey)
+	{
+	}
+}
\ No newline at end of file
diff --git a/Engine/Source/Programs/UnrealBuildAccelerator/Common/Private/UbaFileAccessor.cpp b/Engine/Source/Programs/UnrealBuildAccelerator/Common/Private/UbaFileAccessor.cpp
--- a/Engine/Source/Programs/UnrealBuildAccelerator/Common/Private/UbaFileAccessor.cpp
+++ b/Engine/Source/Programs/UnrealBuildAccelerator/Common/Private/UbaFileAccessor.cpp
@@ -470,7 +470,7 @@ namespace uba
 	{
 		if (m_data)
 		{
-			if (!UnmapViewOfFile(m_logger, m_data, m_size, m_fileName))
+			if (m_size && !UnmapViewOfFile(m_logger, m_data, m_size, m_fileName))
 				return m_logger.Error(TC("Failed to unmap memory for %s (%s)"), m_fileName, LastErrorToText().data);
 			m_data = nullptr;
 		}
diff --git a/Engine/Source/Programs/UnrealBuildAccelerator/Common/Private/UbaSession.cpp b/Engine/Source/Programs/UnrealBuildAccelerator/Common/Private/UbaSession.cpp
--- a/Engine/Source/Programs/UnrealBuildAccelerator/Common/Private/UbaSession.cpp
+++ b/Engine/Source/Programs/UnrealBuildAccelerator/Common/Private/UbaSession.cpp
@@ -30,7 +30,7 @@ extern char **environ;
 #include <servers/bootstrap.h>
 #endif
 
-#define UBA_DEBUG_TRACK_DIR 0 // UBA_DEBUG_LOGGER
+#define UBA_DEBUG_TRACK_DIR 0
 #define UBA_DEBUG_TRACK_WRITES 0 // UBA_DEBUG_LOGGER
 
 //////////////////////////////////////////////////////////////////////////////
@@ -39,6 +39,8 @@ namespace uba
 {
 	bool g_dummy;
 
+	////////////////////////////////////////////////////////////////////////////////////////////////////
+
 	ProcessStartInfo::ProcessStartInfo() = default;
 	ProcessStartInfo::~ProcessStartInfo() = default;
 	ProcessStartInfo::ProcessStartInfo(const ProcessStartInfo&) = default;
@@ -202,6 +204,8 @@ namespace uba
 		process->AddRef();
 	}
 
+	////////////////////////////////////////////////////////////////////////////////////////////////////
+
 	void SessionCreateInfo::Apply(const Config& config)
 	{
 		const ConfigTable* tablePtr = config.GetTable(TC("Session"));
@@ -242,6 +246,8 @@ namespace uba
 		table.GetValueAsU64(keepOutputFileMemoryMapsThreshold, TC("KeepOutputFileMemoryMapsThreshold"));
 	}
 
+	////////////////////////////////////////////////////////////////////////////////////////////////////
+
 	bool Session::WriteFileToDisk(ProcessImpl& process, WrittenFile& file)
 	{
 		auto& rules = *process.m_startInfo.rules;
@@ -577,259 +583,6 @@ namespace uba
 		m_environmentVariables.push_back(0);
 	}
 
-	bool Session::WriteDirectoryEntriesInternal(DirectoryTable::Directory& dir, const StringKey& dirKey, StringView dirPath, bool isRefresh, u32& outTableOffset)
-	{
-		UBA_ASSERT(dir.populatedOffset != 3);
-
-		if (dir.latestOffset != InvalidTableOffset && !isRefresh)
-		{
-			isRefresh = true;
-		}
-
-		auto& dirTable = m_directoryTable;
-
-		u32 volumeSerial = 0;
-		u32 volumeSerialIndex = 0;
-		u32 dirAttributes = 0;
-		u64 fileIndex = 0;
-
-		u64 written = 0;
-		u32 itemCount = 0;
-
-		StringKeyHasher hasher;
-		if (dirPath.count)
-		{
-			StringBuffer<> forHash;
-			forHash.Append(dirPath);
-			if (CaseInsensitiveFs)
-				forHash.MakeLower();
-			hasher.Update(forHash.data, forHash.count);
-		}
-
-		#if UBA_DEBUG_TRACK_DIR
-		m_debugLogger->BeginScope();
-		auto dg = MakeGuard([&]() { m_debugLogger->EndScope(); });
-		StringBuffer<> str;
-		str.Append(TCV("TRACKDIR "));
-		if (isRefresh)
-			str.Append(TCV("(Refresh) "));
-		str.Append(dirPath).Append('\n');
-		m_debugLogger->Log(LogEntryType_Info, str);
-		#endif
-
-
-		Vector<u8> memoryBlock;
-		memoryBlock.resize(4096);
-		BinaryWriter memoryWriter(memoryBlock.data(), 0, memoryBlock.size());
-
-		if (dirKey != m_directoryForcedEmpty)
-		{
-			StringBuffer<4> realPath;
-			if constexpr (IsWindows)
-			{
-				if (dirPath.count == 2)
-					dirPath = realPath.Append(dirPath).Append(PathSeparator);
-			}
-			else
-			{
-				if (!dirPath.count)
-					dirPath = realPath.Append(PathSeparator);
-			}
-
-			bool res = TraverseDir(m_logger, dirPath,
-				[&](const DirectoryEntry& e)
-				{
-					StringBuffer<256> fileNameForHash;
-					fileNameForHash.Append(PathSeparator).Append(e.name, e.nameLen);
-					if (CaseInsensitiveFs)
-						fileNameForHash.MakeLower();
-
-					StringKey fileKey = ToStringKey(hasher, fileNameForHash.data, fileNameForHash.count);
-					auto res = dir.files.try_emplace(fileKey, ~0u);
-					if (!res.second)
-						return;
-					UBA_ASSERT(e.attributes);
-					UBA_ASSERT(e.nameLen < 512);
-					memoryWriter.WriteString(e.name, e.nameLen);
-
-					#if UBA_DEBUG_TRACK_DIR
-					m_debugLogger->Info(TC("    %s (Size: %llu, Attr: %u, Key: %s, Id: %llu)"), e.name, e.size, e.attributes, KeyToString(fileKey).data, e.id);
-					#endif
-
-					u64 id = e.id;
-					if (id == 0xffffffffffffffffllu) // When using projfs we might not have the file yet and in that case we need to make this up.
-						id = ++m_fileIndexCounter;
-
-					res.first->second = u32(memoryWriter.GetPosition()); // Temporary offset that will be used further down to calculate the real offset
-					memoryWriter.WriteFileAttributes(e.attributes);
-					memoryWriter.WriteVolumeSerial(e.volumeSerial == volumeSerial ? volumeSerialIndex : m_volumeCache.GetSerialIndex(e.volumeSerial));
-					memoryWriter.WriteFileIndex(id);
-					if (!IsDirectory(e.attributes))
-					{
-						memoryWriter.WriteFileTime(e.lastWritten);
-						memoryWriter.WriteFileSize(e.size);
-					}
-
-					FileEntryAdded(res.first->first, e.lastWritten, e.size);
-
-
-					++itemCount;
-					if (memoryWriter.GetPosition() > memoryBlock.size() - MaxPath)
-					{
-						memoryBlock.resize(memoryBlock.size() * 2);
-						memoryWriter.ChangeData(memoryBlock.data(), memoryBlock.size());
-					}
-				}, true,
-				[&](const DirectoryInfo& e)
-				{
-					volumeSerial = e.volumeSerial;
-					volumeSerialIndex = m_volumeCache.GetSerialIndex(volumeSerial);
-					dirAttributes = e.attributes;
-					fileIndex = e.id;
-				});
-			if (!res)
-			{
-				#if UBA_DEBUG_TRACK_DIR
-				m_debugLogger->Info(TC("    FAILED (not existing?)"));
-				#endif
-
-				if (!IsWindows || dirPath.count > 3)
-					return false;
-			}
-		}
-		else
-		{
-			#if UBA_DEBUG_TRACK_DIR
-			m_debugLogger->Info(TC("    FORCED EMPTY"));
-			#endif
-
-			#if PLATFORM_WINDOWS
-			dirAttributes = FILE_ATTRIBUTE_DIRECTORY;
-			#else
-			UBA_ASSERTF(false, TC("Not implemented"));
-			#endif
-		}
-
-		written = memoryWriter.GetPosition();
-
-		u64 storageSize = sizeof(StringKey) + Get7BitEncodedCount(dir.latestOffset) + Get7BitEncodedCount(itemCount) + written;
-
-		u32 tableOffset;
-
-		SCOPED_WRITE_LOCK(dirTable.m_memoryLock, memoryLock);
-		u32 writePos = dirTable.m_memorySize;
-		EnsureDirectoryTableMemory(writePos + 128 + storageSize);
-		BinaryWriter tableWriter(m_directoryTableMem + dirTable.m_memorySize);
-
-		if (isRefresh)
-		{
-			tableWriter.Write7BitEncoded(storageSize);
-			tableWriter.WriteStringKey(dirKey);
-			tableOffset = writePos + u32(tableWriter.GetPosition());
-			tableWriter.Write7BitEncoded(dir.latestOffset);
-		}
-		else
-		{
-			storageSize += Get7BitEncodedCount(dirAttributes) + Get7BitEncodedCount(volumeSerialIndex) + sizeof(fileIndex);
-			tableWriter.Write7BitEncoded(storageSize);
-			tableWriter.WriteStringKey(dirKey);
-			tableOffset = writePos + u32(tableWriter.GetPosition());
-			tableWriter.Write7BitEncoded(dir.latestOffset);
-			tableWriter.WriteFileAttributes(dirAttributes);
-			tableWriter.WriteVolumeSerial(volumeSerialIndex);
-			tableWriter.WriteFileIndex(fileIndex);
-		}
-
-		tableWriter.Write7BitEncoded(itemCount);
-		u32 filesOffset = writePos + u32(tableWriter.GetPosition());
-		tableWriter.WriteBytes(memoryBlock.data(), written);
-		dirTable.m_memorySize += u32(tableWriter.GetPosition());
-		
-		memoryLock.Leave();
-
-		// Update offsets to be relative to full memory
-		for (auto& kv : dir.files)
-			kv.second = filesOffset + kv.second;
-
-
-		outTableOffset = tableOffset;
-		dir.latestOffset = tableOffset;
-		return true;
-	}
-
-	void Session::WriteDirectoryEntriesRecursive(const StringKey& dirKey, StringView dirPath, u32& outTableOffset)
-	{
-		auto& dirTable = m_directoryTable;
-		SCOPED_WRITE_LOCK(dirTable.m_lookupLock, lookupLock);
-		auto res = dirTable.m_lookup.try_emplace(dirKey, dirTable.m_memoryBlock);
-		DirectoryTable::Directory& dir = res.first->second;
-		lookupLock.Leave();
-
-		SCOPED_WRITE_LOCK(dir.lock, dirLock);
-
-		if (dir.populatedOffset == 1)
-		{
-			outTableOffset = dir.latestOffset;
-			return;
-		}
-
-		if (!WriteDirectoryEntriesInternal(dir, dirKey, dirPath, false, outTableOffset))
-		{
-			outTableOffset = InvalidTableOffset;
-			dir.populatedOffset = 2;
-		}
-		else
-		{
-			dir.populatedOffset = 1;
-		}
-
-		u64 dirlen = dirPath.count;
-
-		if (!dirlen) // This is for non-windows.. '/' is actually empty to get hashes correct
-			return;
-
-		// scan backwards first
-		const tchar* rit = (tchar*)dirPath.data + dirlen - 2;
-		while (rit > dirPath.data)
-		{
-			if (*rit != PathSeparator)
-			{
-				--rit;
-				continue;
-			}
-			break;
-		}
-
-		if (IsWindows && rit <= dirPath.data) // There are no path separators left, this is the drive
-		{
-			dirPath.count = 0;
-			return;
-		}
-
-		dirPath.count = u32(rit - dirPath.data);
-
-		StringBuffer<> parentDirForHash;
-		parentDirForHash.Append(dirPath);
-		if (CaseInsensitiveFs)
-			parentDirForHash.MakeLower();
-		StringKey parentKey = ToStringKey(parentDirForHash);
-
-		// Traverse through ancestors and populate them, this is an optimization
-		u32 parentOffset;
-		WriteDirectoryEntriesRecursive(parentKey, dirPath, parentOffset);
-	}
-
-	u32 Session::WriteDirectoryEntries(const StringKey& dirKey, const StringView& dirPath, u32* outTableOffset)
-	{
-		auto& dirTable = m_directoryTable;
-		u32 temp;
-		if (!outTableOffset)
-			outTableOffset = &temp;
-		WriteDirectoryEntriesRecursive(dirKey, dirPath, *outTableOffset);
-		SCOPED_READ_LOCK(dirTable.m_memoryLock, memoryLock);
-		return dirTable.m_memorySize;
-	}
-
 	u32 Session::AddFileMapping(StringKey fileNameKey, const tchar* fileName, const tchar* newFileName, u64 fileSize)
 	{
 		UBA_ASSERT(fileNameKey != StringKeyZero);
@@ -1264,338 +1017,6 @@ namespace uba
 		}
 	}
 
-	bool GetDirKey(StringKey& outDirKey, StringBufferBase& outDirName, const tchar*& outLastSlash, const StringView& fileName)
-	{
-		outLastSlash = TStrrchr(fileName.data, PathSeparator);
-		UBA_ASSERTF(outLastSlash, TC("Can't get dir key for path %s"), fileName.data);
-		if (!outLastSlash)
-			return false;
-
-		u64 dirLen = u64(outLastSlash - fileName.data);
-		outDirName.Append(fileName.data, dirLen);
-		outDirKey = CaseInsensitiveFs ? ToStringKeyLower(outDirName) : ToStringKey(outDirName);
-		return true;
-	}
-
-	StringKey GetKeyAndFixedName(StringBuffer<>& outFixedFilePath, StringKeyHasher& outHasher, const tchar* filePath, StringKey* outDirKey = nullptr, const tchar** outBaseFileName = nullptr)
-	{
-		StringBuffer<> workingDir;
-		if (!IsAbsolutePath(filePath))
-		{
-			GetCurrentDirectoryW(workingDir);
-			workingDir.EnsureEndsWithSlash();
-		}
-		FixPath(filePath, workingDir.data, workingDir.count, outFixedFilePath);
-
-		StringKey dirKey;
-		StringBuffer<> dirNameForHash;
-		const tchar* baseFileName = nullptr;
-		GetDirKey(dirKey, dirNameForHash, baseFileName, outFixedFilePath);
-		if (outDirKey)
-			*outDirKey = dirKey;
-		if (outBaseFileName)
-			*outBaseFileName = baseFileName;
-
-		if (CaseInsensitiveFs)
-			dirNameForHash.MakeLower();
-
-		outHasher.Update(dirNameForHash);
-
-		if (baseFileName)
-		{
-			StringBuffer<256> baseFileNameForHash;
-			baseFileNameForHash.Append(baseFileName);
-			if (CaseInsensitiveFs)
-				baseFileNameForHash.MakeLower();
-
-			outHasher.Update(baseFileNameForHash);
-		}
-
-		StringKey result = ToStringKey(outHasher);
-
-		#if UBA_DEBUG
-		StringBuffer<> testPath(outFixedFilePath);
-		if (CaseInsensitiveFs)
-			testPath.MakeLower();
-		StringKey testKey = ToStringKey(testPath);
-		UBA_ASSERTF(testKey == result, TC("Key mismatch for %s"), outFixedFilePath.data);
-		#endif
-
-		return result;
-	}
-
-	StringKey GetKeyAndFixedName(StringBuffer<>& outFixedFilePath, const tchar* filePath, StringKey* outDirKey = nullptr, const tchar** outBaseFileName = nullptr)
-	{
-		StringKeyHasher hasher;
-		return GetKeyAndFixedName(outFixedFilePath, hasher, filePath, outDirKey, outBaseFileName);
-	}
-
-	bool Session::RegisterCreateFileForWrite(StringKey fileNameKey, const StringView& fileName, bool registerRealFile, u64 fileSize, u64 lastWriteTime, bool invalidateStorage)
-	{
-		// Remote is not updating its own directory table
-		if (m_runningRemote)
-			return true;
-
-		auto& dirTable = m_directoryTable;
-
-		StringKey dirKey;
-		const tchar* lastSlash;
-		StringBuffer<> dirName;
-		if (!GetDirKey(dirKey, dirName, lastSlash, fileName))
-			return true;
-			
-		#if 0//_DEBUG  // Bring this back, turned off right now because a few lines above the call to this method we add a mapping
-		{
-			SCOPED_FUTEX(m_fileMappingTableLookupLock, lookupLock);
-			auto findIt = m_fileMappingTableLookup.find(fileNameKey);
-			if (findIt != m_fileMappingTableLookup.end())
-			{
-				FileMappingEntry& entry = findIt->second;
-				lookupLock.Leave();
-				SCOPED_FUTEX(entry.lock, entryCs);
-				UBA_ASSERT(!entry.mapping);
-			}
-		}
-		#endif
-
-		bool shouldWriteToDisk = registerRealFile && ShouldWriteToDisk(fileName);
-
-		// When not writing to disk we need to populate lookup before adding non-written files.. otherwise they will be lost once lookup is actually populated
-		if (!shouldWriteToDisk)
-		{
-			u32 res = WriteDirectoryEntries(dirKey, dirName);
-			UBA_ASSERTF(res, TC("Failed to write directory entries for %s"), dirName.data); (void)res;
-		}
-
-		SCOPED_READ_LOCK(dirTable.m_lookupLock, lookupCs);
-		auto findIt = dirTable.m_lookup.find(dirKey);
-		if (findIt == dirTable.m_lookup.end())
-			return true;
-
-		DirectoryTable::Directory& dir = findIt->second;
-		lookupCs.Leave();
-
-		SCOPED_WRITE_LOCK(dir.lock, dirLock);
-
-		// To prevent race where code creating dir manage to add to lookup but then got here later than this thread.
-		while (dir.populatedOffset == 0)
-		{
-			dirLock.Leave();
-			Sleep(1);
-			dirLock.Enter();
-		}
-
-		// Directory was attempted to be added when it didn't exist. It is still added to dirtable lookup but we set populatedOffset to 2.
-		// If adding a file, clearly it does exist.. so let's reparse it.
-		if (dir.populatedOffset == 2)
-		{
-			dirLock.Leave();
-			u32 res = WriteDirectoryEntries(dirKey, dirName);
-			UBA_ASSERT(res); (void)res;
-			dirLock.Enter();
-		}
-		UBA_ASSERTF(dir.populatedOffset == 1 || (!m_shouldWriteToDisk && dir.populatedOffset != 3), TC("Registering create file for write %s with unexpect dir.populatedOffset %u "), fileName.data, dir.populatedOffset);
-
-		if (fileNameKey == StringKeyZero)
-		{
-			StringBuffer<> forKey;
-			forKey.Append(fileName);
-			if (CaseInsensitiveFs)
-				forKey.MakeLower();
-			fileNameKey = ToStringKey(forKey);
-		}
-		auto insres = dir.files.try_emplace(fileNameKey, ~u32(0));
-
-		u64 fileIndex = InvalidValue;
-		u32 attributes = 0;
-		u32 volumeSerial = 0;
-		bool isDirectory = false;
-
-		if (shouldWriteToDisk)
-		{
-			FileInformation info;
-			if (!GetFileInformation(info, m_logger, fileName.data))
-				return m_logger.Error(TC("Failed to get file information for %s while checking file added for write. This should not happen! (%s)"), fileName.data, LastErrorToText().data);
-
-			attributes = info.attributes;
-			volumeSerial = info.volumeSerialNumber;
-			lastWriteTime = info.lastWriteTime;
-			isDirectory = IsDirectory(attributes);
-			if (isDirectory)
-				fileSize = 0;
-			else
-				fileSize = info.size;
-			fileIndex = info.index;
-		}
-		else
-		{
-			// TODO: Do we need more code here?
-			attributes = DefaultAttributes();
-			volumeSerial = 1;
-			fileIndex = ++m_fileIndexCounter;
-		}
-
-		// Check if new write is actually a write. The file might just have been open with write permissions and then actually never written to.
-		// We check this by using lastWriteTime. If it hasn't change, directory table is already up-to-date
-		if (!insres.second && insres.first->second != ~u32(0))
-		{
-			BinaryReader reader(m_directoryTableMem + insres.first->second);
-			u32 oldAttr = reader.ReadFileAttributes();
-
-			if (IsDirectory(oldAttr))  // Ignore updating directories.. they should always be the same regardless
-			{
-				UBA_ASSERT(isDirectory || !shouldWriteToDisk);
-				return true;
-			}
-			reader.ReadVolumeSerial();
-			u64 oldFileIndex = reader.ReadFileIndex();(void)oldFileIndex;
-
-			u64 oldLastWriteTime = reader.ReadFileTime();
-			if (lastWriteTime == oldLastWriteTime)
-			{
-				#if !PLATFORM_WINDOWS
-				UBA_ASSERT(oldFileIndex == fileIndex); // Checking so it is really the same file
-				#endif
-				u64 oldSize = reader.ReadFileSize();
-				if (oldSize == fileSize && oldAttr == attributes) // Only attributes could change from a chmod
-					return true;
-				// TODO: Somehow this can happen and I have no idea how. last written time should be set on close file so it shouldnt be possible.
-				//else
-				//	m_logger.Warning(TC("Somehow file %s has same last written time at two points in time but different size (old %llu new %llu)"), fileName.data, oldSize, fileSize);
-			}
-		}
-
-		// There are directory crawlers happening in parallel so we need to really make sure to invalidate this one since a crawler can actually
-		// hit this file with information from a query before it was written.. and then it will turn it back to "verified" using old info
-		if (registerRealFile && invalidateStorage)
-			m_storage.InvalidateCachedFileInfo(fileNameKey);
-
-		FileEntryAdded(fileNameKey, lastWriteTime, fileSize);
-
-		u8 temp[1024];
-		u64 written;
-		u64 entryPos;
-		{
-			BinaryWriter writer(temp, 0, sizeof(temp));
-			writer.WriteStringKey(dirKey);
-			writer.Write7BitEncoded(dir.latestOffset); // Previous entry for same directory
-			writer.Write7BitEncoded(1); // Count
-			writer.WriteString(lastSlash + 1);
-			entryPos = writer.GetPosition();
-			writer.WriteFileAttributes(attributes);
-			writer.WriteVolumeSerial(m_volumeCache.GetSerialIndex(volumeSerial));
-			writer.WriteFileIndex(fileIndex);
-			if (!isDirectory)
-			{
-				writer.WriteFileTime(lastWriteTime);
-				writer.WriteFileSize(fileSize);
-			}
-			written = writer.GetPosition();
-		}
-
-		#if UBA_DEBUG_TRACK_DIR
-		m_debugLogger->Info(TC("TRACKADD    %s (Size: %llu, Attr: %u, Key: %s, Id: %llu)"), fileName.data, fileSize, attributes, KeyToString(fileNameKey).data, fileIndex);
-		#endif
-
-
-		SCOPED_WRITE_LOCK(dirTable.m_memoryLock, memoryLock);
-		u32 writePos = dirTable.m_memorySize;
-		EnsureDirectoryTableMemory(writePos + 8 + written);
-		BinaryWriter writer(m_directoryTableMem + writePos);
-		writer.Write7BitEncoded(written); // Storage size
-		insres.first->second = dirTable.m_memorySize + u32(writer.GetPosition() + entryPos); // Storing position to last write time
-		u32 tableOffset = u32(writer.GetPosition()) + sizeof(StringKey);
-		writer.WriteBytes(temp, written);
-		dir.latestOffset = dirTable.m_memorySize + tableOffset;
-		dirTable.m_memorySize += u32(writer.GetPosition());
-		return true;
-	}
-
-	DirTableSize Session::RegisterDeleteFile(StringKey fileNameKey, const StringView& fileName, DirectoryTableOverlay* overlay)
-	{
-		// Remote is not updating its own directory table
-		if (m_runningRemote)
-			return GetDirectoryTableSize(overlay);
-
-		auto& dirTable = m_directoryTable;
-
-		StringKey dirKey;
-		const tchar* lastSlash;
-		StringBuffer<> dirName;
-		if (!GetDirKey(dirKey, dirName, lastSlash, fileName))
-			return {InvalidTableOffset, InvalidTableOffset};
-
-		SCOPED_READ_LOCK(dirTable.m_lookupLock, lookupCs);
-		auto res = dirTable.m_lookup.find(dirKey);
-		if (res == dirTable.m_lookup.end())
-			return {};
-		DirectoryTable::Directory& dir = res->second;
-		lookupCs.Leave();
-		SCOPED_WRITE_LOCK(dir.lock, dirLock);
-
-		while (dir.populatedOffset == 0)
-		{
-			dirLock.Leave();
-			Sleep(1);
-			dirLock.Enter();
-		}
-		
-		if (dir.populatedOffset == 2) // Invalid directory, don't need to register this
-			return {};
-
-		UBA_ASSERTF(dir.populatedOffset == 1, TC("Registering deleted file %s with unexpect dir.populatedOffset %u "), fileName.data, dir.populatedOffset);
-
-		if (fileNameKey == StringKeyZero)
-		{
-			StringBuffer<> forKey;
-			forKey.Append(fileName);
-			if (CaseInsensitiveFs)
-				forKey.MakeLower();
-			fileNameKey = ToStringKey(forKey);
-		}
-		
-		// Does not exist, no need adding to file table
-		if (dir.files.erase(fileNameKey) == 0)
-			return {};
-
-		u8 temp[1024];
-		u64 written;
-		{
-			BinaryWriter writer(temp, 0, sizeof(temp));
-			writer.WriteStringKey(dirKey);
-			writer.Write7BitEncoded(dir.latestOffset); // Previous entry for same directory
-			writer.Write7BitEncoded(1); // Count
-			writer.WriteString(lastSlash + 1);
-			writer.WriteFileAttributes(0);
-			writer.WriteVolumeSerial(0);
-			writer.WriteFileIndex(0);
-			if (true) // !IsDirectory()
-			{
-				writer.WriteFileTime(0);
-				writer.WriteFileSize(0);
-			}
-			written = writer.GetPosition();
-		}
-
-		#if UBA_DEBUG_TRACK_DIR
-		m_debugLogger->Info(TC("TRACKDEL    %s (Key: %s)"), fileName.data, KeyToString(fileNameKey).data);
-		#endif
-
-		{
-			SCOPED_WRITE_LOCK(dirTable.m_memoryLock, memoryLock);
-			u32 writePos = dirTable.m_memorySize;
-			EnsureDirectoryTableMemory(writePos + 8 + written);
-			BinaryWriter writer(m_directoryTableMem + writePos);
-			writer.Write7BitEncoded(written); // Storage size
-			u32 tableOffset = u32(writer.GetPosition()) + sizeof(StringKey);
-			writer.WriteBytes(temp, written);
-			dir.latestOffset = dirTable.m_memorySize + tableOffset;
-			dirTable.m_memorySize += u32(writer.GetPosition());
-		}
-
-		return GetDirectoryTableSize(overlay);
-	}
-
 	bool Session::InitOverlayFiles(u8* outMem, u32 capacity, u32& outSize, StringView overlayFiles)
 	{
 		auto& dirTable = m_directoryTable;
@@ -1835,13 +1256,11 @@ namespace uba
 	}
 
 	Session::Session(const SessionCreateInfo& info, const tchar* logPrefix, bool runningRemote, WorkManager& workManager)
-	:	m_storage(info.storage)
+	:	DirectoryTableHolder(info.logWriter, logPrefix, runningRemote)
+	,	m_storage(info.storage)
 	,	m_sharedMemory(info.storage.GetSharedMemory())
-	,	m_logger(info.logWriter, logPrefix)
 	,	m_workManager(workManager)
 	,	m_ownsTrace(info.trace == nullptr)
-	,	m_directoryTableMemory(TC("SessionDirectoryTable"))
-	,	m_directoryTable(m_directoryTableMemory)
 	,	m_fileMappingBuffer(m_logger, m_sharedMemory)
 	,	m_nameToHashTableMem(TC("SessionNameToHashTable"))
 	,	m_processCommunicationAllocator(m_logger, TC("CommunicationAllocator"))
@@ -1885,7 +1304,6 @@ namespace uba
 		}
 		m_sessionDir.Append(PathSeparator);
 
-		m_runningRemote = runningRemote;
 		m_allowCustomAllocator = info.allowCustomAllocator;
 		m_allowMemoryMaps = info.allowMemoryMaps && IsWindows;
 		m_allowKeepFilesInMemory = info.allowKeepFilesInMemory && IsWindows;
@@ -1957,18 +1375,8 @@ namespace uba
 		m_storage.CreateDirectory(m_sessionLogDir.data);
 		m_sessionLogDir.EnsureEndsWithSlash();
 
-		// We never want to populate files in Temp
-		#if PLATFORM_WINDOWS
 		if (info.treatTempDirAsEmpty)
-		{
-			wchar_t systemTemp[256];
-			GetEnvironmentVariableW(L"TEMP", systemTemp, 256);
-			StringBuffer<> temp;
-			FixPath(systemTemp, nullptr, 0, temp);
-			temp.MakeLower();
-			m_directoryForcedEmpty = ToStringKey(temp);
-		}
-		#endif
+			SetTreatTempDirAsEmpty();
 
 		if (info.traceOutputFile)
 			m_traceOutputFile = info.traceOutputFile;
@@ -1998,14 +1406,6 @@ namespace uba
 		m_fileMappingTableMem = MapViewOfFile(m_logger, m_fileMappingTableHandle, FILE_MAP_WRITE, 0, FileMappingTableMemSize, m_allowDiscardVirtualMemory);
 		LockMemory(m_logger, m_fileMappingTableMem, FileMappingTableMemSize, TC("FileMappings"));
 		UBA_ASSERT(m_fileMappingTableMem);
-
-		m_directoryTableHandle = uba::CreateMemoryMappingW(m_logger, PAGE_READWRITE|SEC_RESERVE, DirTableMemSize, nullptr, TC("DirMappings"));
-		UBA_ASSERT(m_directoryTableHandle.IsValid());
-		m_directoryTableMem = MapViewOfFile(m_logger, m_directoryTableHandle, FILE_MAP_WRITE, 0, DirTableMemSize);
-		UBA_ASSERT(m_directoryTableMem);
-
-		m_directoryTable.m_memory = m_directoryTableMem;
-		m_directoryTable.m_lookup.reserve(30000);
 		m_fileMappingTableLookup.reserve(70000);
 
 		u64 capacity = 92ull*1024*1024*1024;
@@ -2166,9 +1566,6 @@ namespace uba
 		UnmapViewOfFile(m_logger, m_fileMappingTableMem, FileMappingTableMemSize, TC("FileMappingTable"), m_allowDiscardVirtualMemory);
 		CloseFileMapping(m_logger, m_fileMappingTableHandle, TC("FileMappingTable"));
 
-		UnmapViewOfFile(m_logger, m_directoryTableMem, DirTableMemSize, TC("DirectoryTable"), false); // Can't discard since only parts might be mapped
-		CloseFileMapping(m_logger, m_directoryTableHandle, TC("DirectoryTable"));
-
 		//#if !UBA_DEBUG
 		//u32 count;
 		//DeleteAllFiles(m_logger, m_sessionDir.data, count);
@@ -2275,52 +1672,6 @@ namespace uba
 		return h;
 	}
 
-	bool Session::RefreshDirectory(const tchar* dirName, bool forceRegister)
-	{
-		UBA_ASSERT(!m_runningRemote);
-
-		StringBuffer<> dirPath;
-		StringKeyHasher hasher;
-		GetKeyAndFixedName(dirPath, hasher, dirName);
-
-		StringKey dirKey = ToStringKey(hasher);
-
-		auto& dirTable = m_directoryTable;
-		SCOPED_READ_LOCK(dirTable.m_lookupLock, lookupLock);
-		auto res = dirTable.m_lookup.find(dirKey);
-		if (res == dirTable.m_lookup.end())
-		{
-			lookupLock.Leave();
-			if (forceRegister)
-				WriteDirectoryEntries(dirKey, dirPath);
-			return true;
-		}
-		DirectoryTable::Directory& dir = res->second;
-		lookupLock.Leave();
-		SCOPED_WRITE_LOCK(dir.lock, dirLock);
-
-		while (dir.populatedOffset == 0)
-		{
-			dirLock.Leave();
-			Sleep(1);
-			dirLock.Enter();
-		}
-		UBA_ASSERT(dir.populatedOffset == 1);
-
-		//m_directoryTable.PopulateDirectoryNoLock(hasher, dir);
-
-		u32 tableOffset;
-		return WriteDirectoryEntriesInternal(dir, dirKey, dirPath, true, tableOffset);
-	}
-
-	bool Session::RegisterNewFile(const tchar* filePath)
-	{
-		UBA_ASSERT(!m_runningRemote);
-		StringBuffer<> fixedFilePath;
-		auto key = GetKeyAndFixedName(fixedFilePath, filePath);
-		return RegisterCreateFileForWrite(key, fixedFilePath, true);
-	}
-
 	bool Session::RegisterVirtualFile(const tchar* filePath, const tchar* sourceFile, u64 sourceOffset, u64 sourceSize)
 	{
 		UBA_ASSERT(!m_runningRemote);
@@ -2377,25 +1728,6 @@ namespace uba
 		return true;
 	}
 
-	void Session::RegisterDeleteFile(const tchar* filePath)
-	{
-		UBA_ASSERT(!m_runningRemote);
-		StringBuffer<> fixedFilePath;
-		auto key = GetKeyAndFixedName(fixedFilePath, filePath);
-		RegisterDeleteFile(key, fixedFilePath, nullptr);
-	}
-
-	bool Session::RegisterNewDirectory(const tchar* directoryPath)
-	{
-		UBA_ASSERT(!m_runningRemote);
-		StringBuffer<> fixedDirPath;
-		auto dirKey = GetKeyAndFixedName(fixedDirPath, directoryPath);
-		if (!RegisterCreateFileForWrite(dirKey, fixedDirPath, true))
-			return false;
-		WriteDirectoryEntries(dirKey, fixedDirPath);
-		return true;
-	}
-
 	void Session::RegisterCustomService(CustomServiceFunction&& function)
 	{
 		m_customServiceFunction = function;
@@ -2888,18 +2220,6 @@ namespace uba
 		return true;
 	}
 
-	DirTableSize Session::GetDirectoryTableSize(DirectoryTableOverlay* overlay)
-	{
-		u32 overlaySize = 0u;
-		if (overlay)
-		{
-			SCOPED_READ_LOCK(overlay->lock, lock);
-			overlaySize = overlay->size;
-		}
-		SCOPED_READ_LOCK(m_directoryTable.m_memoryLock, lock);
-		return {m_directoryTable.m_memorySize, overlaySize};
-	}
-
 	u32 Session::GetFileMappingSize()
 	{
 		SCOPED_READ_LOCK(m_fileMappingTableMemLock, lock);
@@ -2997,13 +2317,18 @@ namespace uba
 #endif
 	}
 
-	bool Session::ShouldWriteToDisk(const StringView& fileName)
+	bool Session::ShouldWriteToDisk(StringView fileName)
 	{
 		if (m_shouldWriteToDisk)
 			return true;
 		return fileName.EndsWith(TCV(".h"));
 	}
 
+	void Session::InvalidateCachedFileInfo(const StringKey& fileNameKey)
+	{
+		m_storage.InvalidateCachedFileInfo(fileNameKey);
+	}
+
 	bool Session::PrepareProcess(ProcessImpl& process, bool isChild, StringBufferBase& outRealApplication, const tchar*& outRealWorkingDir)
 	{
 		ProcessStartInfoHolder& startInfo = process.m_startInfo;
@@ -3914,10 +3239,6 @@ namespace uba
 		return true;
 	}
 	
-	void Session::FileEntryAdded(StringKey fileNameKey, u64 lastWritten, u64 size)
-	{
-	}
-
 	bool Session::FlushWrittenFiles(ProcessImpl& process)
 	{
 		return true;
@@ -4105,26 +3426,6 @@ namespace uba
 		return true;
 	}
 
-	void Session::EnsureDirectoryTableMemory(u64 neededSize)
-	{
-		auto& dirTable = m_directoryTable;
-		if (neededSize <= m_directoryTableMemCommitted)
-			return;
-
-		u64 newSize = AlignUp(neededSize, u64(1024*1024));
-		if (newSize > DirTableMemSize)
-		{
-			static bool called = m_logger.Error(TC("Directory table overflow. DirTableMemSize need to be increased (Size: %llu)"), dirTable.m_memorySize);
-		}
-
-		u64 toCommit = newSize - m_directoryTableMemCommitted;
-		u8* address = m_directoryTableMem + m_directoryTableMemCommitted;
-		if (!MapViewCommit(m_logger, address, toCommit, m_allowDiscardVirtualMemory))
-			m_logger.Error(TC("Failed to commit memory for directory table (Committed: %llu, ToCommit: %llu) (%s)"), m_directoryTableMemCommitted, toCommit, LastErrorToText().data);
-		m_directoryTableMemCommitted += toCommit;
-		LockMemory(m_logger, address, toCommit, TC("DirectoryTable"));
-	}
-
 	void Session::GetSessionInfo(StringBufferBase& out)
 	{
 		out.Append(TCV("Cas:"));
@@ -4554,6 +3855,8 @@ namespace uba
 	}
 #endif
 
+	////////////////////////////////////////////////////////////////////////////////////////////////////
+
 	void GenerateNameForProcess(StringBufferBase& out, const tchar* arguments, u32 counterSuffix)
 	{
 		const tchar* start = arguments;
@@ -4610,4 +3913,6 @@ namespace uba
 
 		return false;
 	}
+
+	////////////////////////////////////////////////////////////////////////////////////////////////////
 }
diff --git a/Engine/Source/Programs/UnrealBuildAccelerator/Common/Public/UbaDirectoryTableHolder.h b/Engine/Source/Programs/UnrealBuildAccelerator/Common/Public/UbaDirectoryTableHolder.h
new file mode 100644
--- /dev/null
+++ b/Engine/Source/Programs/UnrealBuildAccelerator/Common/Public/UbaDirectoryTableHolder.h
@@ -0,0 +1,75 @@
+// Copyright Epic Games, Inc. All Rights Reserved.
+
+#pragma once
+
+#include "UbaDirectoryTable.h"
+#include "UbaFile.h"
+#include "UbaSharedMemoryView.h"
+
+namespace uba
+{
+	class Logger;
+
+	////////////////////////////////////////////////////////////////////////////////////////////////////
+
+	struct DirectoryTableOverlay
+	{
+		ReaderWriterLock lock;
+		SharedMemoryHandle handle;
+		u32 size = 0;
+	};
+
+	////////////////////////////////////////////////////////////////////////////////////////////////////
+
+	class DirectoryTableHolder
+	{
+	public:
+		DirectoryTableHolder(LogWriter& logWriter, const tchar* logPrefix = TC(""), bool runningRemote = false);
+		virtual ~DirectoryTableHolder();
+		void SetTreatTempDirAsEmpty(); // We never want to populate files in Temp
+
+		bool RefreshDirectory(const tchar* dirName, bool forceRegister = false); // Tell uba a directory on disk has been changed by some other system while session is running
+		bool RegisterNewFile(const tchar* filePath); // Tell uba a new file on disk has been added by some other system while session is running
+		void RegisterDeleteFile(const tchar* filePath); // Tell uba a file on disk has been deleted by some other system while session is running
+		bool RegisterNewDirectory(const tchar* directoryPath); // Tell uba a directory on disk has been added by some other system while session is running
+
+		const u8* GetDirectoryTableMemory();
+		DirTableSize GetDirectoryTableSize(DirectoryTableOverlay* overlay);
+
+
+		MutableLogger m_logger;
+
+	protected:
+
+		bool RegisterCreateFileForWrite(StringKey fileNameKey, StringView fileName, bool registerRealFile, u64 fileSize = 0, u64 lastWriteTime = 0, bool invalidateStorage = true);
+		DirTableSize RegisterDeleteFile(StringKey fileNameKey, const StringView& fileName, DirectoryTableOverlay* overlay);
+
+		u32 WriteDirectoryEntries(const StringKey& dirKey, StringView dirPath, u32* outTableOffset = nullptr);
+		void WriteDirectoryEntriesRecursive(const StringKey& dirKey, StringView dirPath, u32& outTableOffset);
+		bool WriteDirectoryEntriesInternal(DirectoryTable::Directory& dir, const StringKey& dirKey, StringView dirPath, bool isRefresh, u32& outTableOffset);
+		void EnsureDirectoryTableMemory(u64 neededSize);
+
+		virtual void FileEntryAdded(StringKey fileNameKey, u64 lastWritten, u64 size);
+		virtual bool ShouldWriteToDisk(StringView fileName);
+		virtual void InvalidateCachedFileInfo(const StringKey& fileNameKey);
+
+		bool m_runningRemote;
+
+		Atomic<u64> m_fileIndexCounter = 8000000000;
+		MemoryBlock m_directoryTableMemory;
+		FileMappingHandle m_directoryTableHandle;
+		u8* m_directoryTableMem;
+		u64 m_directoryTableMemCommitted = 0;
+		DirectoryTable m_directoryTable;
+		StringKey m_directoryForcedEmpty;
+		VolumeCache m_volumeCache;
+	};
+
+	////////////////////////////////////////////////////////////////////////////////////////////////////
+
+	bool GetDirKey(StringKey& outDirKey, StringBufferBase& outDirName, const tchar*& outLastSlash, const StringView& fileName);
+	StringKey GetKeyAndFixedName(StringBuffer<>& outFixedFilePath, StringKeyHasher& outHasher, const tchar* filePath, StringKey* outDirKey = nullptr, const tchar** outBaseFileName = nullptr);
+	StringKey GetKeyAndFixedName(StringBuffer<>& outFixedFilePath, const tchar* filePath, StringKey* outDirKey = nullptr, const tchar** outBaseFileName = nullptr);
+
+	////////////////////////////////////////////////////////////////////////////////////////////////////
+}
diff --git a/Engine/Source/Programs/UnrealBuildAccelerator/Common/Public/UbaSession.h b/Engine/Source/Programs/UnrealBuildAccelerator/Common/Public/UbaSession.h
--- a/Engine/Source/Programs/UnrealBuildAccelerator/Common/Public/UbaSession.h
+++ b/Engine/Source/Programs/UnrealBuildAccelerator/Common/Public/UbaSession.h
@@ -6,7 +6,7 @@
 #include "UbaDependencyCrawler.h"
 #include "UbaFile.h"
 #include "UbaDefinitions.h"
-#include "UbaDirectoryTable.h"
+#include "UbaDirectoryTableHolder.h"
 #include "UbaFileMappingBuffer.h"
 #include "UbaProcessHandle.h"
 #include "UbaRootPaths.h"
@@ -58,7 +58,7 @@ namespace uba
 	struct NextProcessInfo;
 	using RootsHandle = u64;
 
-	class Session
+	class Session : public DirectoryTableHolder
 	{
 	public:
 		ProcessHandle RunProcess(const ProcessStartInfo& startInfo, bool async = true, bool enableDetour = true); // Run process. if async is false it will not return until process is done
@@ -66,13 +66,9 @@ namespace uba
 		void CancelAllProcesses(); // Cancel all processes
 
 		void PrintSummary(Logger& logger); // Print summary stats of session
-		bool RefreshDirectory(const tchar* dirName, bool forceRegister = false); // Tell uba a directory on disk has been changed by some other system while session is running
-		bool RegisterNewFile(const tchar* filePath); // Tell uba a new file on disk has been added by some other system while session is running
 		bool RegisterVirtualFile(const tchar* filePath, const tchar* sourceFile, u64 sourceOffset, u64 sourceSize);
 		bool CreateVirtualFile(const tchar* filePath, const void* memory, u64 memorySize, bool transient, bool registerToDir = true);
 		bool DeleteVirtualFile(const tchar* filePath, bool registerToDir);
-		void RegisterDeleteFile(const tchar* filePath); // Tell uba a file on disk has been deleted by some other system while session is running
-		bool RegisterNewDirectory(const tchar* directoryPath); // Tell uba a directory on disk has been added by some other system while session is running
 
 		RootsHandle RegisterRoots(const void* rootsData, uba::u64 rootsDataSize);
 
@@ -117,10 +113,6 @@ namespace uba
 
 		bool InitOverlayFiles(u8* outMem, u32 capacity, u32& outSize, StringView overlayFiles);
 
-		// For testing
-		const u8* GetDirectoryTableMemory() { return m_directoryTableMem; }
-		DirTableSize GetDirectoryTableSize(DirectoryTableOverlay* overlay);
-
 		void FlushDeadProcesses();
 
 		const tchar* FixLogFile(const tchar* logFile, StringBufferBase& temp, const tchar* processArguments, u32 processId);
@@ -147,8 +139,9 @@ namespace uba
 		bool IsRarelyRead(ProcessImpl& process, const StringView& fileName) const;
 		bool IsMaybeReadAfterWritten(ProcessImpl& process, const StringView& fileName) const;
 		bool IsKnownSystemFile(const tchar* applicationName);
-		bool ShouldWriteToDisk(const StringView& fileName);
-		u32 WriteDirectoryEntries(const StringKey& dirKey, const StringView& dirPath, u32* outTableOffset = nullptr);
+		virtual bool ShouldWriteToDisk(StringView fileName) override;
+		virtual void InvalidateCachedFileInfo(const StringKey& fileNameKey);
+
 		u32 AddFileMapping(StringKey fileNameKey, const tchar* fileName, const tchar* newFileName, u64 fileSize = InvalidValue);
 		
 		struct MemoryMap { StringBuffer<128> name; u64 size = 0; };
@@ -157,9 +150,6 @@ namespace uba
 		bool GetOrCreateMemoryMapFromStorage(MemoryMap& out, StringKey fileNameKey, const tchar* fileName, const CasKey& casKey, u64 alignment);
 		void ThreadGcFileMappings();
 
-		bool RegisterCreateFileForWrite(StringKey fileNameKey, const StringView& fileName, bool registerRealFile, u64 fileSize = 0, u64 lastWriteTime = 0, bool invalidateStorage = true);
-		DirTableSize RegisterDeleteFile(StringKey fileNameKey, const StringView& fileName, DirectoryTableOverlay* overlay);
-
 		virtual bool PrepareProcess(ProcessImpl& process, bool isChild, StringBufferBase& outRealApplication, const tchar*& outRealWorkingDir);
 		virtual void* GetProcessEnvironmentVariables();
 		virtual void PrintSessionStats(Logger& logger);
@@ -197,7 +187,6 @@ namespace uba
 		virtual bool GetSymbols(const tchar* application, bool isLinux, bool isArm, BinaryReader& reader, BinaryWriter& writer);
 		virtual bool CheckRemapping(ProcessImpl& process, BinaryReader& reader, BinaryWriter& writer);
 		virtual bool RunSpecialProgram(ProcessImpl& process, BinaryReader& reader, BinaryWriter& writer);
-		virtual void FileEntryAdded(StringKey fileNameKey, u64 lastWritten, u64 size);
 		virtual bool FlushWrittenFiles(ProcessImpl& process);
 		virtual bool UpdateEnvironment(ProcessImpl& process, const StringView& reason, bool resetStats);
 		virtual bool LogLine(ProcessImpl& process, const tchar* line, LogEntryType logType);
@@ -210,11 +199,8 @@ namespace uba
 		bool WriteMemoryToDisk(FileAccessor& destinationFile, const void* fileMem, u64 fileSize, StringBufferBase* outWriteMode);
 		bool GetFileMemory(const Function<bool(const void*, u64)>& func, const StringKey& fileNameKey, StringView filePath, bool deleteInternalMapping);
 		void AddEnvironmentVariableNoLock(StringView key, StringView value);
-		bool WriteDirectoryEntriesInternal(DirectoryTable::Directory& dir, const StringKey& dirKey, StringView dirPath, bool isRefresh, u32& outTableOffset);
-		void WriteDirectoryEntriesRecursive(const StringKey& dirKey, StringView dirPath, u32& outTableOffset);
 		bool CopyImports(Vector<BinaryModule>& out, const tchar* library, tchar* applicationDir, tchar* applicationDirEnd, UnorderedSet<TString>& handledImports, const tchar* const* loaderPaths, bool isLinuxBinary);
 		bool CreateProcessJobObject();
-		void EnsureDirectoryTableMemory(u64 neededSize);
 		void GetSessionInfo(StringBufferBase& out);
 
 		bool HasVfs(RootsHandle handle) const;
@@ -232,7 +218,6 @@ namespace uba
 
 		Storage& m_storage;
 		SharedMemoryAllocator& m_sharedMemory;
-		MutableLogger m_logger;
 		WorkManager& m_workManager;
 
 		StringBuffer<32> m_id;
@@ -245,7 +230,6 @@ namespace uba
 		StringBuffer<MaxPath> m_tempPath;
 
 
-		bool m_runningRemote;
 		bool m_allowCustomAllocator;
 		bool m_allowMemoryMaps;
 		bool m_allowKeepFilesInMemory;
@@ -269,15 +253,6 @@ namespace uba
 		u64 m_keepOutputFileMemoryMapsThreshold;
 
 		Atomic<u32> m_processIdCounter = 0;
-
-		MemoryBlock m_directoryTableMemory;
-
-		FileMappingHandle m_directoryTableHandle;
-		u8* m_directoryTableMem;
-		u64 m_directoryTableMemCommitted = 0;
-		DirectoryTable m_directoryTable;
-		StringKey m_directoryForcedEmpty;
-	
 		FileMappingHandle m_fileMappingTableHandle;
 		FileMappingBuffer m_fileMappingBuffer;
 
@@ -325,8 +300,6 @@ namespace uba
 		static constexpr u64 NameToHashMemSize = 48*1024*1024;
 		MemoryBlock m_nameToHashTableMem;
 
-		Atomic<u64> m_fileIndexCounter = 8000000000;
-
 		FixedSizeFileMappingAllocator m_processCommunicationAllocator;
 		std::string m_detoursLibrary[2][2]; // [nonlinux/linux][x64/arm64]
 
@@ -398,8 +371,6 @@ namespace uba
 		SemaphoreServiceInfo* m_semaphoreServiceInfo;
 		#endif
 
-		VolumeCache m_volumeCache;
-
 		DependencyCrawler m_dependencyCrawler;
 
 		friend class ProcessImpl;
@@ -633,13 +604,6 @@ namespace uba
 		TString breadcrumbs;
 		Vector<u8> overlayFiles;
 	};
-
-	struct DirectoryTableOverlay
-	{
-		ReaderWriterLock lock;
-		SharedMemoryHandle handle;
-		u32 size = 0;
-	};
 }
 
 template<> struct std::hash<uba::ProcessHandle> { size_t operator()(const uba::ProcessHandle& g) const { return g.GetHash(); } };
diff --git a/Engine/Source/Programs/UnrealBuildAccelerator/Core/Public/UbaDirectoryTable.h b/Engine/Source/Programs/UnrealBuildAccelerator/Core/Public/UbaDirectoryTable.h
--- a/Engine/Source/Programs/UnrealBuildAccelerator/Core/Public/UbaDirectoryTable.h
+++ b/Engine/Source/Programs/UnrealBuildAccelerator/Core/Public/UbaDirectoryTable.h
@@ -53,9 +53,9 @@ namespace uba
 			Exists_Maybe,
 		};
 
-		inline Exists EntryExists(StringKey entryKey, const StringView& entryName, bool checkIfDir = false, u32* outTableOffset = nullptr);
-		inline Exists EntryExistsNoLock(StringKey entryKey, const StringView& entryName, bool checkIfDir = false, u32* outTableOffset = nullptr);
-		inline Exists EntryExists(const StringView& str, bool checkIfDir = false);
+		inline Exists EntryExists(StringKey entryKey, StringView entryName, bool checkIfDir = false, u32* outTableOffset = nullptr);
+		inline Exists EntryExistsNoLock(StringKey entryKey, StringView entryName, bool checkIfDir = false, u32* outTableOffset = nullptr);
+		inline Exists EntryExists(StringView str, bool checkIfDir = false);
 
 		struct EntryInformation
 		{
@@ -287,13 +287,13 @@ namespace uba
 		}
 	}
 
-	DirectoryTable::Exists DirectoryTable::EntryExists(StringKey entryKey, const StringView& entryName, bool checkIfDir, u32* outTableOffset)
+	DirectoryTable::Exists DirectoryTable::EntryExists(StringKey entryKey, StringView entryName, bool checkIfDir, u32* outTableOffset)
 	{
 		SCOPED_READ_LOCK(m_lookupLock, lock);
 		return EntryExistsNoLock(entryKey, entryName, checkIfDir, outTableOffset);
 	}
 
-	DirectoryTable::Exists DirectoryTable::EntryExistsNoLock(StringKey entryKey, const StringView& entryName, bool checkIfDir, u32* outTableOffset)
+	DirectoryTable::Exists DirectoryTable::EntryExistsNoLock(StringKey entryKey, StringView entryName, bool checkIfDir, u32* outTableOffset)
 	{
 		u32 startSkip = 2;
 		if (checkIfDir)
@@ -370,7 +370,7 @@ namespace uba
 		return Exists_Maybe;
 	}
 
-	DirectoryTable::Exists DirectoryTable::EntryExists(const StringView& str, bool checkIfDir)
+	DirectoryTable::Exists DirectoryTable::EntryExists(StringView str, bool checkIfDir)
 	{
 		StringBuffer<> str2(str);
 		if (str2[str2.count-1] == PathSeparator)
diff --git a/Engine/Source/Programs/UnrealBuildAccelerator/Test/Private/UbaTestAll.h b/Engine/Source/Programs/UnrealBuildAccelerator/Test/Private/UbaTestAll.h
--- a/Engine/Source/Programs/UnrealBuildAccelerator/Test/Private/UbaTestAll.h
+++ b/Engine/Source/Programs/UnrealBuildAccelerator/Test/Private/UbaTestAll.h
@@ -37,6 +37,7 @@ namespace uba
 		UBA_TEST(TestClientServer2) \
 		UBA_TEST(TestClientServerMem) \
 		UBA_TEST(TestClientServerCrypto) \
+		UBA_TEST(TestDirectoryTable) \
 		UBA_TEST(TestStorage) \
 		UBA_TEST(TestRemoteStorageStore) \
 		UBA_TEST(TestRemoteStorageFetch) \
diff --git a/Engine/Source/Programs/UnrealBuildAccelerator/Test/Private/UbaTestCache.cpp b/Engine/Source/Programs/UnrealBuildAccelerator/Test/Private/UbaTestCache.cpp
--- a/Engine/Source/Programs/UnrealBuildAccelerator/Test/Private/UbaTestCache.cpp
+++ b/Engine/Source/Programs/UnrealBuildAccelerator/Test/Private/UbaTestCache.cpp
@@ -289,7 +289,6 @@ namespace uba
 	}
 
 	void GetTestAppPath(LoggerWithWriter& logger, StringBufferBase& out);
-	StringKey GetKeyAndFixedName(StringBuffer<>& fixedFilePath, const tchar* filePath, StringKey* outDirKey = nullptr, const tchar** outBaseFileName = nullptr);
 	void InvalidateCachedInfo(StorageImpl& storage, StringBufferBase& fileName)
 	{
 		StringBuffer<> fixedFilePath;
diff --git a/Engine/Source/Programs/UnrealBuildAccelerator/Test/Private/UbaTestDirectoryTable.cpp b/Engine/Source/Programs/UnrealBuildAccelerator/Test/Private/UbaTestDirectoryTable.cpp
new file mode 100644
--- /dev/null
+++ b/Engine/Source/Programs/UnrealBuildAccelerator/Test/Private/UbaTestDirectoryTable.cpp
@@ -0,0 +1,35 @@
+// Copyright Epic Games, Inc. All Rights Reserved.
+
+#include "UbaDirectoryTableHolder.h"
+#include "UbaTest.h"
+
+namespace uba
+{
+	bool TestDirectoryTable(LoggerWithWriter& logger, const StringBufferBase& testRootDir)
+	{
+		DirectoryTableHolder holder(logger.m_writer, TC(""), false);
+		//holder.SetTreatTempDirAsEmpty();
+
+		if (!holder.RefreshDirectory(testRootDir.data, true))
+			return false;
+
+		MemoryBlock dirMem(TC("Mem"));
+		DirectoryTable dir(dirMem);
+		dir.Init(holder.GetDirectoryTableMemory(), 0, holder.GetDirectoryTableSize(nullptr).main);
+		
+		if (dir.EntryExists(testRootDir, false) != DirectoryTable::Exists_Yes)
+			return false;
+
+		if (dir.EntryExists(testRootDir, true) != DirectoryTable::Exists_Yes)
+			return false;
+
+		StringBuffer<> fooFile;
+		if (!CreateTestFile(fooFile, logger, testRootDir, TCV("foo.txt"), TCV("")))
+			return false;
+
+		if (!holder.RegisterNewFile(fooFile.data))
+			return false;
+
+		return true;
+	}
+}
