[UBA]
* Added some documentation to DirectoryTable

diff --git a/Engine/Source/Programs/UnrealBuildAccelerator/Core/Public/UbaDirectoryTable.h b/Engine/Source/Programs/UnrealBuildAccelerator/Core/Public/UbaDirectoryTable.h
--- a/Engine/Source/Programs/UnrealBuildAccelerator/Core/Public/UbaDirectoryTable.h
+++ b/Engine/Source/Programs/UnrealBuildAccelerator/Core/Public/UbaDirectoryTable.h
@@ -8,8 +8,33 @@
 namespace uba
 {
 	////////////////////////////////////////////////////////////////////////////////////////////////////
+	// 
+	// DirectoryTable is a journal that can be parsed into a table of directories and files
+	// When a directory changes it will create a new entry in the journal pointing back to the previous
+	// Files are attached to the directory entries and can be populated on demand into the files lookup.
+	//
+	// Directory:
+	//   StringKey dir (16 bytes)
+	//   enc7 OffsetToPreviousEntry  (this is offset to previous entry after StringKey)
+	//   if OffsetToPreviousEntry == InvalidTableOffset
+	//      enc7 FileAttributes
+	//      enc7 VolumeSerialIndex
+	//      u64 FileIndex
+	//   enc7 ItemCount
+	//      Items
+	// 
+	// Item (in directory):
+	//   String name
+	//   enc7 FileAttributes
+	//   enc7 VolumeSerialIndex
+	//   u64 FileIndex
+	//   if not directory (FileAttributes provide info)
+	//      u64 FileTime
+	//      enc7 FileSize
+	// 
+	////////////////////////////////////////////////////////////////////////////////////////////////////
 
-	enum : u64 { InvalidTableOffset = 0 };
+	enum : u64 { InvalidTableOffset = 0 }; // Important that it is zero to compress enc7 as good as possible
 	enum : u64 { OverlayTableOffset = 256*1024*1024 };
 
 	struct DirTableSize { u32 main = 0; u32 overlay = 0; };
