[UBA]
* Cleaned up lots of code around directory table in preparation for proper implementation of dir table overlay

diff --git a/Engine/Source/Programs/UnrealBuildAccelerator/Common/Private/UbaDirectoryTableHolder.cpp b/Engine/Source/Programs/UnrealBuildAccelerator/Common/Private/UbaDirectoryTableHolder.cpp
--- a/Engine/Source/Programs/UnrealBuildAccelerator/Common/Private/UbaDirectoryTableHolder.cpp
+++ b/Engine/Source/Programs/UnrealBuildAccelerator/Common/Private/UbaDirectoryTableHolder.cpp
@@ -11,7 +11,7 @@ namespace uba
 {
 	////////////////////////////////////////////////////////////////////////////////////////////////////
 
-	bool GetDirKey(StringKey& outDirKey, StringBufferBase& outDirName, const tchar*& outLastSlash, const StringView& fileName)
+	bool GetDirKey(StringKey& outDirKey, StringBufferBase& outDirName, const tchar*& outLastSlash, StringView fileName)
 	{
 		outLastSlash = TStrrchr(fileName.data, PathSeparator);
 		UBA_ASSERTF(outLastSlash, TC("Can't get dir key for path %s"), fileName.data);
@@ -153,7 +153,7 @@ namespace uba
 		UBA_ASSERT(!m_runningRemote);
 		StringBuffer<> fixedFilePath;
 		auto key = GetKeyAndFixedName(fixedFilePath, filePath);
-		return RegisterCreateFileForWrite(key, fixedFilePath, true);
+		return RegisterCreateFileForWrite(key, fixedFilePath, nullptr, 0, 0, true, true);
 	}
 
 	void DirectoryTableHolder::RegisterDeleteFile(const tchar* filePath)
@@ -169,7 +169,7 @@ namespace uba
 		UBA_ASSERT(!m_runningRemote);
 		StringBuffer<> fixedDirPath;
 		auto dirKey = GetKeyAndFixedName(fixedDirPath, directoryPath);
-		if (!RegisterCreateFileForWrite(dirKey, fixedDirPath, true))
+		if (!RegisterCreateFileForWrite(dirKey, fixedDirPath, nullptr, 0, 0, true, true))
 			return false;
 		WriteDirectoryEntries(dirKey, fixedDirPath);
 		return true;
@@ -192,7 +192,7 @@ namespace uba
 		return {m_directoryTable.m_memorySize, overlaySize};
 	}
 
-	bool DirectoryTableHolder::RegisterCreateFileForWrite(StringKey fileNameKey, StringView fileName, bool registerRealFile, u64 fileSize, u64 lastWriteTime, bool invalidateStorage)
+	bool DirectoryTableHolder::RegisterCreateFileForWrite(StringKey fileNameKey, StringView fileName, DirectoryTableOverlay* overlay, u64 fileSize, u64 lastWriteTime, bool registerRealFile, bool invalidateStorage)
 	{
 		// Remote is not updating its own directory table
 		if (m_runningRemote)
@@ -299,9 +299,9 @@ namespace uba
 
 		// Check if new write is actually a write. The file might just have been open with write permissions and then actually never written to.
 		// We check this by using lastWriteTime. If it hasn't change, directory table is already up-to-date
-		if (!insres.second && insres.first->second != ~u32(0))
+		if (!insres.second && IsValidEntry(insres.first->second))
 		{
-			BinaryReader reader(m_directoryTableMem + insres.first->second);
+			BinaryReader reader(m_directoryTableMem + insres.first->second.internal);
 			u32 oldAttr = reader.ReadFileAttributes();
 
 			if (IsDirectory(oldAttr))  // Ignore updating directories.. they should always be the same regardless
@@ -365,7 +365,7 @@ namespace uba
 		EnsureDirectoryTableMemory(writePos + 8 + written);
 		BinaryWriter writer(m_directoryTableMem + writePos);
 		writer.Write7BitEncoded(written); // Storage size
-		insres.first->second = dirTable.m_memorySize + u32(writer.GetPosition() + entryPos); // Storing position to last write time
+		insres.first->second = {dirTable.m_memorySize + u32(writer.GetPosition() + entryPos)}; // Storing position to last write time
 		u32 tableOffset = u32(writer.GetPosition()) + sizeof(StringKey);
 		writer.WriteBytes(temp, written);
 		dir.latestOffset = dirTable.m_memorySize + tableOffset;
@@ -373,7 +373,7 @@ namespace uba
 		return true;
 	}
 
-	DirTableSize DirectoryTableHolder::RegisterDeleteFile(StringKey fileNameKey, const StringView& fileName, DirectoryTableOverlay* overlay)
+	DirTableSize DirectoryTableHolder::RegisterDeleteFile(StringKey fileNameKey, StringView fileName, DirectoryTableOverlay* overlay)
 	{
 		// Remote is not updating its own directory table
 		if (m_runningRemote)
@@ -614,7 +614,7 @@ namespace uba
 					if (id == 0xffffffffffffffffllu) // When using projfs we might not have the file yet and in that case we need to make this up.
 						id = ++m_fileIndexCounter;
 
-					res.first->second = u32(memoryWriter.GetPosition()); // Temporary offset that will be used further down to calculate the real offset
+					res.first->second = {u32(memoryWriter.GetPosition())}; // Temporary offset that will be used further down to calculate the real offset
 					memoryWriter.WriteFileAttributes(e.attributes);
 					memoryWriter.WriteVolumeSerial(e.volumeSerial == volumeSerial ? volumeSerialIndex : m_volumeCache.GetSerialIndex(e.volumeSerial));
 					memoryWriter.WriteFileIndex(id);
@@ -703,7 +703,7 @@ namespace uba
 
 		// Update offsets to be relative to full memory
 		for (auto& kv : dir.files)
-			kv.second = filesOffset + kv.second;
+			kv.second = {filesOffset + kv.second.internal};
 
 
 		outTableOffset = tableOffset;
diff --git a/Engine/Source/Programs/UnrealBuildAccelerator/Common/Private/UbaProcess.cpp b/Engine/Source/Programs/UnrealBuildAccelerator/Common/Private/UbaProcess.cpp
--- a/Engine/Source/Programs/UnrealBuildAccelerator/Common/Private/UbaProcess.cpp
+++ b/Engine/Source/Programs/UnrealBuildAccelerator/Common/Private/UbaProcess.cpp
@@ -962,7 +962,7 @@ namespace uba
 
 		m_session.m_storage.InvalidateCachedFileInfo(toKey);
 
-		m_session.RegisterCreateFileForWrite(toKey, toName, true);
+		m_session.RegisterCreateFileForWrite(toKey, toName, nullptr, 0, 0, true, true);
 
 		//m_writtenFiles.try_emplace(name);
 		//WrittenFile& writtenFile = m_writtenFiles[toName.data];
diff --git a/Engine/Source/Programs/UnrealBuildAccelerator/Common/Private/UbaSession.cpp b/Engine/Source/Programs/UnrealBuildAccelerator/Common/Private/UbaSession.cpp
--- a/Engine/Source/Programs/UnrealBuildAccelerator/Common/Private/UbaSession.cpp
+++ b/Engine/Source/Programs/UnrealBuildAccelerator/Common/Private/UbaSession.cpp
@@ -1679,7 +1679,7 @@ namespace uba
 		auto fileNameKey = GetKeyAndFixedName(fixedFilePath, filePath);
 		if (!RegisterVirtualFileInternal(fileNameKey, fixedFilePath, sourceFile, sourceOffset, sourceSize))
 			return false;
-		return RegisterCreateFileForWrite(fileNameKey, fixedFilePath, false, sourceSize, 0, false);
+		return RegisterCreateFileForWrite(fileNameKey, fixedFilePath, nullptr, sourceSize, 0, false, false);
 	}
 
 	bool Session::CreateVirtualFile(const tchar* filePath, const void* memory, u64 memorySize, bool transient, bool registerToDir)
@@ -1691,7 +1691,7 @@ namespace uba
 			return false;
 		if (!registerToDir)
 			return true;
-		return RegisterCreateFileForWrite(fileNameKey, fixedFilePath, false, memorySize, 0, false);
+		return RegisterCreateFileForWrite(fileNameKey, fixedFilePath, nullptr, memorySize, 0, false, false);
 	}
 
 	bool Session::DeleteVirtualFile(const tchar* filePath, bool registerToDir)
@@ -2853,14 +2853,14 @@ namespace uba
 			if (!msg.newName.IsEmpty())
 			{
 				RegisterDeleteFile(activeFile.nameKey, activeFile.name, &shared.overlay);
-				if (RegisterCreateFileForWrite(msg.newNameKey, msg.newName, registerRealFile, fileSize, lastWriteTime) && registerRealFile)
+				if (RegisterCreateFileForWrite(msg.newNameKey, msg.newName, nullptr, fileSize, lastWriteTime, registerRealFile, true) && registerRealFile)
 					TraceWrittenFile(msg.process.m_id, msg.newName, fileSize);
 			}
 			else if (msg.deleteOnClose)
 				RegisterDeleteFile(activeFile.nameKey, activeFile.name, &shared.overlay);
 			else
 			{
-				if (RegisterCreateFileForWrite(activeFile.nameKey, activeFile.name, registerRealFile, fileSize, lastWriteTime) && registerRealFile)
+				if (RegisterCreateFileForWrite(activeFile.nameKey, activeFile.name, nullptr, fileSize, lastWriteTime, registerRealFile, true) && registerRealFile)
 					TraceWrittenFile(msg.process.m_id, activeFile.name, fileSize);
 			}
 		}
@@ -2979,7 +2979,7 @@ namespace uba
 		}
 
 		if (!process.GetStartInfo().rules->IsInvisible(msg.toName))
-			if (RegisterCreateFileForWrite(msg.toKey, msg.toName, true))
+			if (RegisterCreateFileForWrite(msg.toKey, msg.toName, nullptr, 0, 0, true, true))
 				TraceWrittenFile(process.m_id, msg.toName, 0);
 
 		out.errorCode = ERROR_SUCCESS;
@@ -2995,7 +2995,7 @@ namespace uba
 		out.errorCode = 0;
 		if (chmod(msg.fileName.data, (mode_t)msg.fileMode) == 0)
 		{
-			RegisterCreateFileForWrite(msg.fileNameKey, msg.fileName, true);
+			RegisterCreateFileForWrite(msg.fileNameKey, msg.fileName, nullptr, 0, 0, true, true);
 			return true;
 		}
 		out.errorCode = errno;
@@ -3021,7 +3021,7 @@ namespace uba
 		if (out.result || out.errorCode == ERROR_ALREADY_EXISTS)
 		{
 			// Both these functions need to be called. otherwise we can get created directories that does not end up in directory table
-			RegisterCreateFileForWrite(msg.nameKey, msg.name, true);
+			RegisterCreateFileForWrite(msg.nameKey, msg.name, nullptr, 0, 0, true, true);
 			WriteDirectoryEntries(dirKey, dirName);
 		}
 
@@ -3556,7 +3556,7 @@ namespace uba
 			m_sharedMemory.UnmapView(symHandle, mem2);
 
 			lastWriteTime = GetSystemTimeAsFileTime();
-			if (!RegisterCreateFileForWrite(symFileKey, exportsFile, false, memoryBlock.writtenSize, lastWriteTime))
+			if (!RegisterCreateFileForWrite(symFileKey, exportsFile, nullptr, memoryBlock.writtenSize, lastWriteTime, false, true))
 				return false;
 			mg.Cancel();
 		}
@@ -3569,7 +3569,7 @@ namespace uba
 				return false;
 			if (!destFile.Close(&lastWriteTime))
 				return false;
-			if (!RegisterCreateFileForWrite(symFileKey, exportsFile, true, memoryBlock.writtenSize, lastWriteTime))
+			if (!RegisterCreateFileForWrite(symFileKey, exportsFile, nullptr, memoryBlock.writtenSize, lastWriteTime, true, true))
 				return false;
 		}
 
diff --git a/Engine/Source/Programs/UnrealBuildAccelerator/Common/Private/UbaSessionClient.cpp b/Engine/Source/Programs/UnrealBuildAccelerator/Common/Private/UbaSessionClient.cpp
--- a/Engine/Source/Programs/UnrealBuildAccelerator/Common/Private/UbaSessionClient.cpp
+++ b/Engine/Source/Programs/UnrealBuildAccelerator/Common/Private/UbaSessionClient.cpp
@@ -132,46 +132,12 @@ namespace uba
 				},
 			[this](const StringView& path, const DependencyCrawler::FileFunc& fileFunc) // TraverseFilesFunc
 				{
-					u32 tableOffset = 0;
+					DirTableOffset tableOffset;
 					if (!EntryExists(path, tableOffset))
 						return;
-					if ((tableOffset & 0x80000000) == 0)
+					if (!IsDirectoryEntry(tableOffset))
 						return;
-					tableOffset &= ~0x80000000;
-
-					// This is not entirely correct since files could be added while uba is running.. but since this is only used for includes that
-					// we know are created before uba started we're good.. and there should be no new files added to those dirs causing table offset to be chained
-					BinaryReader reader(m_directoryTableMem, tableOffset, m_directoryTableMemPos);
-					while (true)
-					{
-						u32 prevTableOffset = u32(reader.Read7BitEncoded());
-						if (!prevTableOffset)
-							break;
-						reader.SetPosition(prevTableOffset);
-					}
-
-					u32 dirAttr = reader.ReadFileAttributes();
-					if (!dirAttr)
-						return;
-					reader.ReadVolumeSerial();
-					reader.ReadFileIndex();
-					u64 itemCount = reader.Read7BitEncoded();
-					while (itemCount--)
-					{
-						StringBuffer<> fileName;
-						reader.ReadString(fileName);
-						if (CaseInsensitiveFs)
-							fileName.MakeLower();
-						u32 attr = reader.ReadFileAttributes();
-						bool isDir = IsDirectory(attr);
-						fileFunc(fileName, isDir);
-						reader.ReadVolumeSerial();
-						reader.ReadFileIndex();
-						if (isDir)
-							continue;
-						reader.ReadFileTime();
-						reader.ReadFileSize();
-					}
+					m_directoryTable.TraverseOriginalFilesNoPopulate(tableOffset, fileFunc);
 				});
 	}
 
@@ -246,7 +212,7 @@ namespace uba
 
 	bool SessionClient::Exists(const StringView& path, u32& outAttributes)
 	{
-		u32 tableOffset = 0;
+		DirTableOffset tableOffset;
 		if (!EntryExists(path, tableOffset))
 			return false;
 		outAttributes = m_directoryTable.GetAttributes(tableOffset);
@@ -2437,7 +2403,7 @@ namespace uba
 		return true;
 	}
 
-	bool SessionClient::EntryExists(const StringView& path, u32& outTableOffset)
+	bool SessionClient::EntryExists(StringView path, DirTableOffset& outTableOffset)
 	{
 		StringKey key;
 		if (path.data[path.count - 1] == PathSeparator)
@@ -2446,7 +2412,7 @@ namespace uba
 			key = ToStringKey(path);
 
 		StringBuffer<> dirName;
-		u32 tableOffset = 0;
+		DirTableOffset tableOffset;
 		DirectoryTable::Exists exists = m_directoryTable.EntryExists(key, path, true, &tableOffset);
 		if (exists == DirectoryTable::Exists_Maybe)
 		{
diff --git a/Engine/Source/Programs/UnrealBuildAccelerator/Common/Private/UbaSessionServer.cpp b/Engine/Source/Programs/UnrealBuildAccelerator/Common/Private/UbaSessionServer.cpp
--- a/Engine/Source/Programs/UnrealBuildAccelerator/Common/Private/UbaSessionServer.cpp
+++ b/Engine/Source/Programs/UnrealBuildAccelerator/Common/Private/UbaSessionServer.cpp
@@ -1497,7 +1497,7 @@ namespace uba
 		if (!rules.IsInvisible(destination))
 		{
 			bool invalidateStorage = false; // No need, already handled in m_storage.CopyOrLink
-			RegisterCreateFileForWrite(StringKeyZero, destination, shouldWriteToDisk, 0, 0, invalidateStorage);
+			RegisterCreateFileForWrite(StringKeyZero, destination, nullptr, 0, 0, shouldWriteToDisk, invalidateStorage);
 
 			if (startInfo.trackInputs)
 			{
@@ -1615,7 +1615,7 @@ namespace uba
 		bool result = uba::CopyFileW(fromName.data, toName.data, false);
 		u32 errorCode = GetLastError();
 		if (result)
-			RegisterCreateFileForWrite(toNameKey, toName, true);
+			RegisterCreateFileForWrite(toNameKey, toName, nullptr, 0, 0, true, true);
 		writer.WriteU32(errorCode);
 		return true;
 	}
@@ -2863,7 +2863,7 @@ namespace uba
 				{
 					StringKey fileNameKey = fileKv.first;
 
-					BinaryReader reader(dirTable.m_memory, fileKv.second);
+					BinaryReader reader(dirTable.m_memory, fileKv.second.internal);
 
 					u64 lastWritten = reader.ReadU64();
 					u32 attr = reader.ReadU32();
diff --git a/Engine/Source/Programs/UnrealBuildAccelerator/Common/Public/UbaDirectoryTableHolder.h b/Engine/Source/Programs/UnrealBuildAccelerator/Common/Public/UbaDirectoryTableHolder.h
--- a/Engine/Source/Programs/UnrealBuildAccelerator/Common/Public/UbaDirectoryTableHolder.h
+++ b/Engine/Source/Programs/UnrealBuildAccelerator/Common/Public/UbaDirectoryTableHolder.h
@@ -41,8 +41,8 @@ namespace uba
 
 	protected:
 
-		bool RegisterCreateFileForWrite(StringKey fileNameKey, StringView fileName, bool registerRealFile, u64 fileSize = 0, u64 lastWriteTime = 0, bool invalidateStorage = true);
-		DirTableSize RegisterDeleteFile(StringKey fileNameKey, const StringView& fileName, DirectoryTableOverlay* overlay);
+		bool RegisterCreateFileForWrite(StringKey fileNameKey, StringView fileName, DirectoryTableOverlay* overlay, u64 fileSize, u64 lastWriteTime, bool registerRealFile, bool invalidateStorage);
+		DirTableSize RegisterDeleteFile(StringKey fileNameKey, StringView fileName, DirectoryTableOverlay* overlay);
 
 		u32 WriteDirectoryEntries(const StringKey& dirKey, StringView dirPath, u32* outTableOffset = nullptr);
 		void WriteDirectoryEntriesRecursive(const StringKey& dirKey, StringView dirPath, u32& outTableOffset);
@@ -55,7 +55,7 @@ namespace uba
 
 		bool m_runningRemote;
 
-		Atomic<u64> m_fileIndexCounter = 8000000000;
+		Atomic<u64> m_fileIndexCounter = 8000000000; // Try not to collide with ntfs
 		MemoryBlock m_directoryTableMemory;
 		FileMappingHandle m_directoryTableHandle;
 		u8* m_directoryTableMem;
@@ -67,7 +67,7 @@ namespace uba
 
 	////////////////////////////////////////////////////////////////////////////////////////////////////
 
-	bool GetDirKey(StringKey& outDirKey, StringBufferBase& outDirName, const tchar*& outLastSlash, const StringView& fileName);
+	bool GetDirKey(StringKey& outDirKey, StringBufferBase& outDirName, const tchar*& outLastSlash, StringView fileName);
 	StringKey GetKeyAndFixedName(StringBuffer<>& outFixedFilePath, StringKeyHasher& outHasher, const tchar* filePath, StringKey* outDirKey = nullptr, const tchar** outBaseFileName = nullptr);
 	StringKey GetKeyAndFixedName(StringBuffer<>& outFixedFilePath, const tchar* filePath, StringKey* outDirKey = nullptr, const tchar** outBaseFileName = nullptr);
 
diff --git a/Engine/Source/Programs/UnrealBuildAccelerator/Common/Public/UbaSessionClient.h b/Engine/Source/Programs/UnrealBuildAccelerator/Common/Public/UbaSessionClient.h
--- a/Engine/Source/Programs/UnrealBuildAccelerator/Common/Public/UbaSessionClient.h
+++ b/Engine/Source/Programs/UnrealBuildAccelerator/Common/Public/UbaSessionClient.h
@@ -114,7 +114,7 @@ namespace uba
 		void GetLogFileName(StringBufferBase& out, const tchar* logFile, const tchar* arguments, u32 processId);
 		u32 WriteLogLines(BinaryWriter& writer, ProcessImpl& process);
 		bool ParseDirectoryTable();
-		bool EntryExists(const StringView& path, u32& outTableOffset);
+		bool EntryExists(StringView path, DirTableOffset& outTableOffset);
 		void RunSpeedTest(StringView args);
 
 		void ThreadCreateProcessLoop();
diff --git a/Engine/Source/Programs/UnrealBuildAccelerator/Core/Public/UbaDirectoryTable.h b/Engine/Source/Programs/UnrealBuildAccelerator/Core/Public/UbaDirectoryTable.h
--- a/Engine/Source/Programs/UnrealBuildAccelerator/Core/Public/UbaDirectoryTable.h
+++ b/Engine/Source/Programs/UnrealBuildAccelerator/Core/Public/UbaDirectoryTable.h
@@ -36,17 +36,22 @@ namespace uba
 
 	enum : u64 { InvalidTableOffset = 0 }; // Important that it is zero to compress enc7 as good as possible
 	enum : u64 { OverlayTableOffset = 256*1024*1024 };
+	enum : u32 { DirectoryEntryFlag = 0x80000000 };
 
 	struct DirTableSize { u32 main = 0; u32 overlay = 0; };
 	inline DirTableSize FromU64(u64 v) { return {u32(v), u32(v >> 32u)}; }
 	inline u64 ToU64(DirTableSize size) { return size.main + (u64(size.overlay) << 32llu); }
 
+	struct DirTableOffset { u32 internal = InvalidTableOffset; };
+	inline bool IsDirectoryEntry(DirTableOffset d) { return (d.internal & DirectoryEntryFlag) != 0; }
+	inline bool IsValidEntry(DirTableOffset d) { return d.internal != InvalidTableOffset; }
+
 	////////////////////////////////////////////////////////////////////////////////////////////////////
 
 	class DirectoryTable
 	{
 	public:
-		using EntryLookup = GrowingUnorderedMap<StringKey, u32>;
+		using EntryLookup = GrowingUnorderedMap<StringKey, DirTableOffset>;
 
 		struct Directory
 		{
@@ -78,8 +83,11 @@ namespace uba
 			Exists_Maybe,
 		};
 
-		inline Exists EntryExists(StringKey entryKey, StringView entryName, bool checkIfDir = false, u32* outTableOffset = nullptr);
-		inline Exists EntryExistsNoLock(StringKey entryKey, StringView entryName, bool checkIfDir = false, u32* outTableOffset = nullptr);
+		// When checkIfDir is true these functions will search for the offset where a directory stores its files.
+		// It is only useful when the intention is to traverse files inside the dir.
+		// NOTE, even if checkIfDir is true you can still get back a normal "file entry" representing the dir. "IsDirectoryEntry(tableOffset)" can be used to check which one it is
+		inline Exists EntryExists(StringKey entryKey, StringView entryName, bool checkIfDir = false, DirTableOffset* outTableOffset = nullptr);
+		inline Exists EntryExistsNoLock(StringKey entryKey, StringView entryName, bool checkIfDir = false, DirTableOffset* outTableOffset = nullptr);
 		inline Exists EntryExists(StringView str, bool checkIfDir = false);
 
 		struct EntryInformation
@@ -91,9 +99,14 @@ namespace uba
 			u64 lastWrite;
 		};
 
-		inline u32 GetAttributes(u32 tableOffset);
-		inline u32 GetEntryInformation(EntryInformation& outInfo, u32 tableOffset, tchar* outFileName = nullptr, u32 fileNameCapacity = 0);
+		inline u32 GetAttributes(DirTableOffset tableOffset);
+		inline u32 GetEntryInformation(EntryInformation& outInfo, DirTableOffset tableOffset, tchar* outFileName = nullptr, u32 fileNameCapacity = 0);
 		inline void GetFinalPath(StringBufferBase& out, const tchar* path);
+		inline bool GetLatestOffset(DirTableOffset& out, const Directory& dir);
+
+		// This function only traverses the entries that was in the original populate and does not populate the lookup table.
+		// Can be used for scenarios where we know we only want the files that existed before the session started and not files that were added while active
+		template<typename Func> inline void TraverseOriginalFilesNoPopulate(DirTableOffset tableOffset, const Func& func);
 
 		#if PLATFORM_WINDOWS
 		template<typename Func> inline void TraverseFilesRecursiveNoLock(const StringBufferBase& path, const Func& func);
@@ -163,7 +176,7 @@ namespace uba
 
 			for (auto j=dir.files.begin(); j!=dir.files.end();)
 			{
-				if (j->second < OverlayTableOffset)
+				if (j->second.internal < OverlayTableOffset)
 					++j;
 				else
 					j = dir.files.erase(j);
@@ -308,17 +321,17 @@ namespace uba
 			}
 
 			StringKey filenameKey = ToStringKey(hasher, filename.data, filename.count);
-			files[filenameKey] = offset; // Always write, since same file might have been added with new info
+			files[filenameKey] = {offset}; // Always write, since same file might have been added with new info
 		}
 	}
 
-	DirectoryTable::Exists DirectoryTable::EntryExists(StringKey entryKey, StringView entryName, bool checkIfDir, u32* outTableOffset)
+	DirectoryTable::Exists DirectoryTable::EntryExists(StringKey entryKey, StringView entryName, bool checkIfDir, DirTableOffset* outTableOffset)
 	{
 		SCOPED_READ_LOCK(m_lookupLock, lock);
 		return EntryExistsNoLock(entryKey, entryName, checkIfDir, outTableOffset);
 	}
 
-	DirectoryTable::Exists DirectoryTable::EntryExistsNoLock(StringKey entryKey, StringView entryName, bool checkIfDir, u32* outTableOffset)
+	DirectoryTable::Exists DirectoryTable::EntryExistsNoLock(StringKey entryKey, StringView entryName, bool checkIfDir, DirTableOffset* outTableOffset)
 	{
 		u32 startSkip = 2;
 		if (checkIfDir)
@@ -327,7 +340,7 @@ namespace uba
 			if (findIt != m_lookup.end())
 			{
 				if (outTableOffset)
-					*outTableOffset = u32(findIt->second.latestOffset) | 0x80000000; // Use significant bit to say that this is a dir
+					outTableOffset->internal = u32(findIt->second.latestOffset) | DirectoryEntryFlag; // Use significant bit to say that this is a dir
 				return Exists_Yes;
 			}
 			startSkip = 1;
@@ -377,14 +390,14 @@ namespace uba
 					return Exists_No;
 				if (inAncestor)
 				{
-					BinaryReader reader(m_memory, entryIt->second, m_memorySize);
+					BinaryReader reader(m_memory, entryIt->second.internal, m_memorySize);
 					reader.SkipString();
 					if (!IsDirectory(reader.ReadFileAttributes()))
 						return Exists_No;
 					return Exists_Maybe;
 				}
 				if (outTableOffset)
-					*outTableOffset = entryIt->second;
+					outTableOffset->internal = entryIt->second.internal;
 				return Exists_Yes;
 			}
 
@@ -405,12 +418,11 @@ namespace uba
 		return EntryExists(ToStringKey(str2), str2, checkIfDir, nullptr);
 	}
 
-	u32 DirectoryTable::GetAttributes(u32 tableOffset)
+	u32 DirectoryTable::GetAttributes(DirTableOffset tableOffset)
 	{
-		if (tableOffset & 0x80000000)
+		if (IsDirectoryEntry(tableOffset))
 		{
-			tableOffset = tableOffset & ~0x80000000;
-			BinaryReader reader(m_memory, tableOffset);
+			BinaryReader reader(m_memory, tableOffset.internal & ~DirectoryEntryFlag);
 			u64 prevTableOffset = reader.Read7BitEncoded();
 			while (prevTableOffset != InvalidTableOffset)
 			{
@@ -419,17 +431,16 @@ namespace uba
 			}
 			return reader.ReadFileAttributes();
 		}
-		BinaryReader reader(m_memory, tableOffset);
+		BinaryReader reader(m_memory, tableOffset.internal);
 		reader.SkipString();
 		return reader.ReadFileAttributes();
 	}
 
-	u32 DirectoryTable::GetEntryInformation(EntryInformation& outInfo, u32 tableOffset, tchar* outFileName, u32 fileNameCapacity)
+	u32 DirectoryTable::GetEntryInformation(EntryInformation& outInfo, DirTableOffset tableOffset, tchar* outFileName, u32 fileNameCapacity)
 	{
-		if (tableOffset & 0x80000000)
+		if (IsDirectoryEntry(tableOffset))
 		{
-			tableOffset = tableOffset & ~0x80000000;
-			BinaryReader reader(GetReader(tableOffset));
+			BinaryReader reader(GetReader(tableOffset.internal & ~DirectoryEntryFlag));
 			u64 prevTableOffset = reader.Read7BitEncoded();
 			while (prevTableOffset != InvalidTableOffset)
 			{
@@ -448,7 +459,7 @@ namespace uba
 			return ~u32(0);
 		}
 
-		BinaryReader reader(GetReader(tableOffset));
+		BinaryReader reader(GetReader(tableOffset.internal));
 		if (outFileName)
 			reader.ReadString(outFileName, fileNameCapacity);
 		else
@@ -511,8 +522,8 @@ namespace uba
 				auto fileIt = directory->files.find(fileNameKey);
 				if (fileIt != directory->files.end())
 				{
-					UBA_ASSERT(fileIt->second != ~0u);
-					BinaryReader reader(GetReader(fileIt->second));
+					UBA_ASSERT(IsValidEntry(fileIt->second));
+					BinaryReader reader(GetReader(fileIt->second.internal));
 					StringBuffer<> fileName;
 					reader.ReadString(fileName);
 					out.Append(PathSeparator).Append(fileName);
@@ -539,6 +550,54 @@ namespace uba
 		}
 	}
 
+	bool DirectoryTable::GetLatestOffset(DirTableOffset& out, const Directory& dir)
+	{
+		if (dir.latestOffset == InvalidFileHandle)
+			return false;
+		out.internal = dir.latestOffset | DirectoryEntryFlag;
+		return true;
+	}
+
+	template<typename Func>
+	void DirectoryTable::TraverseOriginalFilesNoPopulate(DirTableOffset tableOffset, const Func& func)
+	{
+		UBA_ASSERT(IsDirectoryEntry(tableOffset));
+		tableOffset.internal &= ~DirectoryEntryFlag;
+
+		// Go to first offset
+		BinaryReader reader(m_memory, tableOffset.internal, m_memorySize);
+		while (true)
+		{
+			u32 prevTableOffset = u32(reader.Read7BitEncoded());
+			if (!prevTableOffset)
+				break;
+			reader.SetPosition(prevTableOffset);
+		}
+
+		u32 dirAttr = reader.ReadFileAttributes();
+		if (!dirAttr)
+			return;
+		reader.ReadVolumeSerial();
+		reader.ReadFileIndex();
+		u64 itemCount = reader.Read7BitEncoded();
+		while (itemCount--)
+		{
+			StringBuffer<> fileName;
+			reader.ReadString(fileName);
+			if (CaseInsensitiveFs)
+				fileName.MakeLower();
+			u32 attr = reader.ReadFileAttributes();
+			bool isDir = IsDirectory(attr);
+			func(fileName, isDir);
+			reader.ReadVolumeSerial();
+			reader.ReadFileIndex();
+			if (isDir)
+				continue;
+			reader.ReadFileTime();
+			reader.ReadFileSize();
+		}
+	}
+
 #if PLATFORM_WINDOWS
 	template<typename Func>
 	void DirectoryTable::TraverseFilesRecursiveNoLock(const StringBufferBase& path, const Func& func)
diff --git a/Engine/Source/Programs/UnrealBuildAccelerator/Detours/Private/PosixOS/UbaDetoursFunctionsPosix.cpp b/Engine/Source/Programs/UnrealBuildAccelerator/Detours/Private/PosixOS/UbaDetoursFunctionsPosix.cpp
--- a/Engine/Source/Programs/UnrealBuildAccelerator/Detours/Private/PosixOS/UbaDetoursFunctionsPosix.cpp
+++ b/Engine/Source/Programs/UnrealBuildAccelerator/Detours/Private/PosixOS/UbaDetoursFunctionsPosix.cpp
@@ -379,15 +379,14 @@ int Shared_open(const char* funcName, const char* file, int flags, int mode, con
 
 	u64 size = InvalidValue;
 	u32 closeId = 0;
-	u32 dirTableOffset = ~u32(0);
+	DirTableOffset dirTableOffset;
 
 	if (g_allowDirectoryCache)
 	{
 		if (!isWrite) // We need to skip SystemTemp.. lots of stuff going on there.
 		{
-			dirTableOffset = Rpc_GetEntryOffset(fileNameKey, fileName, false);
 			bool allowEarlyOut = true;
-			if (dirTableOffset == ~u32(0))
+			if (!Rpc_GetEntryOffset(dirTableOffset, fileNameKey, fileName, false))
 			{
 				// This could be a written file not reported to server yet
 				{
@@ -1420,7 +1419,7 @@ UBA_EXPORT int UBA_WRAPPER(fclose)(FILE* stream)
 
 struct DirInfo
 {
-	Vector<u32> fileTableOffsets;
+	Vector<DirTableOffset> fileTableOffsets;
 	int it = -1;
 	dirent ent;
 };
@@ -1457,9 +1456,9 @@ UBA_EXPORT DIR* UBA_WRAPPER(opendir)(const char* name)
 			Rpc_UpdateDirectory(hash.key, dirName.data, dirName.count, false);
 
 	bool exists = false;
-	if (dir.tableOffset != InvalidTableOffset)
+	DirTableOffset entryOffset;
+	if (g_directoryTable.GetLatestOffset(entryOffset, dir))
 	{
-		u32 entryOffset = dir.tableOffset | 0x80000000;
 		DirectoryTable::EntryInformation entryInfo;
 		g_directoryTable.GetEntryInformation(entryInfo, entryOffset);
 		exists = entryInfo.attributes != 0;
@@ -1514,7 +1513,7 @@ dirent* Shared_readdir(const char* func, DIR* dirp)
 			DEBUG_LOG_DETOURED(func, "(%p) -> nullptr (found %u entries)", dirp, dirInfo.fileTableOffsets.size());
 			return nullptr;
 		}
-		u32 fileTableOffset = dirInfo.fileTableOffsets[dirInfo.it];
+		DirTableOffset fileTableOffset = dirInfo.fileTableOffsets[dirInfo.it];
 
 		DirectoryTable::EntryInformation info;
 		g_directoryTable.GetEntryInformation(info, fileTableOffset, dirInfo.ent.d_name, 256);
diff --git a/Engine/Source/Programs/UnrealBuildAccelerator/Detours/Private/UbaDetoursFileMappingTable.cpp b/Engine/Source/Programs/UnrealBuildAccelerator/Detours/Private/UbaDetoursFileMappingTable.cpp
--- a/Engine/Source/Programs/UnrealBuildAccelerator/Detours/Private/UbaDetoursFileMappingTable.cpp
+++ b/Engine/Source/Programs/UnrealBuildAccelerator/Detours/Private/UbaDetoursFileMappingTable.cpp
@@ -333,9 +333,8 @@ namespace uba
 		view.AddRequestedMemory(allocatorHandle, reader);
 	}
 
-	u32 Rpc_GetEntryOffset(const StringKey& entryNameKey, StringView entryName, bool checkIfDir)
+	bool Rpc_GetEntryOffset(DirTableOffset& out, const StringKey& entryNameKey, StringView entryName, bool checkIfDir)
 	{
-		u32 dirTableOffset = ~u32(0);
 		StringBuffer<MaxPath> entryNameForKey;
 		entryNameForKey.Append(entryName);
 		if (CaseInsensitiveFs)
@@ -344,15 +343,17 @@ namespace uba
 			checkIfDir = true;
 
 		CHECK_PATH(entryNameForKey);
-		DirectoryTable::Exists exists = g_directoryTable.EntryExists(entryNameKey, entryNameForKey, checkIfDir, &dirTableOffset);
-		if (exists != DirectoryTable::Exists_Maybe)
-			return dirTableOffset;
+		DirectoryTable::Exists exists = g_directoryTable.EntryExists(entryNameKey, entryNameForKey, checkIfDir, &out);
+		if (exists == DirectoryTable::Exists_No)
+			return false;
+		if (exists == DirectoryTable::Exists_Yes)
+			return true;
 
 		const tchar* lastPathSeparator = TStrrchr(entryName.data, PathSeparator);
 		if (!lastPathSeparator)
 		{
 			UBA_ASSERTF(lastPathSeparator, TC("No path separator found in %s"), entryName.count > 0 ? entryName.data : TC("(NULL)"));
-			return ~u32(0);
+			return false;
 		}
 
 		#if PLATFORM_WINDOWS
@@ -363,24 +364,27 @@ namespace uba
 		DirHash hash(StringView(entryNameForKey.data, dirNameLen));
 
 		if (Rpc_UpdateDirectory(hash.key, entryName.data, dirNameLen) == ~u32(0))
-			return ~u32(0);
+			return false;
 
 		SCOPED_WRITE_LOCK(g_directoryTable.m_lookupLock, lookLock);
 		auto dirFindIt = g_directoryTable.m_lookup.find(hash.key);
 		if (dirFindIt == g_directoryTable.m_lookup.end())
-			return ~u32(0);
+			return false;
 		auto& dir = dirFindIt->second;
 
 		if (checkIfDir)
-			return u32(dir.latestOffset) | 0x80000000; // Use significant bit to say that this is a dir
-
+		{
+			g_directoryTable.GetLatestOffset(out, dir);
+			return true;
+		}
 		g_directoryTable.PopulateDirectory(hash.open, dir);
 
 		SCOPED_READ_LOCK(dir.lock, lock);
 		auto findIt = dir.files.find(entryNameKey);
 		if (findIt == dir.files.end())
-			return ~u32(0);
-		return findIt->second;
+			return false;
+		out = findIt->second;
+		return true;
 	}
 
 	void Rpc_GetFullFileName(const tchar*& path, u64& pathLen, StringBufferBase& tempBuf, bool useVirtualName, const tchar* const* loaderPaths)
diff --git a/Engine/Source/Programs/UnrealBuildAccelerator/Detours/Private/UbaDetoursFileMappingTable.h b/Engine/Source/Programs/UnrealBuildAccelerator/Detours/Private/UbaDetoursFileMappingTable.h
--- a/Engine/Source/Programs/UnrealBuildAccelerator/Detours/Private/UbaDetoursFileMappingTable.h
+++ b/Engine/Source/Programs/UnrealBuildAccelerator/Detours/Private/UbaDetoursFileMappingTable.h
@@ -5,6 +5,7 @@
 #include "UbaBinaryReaderWriter.h"
 #include "UbaDefinitions.h"
 #include "UbaDetoursShared.h"
+#include "UbaDirectoryTable.h"
 #include "UbaSharedMemoryView.h"
 
 
@@ -97,7 +98,7 @@ namespace uba
 	u32  Rpc_UpdateDirectory(const StringKey& dirKey, const tchar* dirName, u64 dirNameLen, bool lockDirTable = true);
 	void Rpc_UpdateCloseHandle(const tchar* handleName, u32 closeId, bool deleteOnClose, const tchar* newName, SharedMemoryHandle memoryHandle, u64 memoryWritten, bool success);
 	void Rpc_UpdateTables();
-	u32  Rpc_GetEntryOffset(const StringKey& entryNameKey, StringView entryName, bool checkIfDir = false);
+	bool Rpc_GetEntryOffset(DirTableOffset& out, const StringKey& entryNameKey, StringView entryName, bool checkIfDir = false);
 	void Rpc_GetFullFileName(const tchar*& path, u64& pathLen, StringBufferBase& tempBuf, bool useVirtualName, const tchar* const* loaderPaths = nullptr);
 	void Rpc_GetFullFileName2(const tchar* path, StringBufferBase& outReal, StringBufferBase& outVirtual, const tchar* const* loaderPaths = nullptr);
 	void Rpc_GetWrittenFiles();
diff --git a/Engine/Source/Programs/UnrealBuildAccelerator/Detours/Private/UbaDetoursShared.cpp b/Engine/Source/Programs/UnrealBuildAccelerator/Detours/Private/UbaDetoursShared.cpp
--- a/Engine/Source/Programs/UnrealBuildAccelerator/Detours/Private/UbaDetoursShared.cpp
+++ b/Engine/Source/Programs/UnrealBuildAccelerator/Detours/Private/UbaDetoursShared.cpp
@@ -618,9 +618,9 @@ namespace uba
 		// This is an optimization where we populate directory table and use that to figure out if file exists or not..
 		// .. in msvc's case it doesn't matter much because these tables are already up to date when msvc use CreateFile.
 		// .. clang otoh is using CreateFile with tooons of different paths trying to open files.. in remote worker case this becomes super expensive
-		u32 dirTableOffset = Rpc_GetEntryOffset(fileNameKey, fileName, checkIfDir);
+		DirTableOffset dirTableOffset;
 
-		if (dirTableOffset == ~u32(0))
+		if (!Rpc_GetEntryOffset(dirTableOffset, fileNameKey, fileName, checkIfDir))
 		{
 			if (foundMapping)
 				return fileName.data;
diff --git a/Engine/Source/Programs/UnrealBuildAccelerator/Detours/Private/Windows/UbaDetoursFunctionsKernelBase.inl b/Engine/Source/Programs/UnrealBuildAccelerator/Detours/Private/Windows/UbaDetoursFunctionsKernelBase.inl
--- a/Engine/Source/Programs/UnrealBuildAccelerator/Detours/Private/Windows/UbaDetoursFunctionsKernelBase.inl
+++ b/Engine/Source/Programs/UnrealBuildAccelerator/Detours/Private/Windows/UbaDetoursFunctionsKernelBase.inl
@@ -272,8 +272,7 @@ BOOL Detoured_CreateDirectoryW(LPCWSTR lpPathName, LPSECURITY_ATTRIBUTES lpSecur
 	StringKey pathNameKey = ToStringKey(lowerName);
 
 	bool checkIfDir = true;
-	u32 dirTableOffset = ~u32(0);
-	DirectoryTable::Exists exists = g_directoryTable.EntryExists(pathNameKey, lowerName, checkIfDir, &dirTableOffset);
+	DirectoryTable::Exists exists = g_directoryTable.EntryExists(pathNameKey, lowerName, checkIfDir, nullptr);
 	if (exists == DirectoryTable::Exists_Yes)
 	{
 		SetLastError(ERROR_ALREADY_EXISTS);
@@ -491,26 +490,24 @@ BOOL Detoured_GetVolumeInformationByHandleW(HANDLE hFile, LPWSTR lpVolumeNameBuf
 	DETOURED_CALL(GetVolumeInformationByHandleW);
 	HANDLE trueHandle = hFile;
 
-	u32 entryOffset = ~u32(0);
+	DirTableOffset entryOffset;
 
 	if (isDetouredHandle(hFile))
 	{
 		DetouredHandle& dh = asDetouredHandle(hFile);
 		trueHandle = dh.trueHandle;
 		entryOffset = dh.dirTableOffset;
-		UBA_ASSERT(entryOffset != ~u32(0) || trueHandle != INVALID_HANDLE_VALUE);
+		UBA_ASSERT(IsValidEntry(entryOffset) || trueHandle != INVALID_HANDLE_VALUE);
 	}
 	else if (isListDirectoryHandle(hFile))
 	{
 		auto& listHandle = asListDirectoryHandle(hFile);
-		if (listHandle.dir.latestOffset != InvalidTableOffset)
-			entryOffset = listHandle.dir.latestOffset | 0x80000000;
-		else
+		if (!g_directoryTable.GetLatestOffset(entryOffset, listHandle.dir))
 			UBA_ASSERT(false);
 		trueHandle = INVALID_HANDLE_VALUE;
 	}
 
-	if (entryOffset != ~u32(0))
+	if (IsValidEntry(entryOffset))
 	{
 		UBA_ASSERT(!lpVolumeNameBuffer);
 		UBA_ASSERT(!lpMaximumComponentLength);
@@ -827,8 +824,8 @@ BOOL Detoured_SetFileAttributesW(LPCWSTR lpFileName, DWORD dwFileAttributes)
 	if (g_runningRemote && !fixedFile.StartsWith(g_systemTemp))
 	{
 		StringKey fileNameKey = ToStringKeyLower(fixedFile);
-		u32 dirTableOffset = Rpc_GetEntryOffset(fileNameKey, fixedFile, true);
-		if (dirTableOffset != ~u32(0))
+		DirTableOffset dirTableOffset;
+		if (Rpc_GetEntryOffset(dirTableOffset, fileNameKey, fixedFile, true))
 		{
 			DirectoryTable::EntryInformation info;
 			g_directoryTable.GetEntryInformation(info, dirTableOffset);
@@ -1493,7 +1490,7 @@ bool Shared_GetNextFile(WIN32_FIND_DATA& outData, ListDirectoryHandle& listHandl
 		}
 
 		DirectoryTable::EntryInformation info;
-		u32 fileTableOffset = listHandle.fileTableOffsets[listHandle.it++];
+		DirTableOffset fileTableOffset = listHandle.fileTableOffsets[listHandle.it++];
 		g_directoryTable.GetEntryInformation(info, fileTableOffset, outData.cFileName, maxLen);
 		if (info.attributes == 0) // File was deleted
 			continue;
@@ -1575,7 +1572,8 @@ __forceinline HANDLE Shared_FindFirstFileExW(LPCWSTR lpFileName, FINDEX_INFO_LEV
 	bool exists = false;
 	if (dir.latestOffset != InvalidTableOffset)
 	{
-		u32 entryOffset = dir.latestOffset | 0x80000000;
+		DirTableOffset entryOffset;
+		g_directoryTable.GetLatestOffset(entryOffset, dir);
 		DirectoryTable::EntryInformation entryInfo;
 		g_directoryTable.GetEntryInformation(entryInfo, entryOffset);
 		exists = entryInfo.attributes != 0;
@@ -1751,7 +1749,7 @@ BOOL Detoured_GetFileInformationByHandleEx(HANDLE hFile, FILE_INFO_BY_HANDLE_CLA
 
 	HANDLE trueHandle = hFile;
 
-	u32 entryOffset = ~u32(0);
+	DirTableOffset entryOffset;
 
 #if UBA_DEBUG_VALIDATE
 	const wchar_t* originalName = nullptr;
@@ -1769,7 +1767,7 @@ BOOL Detoured_GetFileInformationByHandleEx(HANDLE hFile, FILE_INFO_BY_HANDLE_CLA
 		FileInfo& fi = *fo.fileInfo;
 		fileSize = fi.size;
 
-		if (entryOffset == ~u32(0) && trueHandle == INVALID_HANDLE_VALUE)
+		if (!IsValidEntry(entryOffset) && trueHandle == INVALID_HANDLE_VALUE)
 		{
 			MemoryFile* mf = fi.memoryFile;
 			UBA_ASSERTF(mf || fi.name[0] == '$', L"GetFileInformationByHandleEx called on file %s which has no entry offset or real handle", HandleToName(hFile));
@@ -1843,14 +1841,12 @@ BOOL Detoured_GetFileInformationByHandleEx(HANDLE hFile, FILE_INFO_BY_HANDLE_CLA
 		}
 
 		auto& listHandle = asListDirectoryHandle(hFile);
-		if (listHandle.dir.latestOffset != InvalidTableOffset)
-			entryOffset = listHandle.dir.latestOffset | 0x80000000;
-		else
+		if (!g_directoryTable.GetLatestOffset(entryOffset, listHandle.dir))
 			UBA_ASSERT(false);
 		trueHandle = INVALID_HANDLE_VALUE;
 	}
 
-	if (entryOffset != ~u32(0))
+	if (IsValidEntry(entryOffset))
 	{
 		DirectoryTable::EntryInformation entryInfo;
 		g_directoryTable.GetEntryInformation(entryInfo, entryOffset);
@@ -1939,12 +1935,12 @@ BOOL Detoured_GetFileInformationByHandle(HANDLE hFile, LPBY_HANDLE_FILE_INFORMAT
 	if (isDetouredHandle(hFile))
 	{
 		auto& dh = asDetouredHandle(hFile);
-		u32 dirTableOffset = dh.dirTableOffset;
+		DirTableOffset dirTableOffset = dh.dirTableOffset;
 
 		UBA_ASSERT(dh.fileObject->fileInfo);
 		FileInfo& fi = *dh.fileObject->fileInfo;
 
-		if (dirTableOffset != ~u32(0))
+		if (IsValidEntry(dirTableOffset))
 		{
 			DirectoryTable::EntryInformation entryInfo;
 			g_directoryTable.GetEntryInformation(entryInfo, dirTableOffset);
diff --git a/Engine/Source/Programs/UnrealBuildAccelerator/Detours/Private/Windows/UbaDetoursFunctionsNtDll.inl b/Engine/Source/Programs/UnrealBuildAccelerator/Detours/Private/Windows/UbaDetoursFunctionsNtDll.inl
--- a/Engine/Source/Programs/UnrealBuildAccelerator/Detours/Private/Windows/UbaDetoursFunctionsNtDll.inl
+++ b/Engine/Source/Programs/UnrealBuildAccelerator/Detours/Private/Windows/UbaDetoursFunctionsNtDll.inl
@@ -180,7 +180,7 @@ NTSTATUS Detoured_NtQueryVolumeInformationFile(HANDLE FileHandle, PIO_STATUS_BLO
 			return STATUS_SUCCESS;
 		};
 
-	auto runDefaultFsVolumeInfo = [&](u32 entryOffset)
+	auto runDefaultFsVolumeInfo = [&](DirTableOffset entryOffset)
 		{
 			ULONG required = offsetof(FILE_FS_VOLUME_INFORMATION, VolumeLabel);
 			if (Length < required)
@@ -192,7 +192,7 @@ NTSTATUS Detoured_NtQueryVolumeInformationFile(HANDLE FileHandle, PIO_STATUS_BLO
 			}
 
 			auto& info = *(FILE_FS_VOLUME_INFORMATION*)FsInformation;
-			UBA_ASSERT(entryOffset != ~0u);
+			UBA_ASSERT(IsValidEntry(entryOffset));
 			DirectoryTable::EntryInformation entryInfo;
 			g_directoryTable.GetEntryInformation(entryInfo, entryOffset);
 			UBA_ASSERT(entryInfo.attributes != 0);
@@ -233,7 +233,11 @@ NTSTATUS Detoured_NtQueryVolumeInformationFile(HANDLE FileHandle, PIO_STATUS_BLO
 		if (FsInformationClass == FileFsDeviceInformation)
 			return runDefaultFsDeviceInfo();
 		if (FsInformationClass == FileFsVolumeInformation)
-			return runDefaultFsVolumeInfo(listHandle.dir.latestOffset | 0x80000000);
+		{
+			DirTableOffset tableOffset;
+			g_directoryTable.GetLatestOffset(tableOffset, listHandle.dir);
+			return runDefaultFsVolumeInfo(tableOffset);
+		}
 		UBA_ASSERTF(false, L"NtQueryVolumeInformationFile called in ListDirectoryHandle using class %u which is not implemented (%ls)", FsInformationClass, HandleToName(FileHandle));
 	}
 	auto res = True_NtQueryVolumeInformationFile(TrueHandle, IoStatusBlock, FsInformation, Length, FsInformationClass);
@@ -262,7 +266,8 @@ NTSTATUS Detoured_NtQueryInformationFile(HANDLE FileHandle, PIO_STATUS_BLOCK IoS
 			auto& info = *(FILE_ID_INFORMATION*)FileInformation;
 			if (listHandle.dir.latestOffset != InvalidTableOffset)
 			{
-				u32 entryOffset = listHandle.dir.latestOffset | 0x80000000;
+				DirTableOffset entryOffset;
+				g_directoryTable.GetLatestOffset(entryOffset, listHandle.dir);
 				DirectoryTable::EntryInformation entryInfo;
 				g_directoryTable.GetEntryInformation(entryInfo, entryOffset);
 				info.VolumeSerialNumber = entryInfo.volumeSerial;
@@ -418,7 +423,7 @@ NTSTATUS Detoured_NtQueryInformationFile(HANDLE FileHandle, PIO_STATUS_BLOCK IoS
 	}
 	if (FileInformationClass == FileBasicInformation)
 	{
-		if (dh.dirTableOffset == ~0u)
+		if (!IsValidEntry(dh.dirTableOffset))
 			return runTrue();
 		DirectoryTable::EntryInformation entryInfo;
 		g_directoryTable.GetEntryInformation(entryInfo, dh.dirTableOffset);
@@ -429,7 +434,7 @@ NTSTATUS Detoured_NtQueryInformationFile(HANDLE FileHandle, PIO_STATUS_BLOCK IoS
 	}
 	if (FileInformationClass == FileAllInformation)
 	{
-		if (dh.dirTableOffset == ~0u)
+		if (!IsValidEntry(dh.dirTableOffset))
 			return runTrue();
 		if (dh.trueHandle != INVALID_HANDLE_VALUE)
 			return runTrue();
@@ -458,7 +463,7 @@ NTSTATUS Detoured_NtQueryInformationFile(HANDLE FileHandle, PIO_STATUS_BLOCK IoS
 	}
 	if (FileInformationClass == FileNetworkOpenInformation)
 	{
-		if (dh.dirTableOffset == ~0u)
+		if (!IsValidEntry(dh.dirTableOffset))
 			return runTrue();
 		DirectoryTable::EntryInformation entryInfo;
 		g_directoryTable.GetEntryInformation(entryInfo, dh.dirTableOffset);
@@ -585,7 +590,7 @@ NTSTATUS NTAPI Detoured_NtQueryDirectoryFile(HANDLE FileHandle, HANDLE Event, PI
 			if (listHandle.it == listHandle.fileTableOffsets.size())
 				break;
 
-			u32 fileOffset = listHandle.fileTableOffsets[listHandle.it++];
+			DirTableOffset fileOffset = listHandle.fileTableOffsets[listHandle.it++];
 
 			DirectoryTable::EntryInformation entryInfo;
 			wchar_t fileName[512];
@@ -1379,7 +1384,7 @@ NTSTATUS NTAPI Shared_NtCreateFile(bool IsCreateFunc, PHANDLE hFileHandle, ACCES
 
 	DevirtualizePath(fileName);
 
-	u32 dirTableOffset = ~u32(0);
+	DirTableOffset dirTableOffset;
 
 	//UBA_ASSERT(CreateDisposition != FILE_SUPERSEDE);
 	bool isDeleteOnClose = (CreateOptions & FILE_DELETE_ON_CLOSE) != 0; // clang is using CreateFile with DeleteOnClose to delete files after build errors
@@ -1435,10 +1440,9 @@ NTSTATUS NTAPI Shared_NtCreateFile(bool IsCreateFunc, PHANDLE hFileHandle, ACCES
 		if (!isWrite && !isSystemOrTempFile) // We need to skip SystemTemp.. lots of stuff going on there.
 		{
 			CHECK_PATH(fileNameLower);
-			dirTableOffset = Rpc_GetEntryOffset(fileNameKey, fileName, checkIfDir);
 
 			bool allowEarlyOut = true;
-			if (dirTableOffset == ~u32(0))
+			if (!Rpc_GetEntryOffset(dirTableOffset, fileNameKey, fileName, checkIfDir))
 			{
 				// This could be a written file not reported to server yet
 				{
@@ -1538,7 +1542,8 @@ NTSTATUS NTAPI Shared_NtCreateFile(bool IsCreateFunc, PHANDLE hFileHandle, ACCES
 		bool exists = false;
 		if (dir.latestOffset != InvalidTableOffset)
 		{
-			u32 entryOffset = dir.latestOffset | 0x80000000;
+			DirTableOffset entryOffset;
+			g_directoryTable.GetLatestOffset(entryOffset, dir);
 			DirectoryTable::EntryInformation entryInfo;
 			g_directoryTable.GetEntryInformation(entryInfo, entryOffset);
 			exists = entryInfo.attributes != 0;
diff --git a/Engine/Source/Programs/UnrealBuildAccelerator/Detours/Private/Windows/UbaDetoursFunctionsWin.cpp b/Engine/Source/Programs/UnrealBuildAccelerator/Detours/Private/Windows/UbaDetoursFunctionsWin.cpp
--- a/Engine/Source/Programs/UnrealBuildAccelerator/Detours/Private/Windows/UbaDetoursFunctionsWin.cpp
+++ b/Engine/Source/Programs/UnrealBuildAccelerator/Detours/Private/Windows/UbaDetoursFunctionsWin.cpp
@@ -173,7 +173,7 @@ struct ListDirectoryHandle
 	StringKey dirNameKey;
 	DirectoryTable::Directory& dir;
 	int it;
-	Vector<u32> fileTableOffsets;
+	Vector<DirTableOffset> fileTableOffsets;
 	HANDLE validateHandle;
 	TString wildcard;
 	const wchar_t* originalName = nullptr;
diff --git a/Engine/Source/Programs/UnrealBuildAccelerator/Detours/Private/Windows/UbaDetoursUtilsWin.h b/Engine/Source/Programs/UnrealBuildAccelerator/Detours/Private/Windows/UbaDetoursUtilsWin.h
--- a/Engine/Source/Programs/UnrealBuildAccelerator/Detours/Private/Windows/UbaDetoursUtilsWin.h
+++ b/Engine/Source/Programs/UnrealBuildAccelerator/Detours/Private/Windows/UbaDetoursUtilsWin.h
@@ -52,7 +52,7 @@ namespace uba
 		DetouredHandle(HandleType t, HANDLE th = INVALID_HANDLE_VALUE) : trueHandle(th), type(t) {}
 
 		HANDLE trueHandle;
-		u32 dirTableOffset = ~u32(0);
+		DirTableOffset dirTableOffset;
 		HandleType type;
 
 		// Only for files
diff --git a/Engine/Source/Programs/UnrealBuildAccelerator/Test/Private/UbaTestDirectoryTable.cpp b/Engine/Source/Programs/UnrealBuildAccelerator/Test/Private/UbaTestDirectoryTable.cpp
--- a/Engine/Source/Programs/UnrealBuildAccelerator/Test/Private/UbaTestDirectoryTable.cpp
+++ b/Engine/Source/Programs/UnrealBuildAccelerator/Test/Private/UbaTestDirectoryTable.cpp
@@ -26,10 +26,18 @@ namespace uba
 		StringBuffer<> fooFile;
 		if (!CreateTestFile(fooFile, logger, testRootDir, TCV("foo.txt"), TCV("")))
 			return false;
-
 		if (!holder.RegisterNewFile(fooFile.data))
 			return false;
 
+
+
+		StringBuffer<> barFile;
+		if (!CreateTestFile(barFile, logger, testRootDir, TCV("bar.txt"), TCV("")))
+			return false;
+		if (!holder.RegisterNewFile(barFile.data))
+			return false;
+
+
 		return true;
 	}
 }
