[UBA]
* Changed "Last write time has changed since ..." log entry to be info instead of warning if new time is older than old. It is still very strange but at least I believe this is of no risk then

diff --git a/Engine/Source/Programs/UnrealBuildAccelerator/Common/Private/UbaStorage.cpp b/Engine/Source/Programs/UnrealBuildAccelerator/Common/Private/UbaStorage.cpp
--- a/Engine/Source/Programs/UnrealBuildAccelerator/Common/Private/UbaStorage.cpp
+++ b/Engine/Source/Programs/UnrealBuildAccelerator/Common/Private/UbaStorage.cpp
@@ -455,7 +455,9 @@ namespace uba
 	
 		if (deferredSource && deferredSource->lastWriteTime != basicInfo.lastWriteTime)
 		{
-			m_logger.Warning(TC("[WriteCasFileNoCheck] Last write time has changed since %s was deferred for cas creation. (%llu to %llu)"), fileName, deferredSource->lastWriteTime, basicInfo.lastWriteTime);
+			// If last write time is newer for deferred source, then it must be a bug where last write time is stored wrong when the file is written.
+			LogEntryType logType = deferredSource->lastWriteTime < basicInfo.lastWriteTime ? LogEntryType_Warning : LogEntryType_Info;
+			m_logger.Logf(logType, TC("[WriteCasFileNoCheck] Last write time has changed since %s was deferred for cas creation. (LastWrite %llu to %llu, Size %llu to %llu)"), fileName, deferredSource->lastWriteTime, basicInfo.lastWriteTime, deferredSource->fileSize, basicInfo.size);
 			if (allowFail)
 				return false;
 		}
