[UBA]
* Promoted PopulateOverlayFile to UbaDirectoryTableHolder

diff --git a/Engine/Source/Programs/UnrealBuildAccelerator/Common/Private/UbaDirectoryTableHolder.cpp b/Engine/Source/Programs/UnrealBuildAccelerator/Common/Private/UbaDirectoryTableHolder.cpp
--- a/Engine/Source/Programs/UnrealBuildAccelerator/Common/Private/UbaDirectoryTableHolder.cpp
+++ b/Engine/Source/Programs/UnrealBuildAccelerator/Common/Private/UbaDirectoryTableHolder.cpp
@@ -200,13 +200,12 @@ namespace uba
 		UBA_ASSERT(overlay.handle.IsValid());
 		overlay.memory = MapViewOfFile(m_logger, overlay.handle, FILE_MAP_WRITE, 0, OverlayTableMemSize);
 		UBA_ASSERT(overlay.memory);
-		overlay.committed = 0;
-		overlay.size = 0;
 	}
 
 	void DirectoryTableHolder::ResetOverlay(DirectoryTableOverlay& overlay, void* data, u64 dataSize)
 	{
 		overlay.size = 0;
+		overlay.lookup.clear();
 		if (!dataSize)
 			return;
 		EnsureOverlayMemory(overlay, dataSize);
@@ -214,15 +213,16 @@ namespace uba
 		overlay.size = u32(dataSize);
 	}
 
-	void DirectoryTableHolder::EnsureOverlayMemory(DirectoryTableOverlay& overlay, u64 neededSize)
+	bool DirectoryTableHolder::EnsureOverlayMemory(DirectoryTableOverlay& overlay, u64 neededSize)
 	{
 		u64 totalSize = overlay.size + neededSize;
 		if (totalSize <= overlay.committed)
-			return;
-		u64 newSize = AlignUp(totalSize, 1024ull*1024);
+			return true;
+		u64 newSize = AlignUp(totalSize, 32ull*1024);
 		if (newSize > OverlayTableMemSize)
 		{
 			static bool called = m_logger.Error(TC("Overlay table overflow. OverlayTableMemSize need to be increased (Size: %llu)"), overlay.size);
+			return false;
 		}
 
 		u64 toCommit = newSize - overlay.committed;
@@ -230,6 +230,7 @@ namespace uba
 		if (!MapViewCommit(m_logger, address, toCommit, false))
 			m_logger.Error(TC("Failed to commit memory for directory table (Committed: %llu, ToCommit: %llu) (%s)"), m_directoryTableMemCommitted, toCommit, LastErrorToText().data);
 		overlay.committed += u32(toCommit);
+		return true;
 	}
 
 	void DirectoryTableHolder::DestroyOverlay(DirectoryTableOverlay& overlay)
@@ -243,6 +244,82 @@ namespace uba
 		overlay.handle = {};
 		overlay.committed = 0;
 		overlay.size = 0;
+		overlay.lookup.clear();
+	}
+
+	bool DirectoryTableHolder::PopulateOverlayFile(DirectoryTableOverlay& overlay, StringView fileName, const Function<bool(FileInformation&, StringKey)>& getFileInformation)
+	{
+		auto& dirTable = m_directoryTable;
+
+		StringKey dirKey;
+		const tchar* lastSlash = nullptr;
+		StringBuffer<> fixedFileName;
+		StringKey fileNameKey = GetKeyAndFixedName(fixedFileName, fileName.data, &dirKey, &lastSlash);
+
+		FileInformation fileInfo;
+		if (!getFileInformation(fileInfo, fileNameKey))
+			return false;
+
+		auto dirOffsetIns = overlay.lookup.try_emplace(dirKey);
+		if (dirOffsetIns.second)
+		{
+			StringView dirName(fixedFileName.data, u32(lastSlash - fixedFileName.data));
+
+			SCOPED_READ_LOCK(dirTable.m_lookupLock, lookupLock);
+			auto findIt = dirTable.m_lookup.find(dirKey);
+			if (findIt == dirTable.m_lookup.end())
+			{
+				lookupLock.Leave();
+				u32 res = WriteDirectoryEntries(dirKey, dirName);
+				UBA_ASSERTF(res, TC("Failed to write directory entries for %.*s"), dirName.count, dirName.data); (void)res;
+				lookupLock.Enter();
+				findIt = dirTable.m_lookup.find(dirKey);
+				if (findIt == dirTable.m_lookup.end())
+					return m_logger.Error(TC("OverlayFiles can only exist in existing directories (%.*s)"), fixedFileName.count, fixedFileName.data);
+			}
+			DirectoryTable::Directory& dir = findIt->second;
+			lookupLock.Leave();
+
+			SCOPED_WRITE_LOCK(dir.lock, dirLock);
+			if (dir.populatedOffset != 3)
+			{
+				while (dir.populatedOffset == 0)
+				{
+					dirLock.Leave();
+					Sleep(1);
+					dirLock.Enter();
+				}
+				dir.populatedOffset = 3; // 3 means it is locked, it is now a directory owned by overlays
+
+#if UBA_DEBUG_TRACK_DIR
+				m_debugLogger->Info(TC("TRACKOVERL %.*s"), dirName.count, dirName.data);
+#endif
+			}
+			dirOffsetIns.first->second = dir.latestOffset;
+		}
+
+		u8 temp[1024];
+		BinaryWriter writer(temp, 0, sizeof(temp));
+		writer.WriteStringKey(dirKey);
+		writer.Write7BitEncoded(dirOffsetIns.first->second); // Previous entry for same directory
+		writer.Write7BitEncoded(1); // Count
+		writer.WriteString(lastSlash + 1);
+		writer.WriteFileAttributes(fileInfo.attributes);
+		writer.WriteVolumeSerial(fileInfo.volumeSerialNumber);
+		writer.WriteFileIndex(fileInfo.index);
+		writer.WriteFileTime(fileInfo.lastWriteTime);
+		writer.WriteFileSize(fileInfo.size);
+		u64 written = writer.GetPosition();
+
+		if (!EnsureOverlayMemory(overlay, written + 8))
+			return false;
+
+		BinaryWriter memWriter(overlay.memory, overlay.size, overlay.committed);
+		memWriter.Write7BitEncoded(written); // Storage size
+		dirOffsetIns.first->second = u32(memWriter.GetPosition()) + sizeof(StringKey) + OverlayTableOffset;
+		memWriter.WriteBytes(writer.GetData(), written);
+		overlay.size = u32(memWriter.GetPosition());
+		return true;
 	}
 
 	bool DirectoryTableHolder::RegisterCreateFileForWrite(StringKey fileNameKey, StringView fileName, DirectoryTableOverlay* overlay, u64 fileSize, u64 lastWriteTime, bool registerRealFile, bool invalidateStorage)
diff --git a/Engine/Source/Programs/UnrealBuildAccelerator/Common/Private/UbaSession.cpp b/Engine/Source/Programs/UnrealBuildAccelerator/Common/Private/UbaSession.cpp
--- a/Engine/Source/Programs/UnrealBuildAccelerator/Common/Private/UbaSession.cpp
+++ b/Engine/Source/Programs/UnrealBuildAccelerator/Common/Private/UbaSession.cpp
@@ -1019,11 +1019,21 @@ namespace uba
 
 	bool Session::PopulateOverlayFiles(DirectoryTableOverlay& overlay, StringView overlayFiles)
 	{
-		auto& dirTable = m_directoryTable;
-		UBA_ASSERT(overlay.size == 0);
-		BinaryWriter memWriter(overlay.memory, 0, OverlayTableMemSize);
-
-		UnorderedMap<StringKey, u32> dirOffsets;
+		auto getFileInformation = [&](FileInformation& out, StringKey fileNameKey)
+			{
+				SCOPED_FUTEX(m_fileMappingTableLookupLock, mappingLookupLock);
+				auto findIt = m_fileMappingTableLookup.find(fileNameKey);
+				if (findIt == m_fileMappingTableLookup.end())
+					return false;
+				FileMappingEntry& entry = findIt->second;
+				mappingLookupLock.Leave();
+				out.size = entry.contentSize;
+				out.lastWriteTime = entry.lastWriteTime;
+				out.attributes = DefaultAttributes();
+				out.volumeSerialNumber = 0;
+				out.index = ++m_fileIndexCounter;
+				return true;
+			};
 
 		tchar* it = const_cast<tchar*>(overlayFiles.data);
 		tchar* end = it + overlayFiles.count;
@@ -1034,88 +1044,13 @@ namespace uba
 				*fileEnd = 0;
 			else
 				fileEnd = end;
-
-			StringKey dirKey;
-			const tchar* lastSlash = nullptr;
-			StringBuffer<> fixedFileName;
-			StringKey fileNameKey = GetKeyAndFixedName(fixedFileName, it, &dirKey, &lastSlash);
-
-			u64 fileSize = 0;
-			u64 lastWriteTime = 0;
-			{
-				SCOPED_FUTEX(m_fileMappingTableLookupLock, mappingLookupLock);
-				auto findIt = m_fileMappingTableLookup.find(fileNameKey);
-				if (findIt == m_fileMappingTableLookup.end())
-					return false;
-				FileMappingEntry& entry = findIt->second;
-				mappingLookupLock.Leave();
-				fileSize = entry.contentSize;
-				lastWriteTime = entry.lastWriteTime;
-			}
-
-			auto dirOffsetIns = dirOffsets.try_emplace(dirKey);
-			if (dirOffsetIns.second)
-			{
-				StringView dirName(fixedFileName.data, u32(lastSlash - fixedFileName.data));
-
-				SCOPED_READ_LOCK(dirTable.m_lookupLock, lookupLock);
-				auto findIt = dirTable.m_lookup.find(dirKey);
-				if (findIt == dirTable.m_lookup.end())
-				{
-					lookupLock.Leave();
-					u32 res = WriteDirectoryEntries(dirKey, dirName);
-					UBA_ASSERTF(res, TC("Failed to write directory entries for %.*s"), dirName.count, dirName.data); (void)res;
-					lookupLock.Enter();
-					findIt = dirTable.m_lookup.find(dirKey);
-					if (findIt == dirTable.m_lookup.end())
-						return m_logger.Error(TC("OverlayFiles can only exist in existing directories (%.*s)"), fixedFileName.count, fixedFileName.data);
-				}
-				DirectoryTable::Directory& dir = findIt->second;
-				lookupLock.Leave();
-
-				SCOPED_WRITE_LOCK(dir.lock, dirLock);
-				if (dir.populatedOffset != 3)
-				{
-					while (dir.populatedOffset == 0)
-					{
-						dirLock.Leave();
-						Sleep(1);
-						dirLock.Enter();
-					}
-					dir.populatedOffset = 3; // 3 means it is locked, it is now a directory owned by overlays
-
-					#if UBA_DEBUG_TRACK_DIR
-					m_debugLogger->Info(TC("TRACKOVERL %.*s"), dirName.count, dirName.data);
-					#endif
-				}
-				dirOffsetIns.first->second = dir.latestOffset;
-			}
-
-
-			u8 temp[1024];
-			BinaryWriter writer(temp, 0, sizeof(temp));
-			writer.WriteStringKey(dirKey);
-			writer.Write7BitEncoded(dirOffsetIns.first->second); // Previous entry for same directory
-			writer.Write7BitEncoded(1); // Count
-			writer.WriteString(lastSlash + 1);
-			writer.WriteFileAttributes(DefaultAttributes());
-			writer.WriteVolumeSerial(0);
-			writer.WriteFileIndex(++m_fileIndexCounter);
-			writer.WriteFileTime(lastWriteTime);
-			writer.WriteFileSize(fileSize);
-			u64 written = writer.GetPosition();
-
-			EnsureOverlayMemory(overlay, written + 8);
-
-			memWriter.Write7BitEncoded(written); // Storage size
-			dirOffsetIns.first->second = u32(memWriter.GetPosition()) + sizeof(StringKey) + OverlayTableOffset;
-			memWriter.WriteBytes(writer.GetData(), written);
+			if (!PopulateOverlayFile(overlay, StringView(it, u32(fileEnd - it)), getFileInformation))
+				return false;
 			if (fileEnd == end)
 				break;
 			*fileEnd = ';';
 			it = ++fileEnd;
 		}
-		overlay.size = u32(memWriter.GetPosition());
 		return true;
 	}
 
diff --git a/Engine/Source/Programs/UnrealBuildAccelerator/Common/Public/UbaDirectoryTableHolder.h b/Engine/Source/Programs/UnrealBuildAccelerator/Common/Public/UbaDirectoryTableHolder.h
--- a/Engine/Source/Programs/UnrealBuildAccelerator/Common/Public/UbaDirectoryTableHolder.h
+++ b/Engine/Source/Programs/UnrealBuildAccelerator/Common/Public/UbaDirectoryTableHolder.h
@@ -19,6 +19,7 @@ namespace uba
 		u8* memory = nullptr;
 		u32 committed = 0;
 		u32 size = 0;
+		UnorderedMap<StringKey, u32> lookup;
 	};
 
 	////////////////////////////////////////////////////////////////////////////////////////////////////
@@ -40,12 +41,14 @@ namespace uba
 
 		void CreateOverlay(DirectoryTableOverlay& overlay);
 		void ResetOverlay(DirectoryTableOverlay& overlay, void* data, u64 dataSize);
-		void EnsureOverlayMemory(DirectoryTableOverlay& overlay, u64 neededSize);
+		bool EnsureOverlayMemory(DirectoryTableOverlay& overlay, u64 neededSize);
 		void DestroyOverlay(DirectoryTableOverlay& overlay);
 
 		MutableLogger m_logger;
 
-	protected:
+	//protected:
+
+		bool PopulateOverlayFile(DirectoryTableOverlay& overlay, StringView fileName, const Function<bool(FileInformation& out, StringKey)>& getFileInformation);
 
 		bool RegisterCreateFileForWrite(StringKey fileNameKey, StringView fileName, DirectoryTableOverlay* overlay, u64 fileSize, u64 lastWriteTime, bool registerRealFile, bool invalidateStorage);
 		DirTableSize RegisterDeleteFile(StringKey fileNameKey, StringView fileName, DirectoryTableOverlay* overlay);
