[UBA]
* Changed so overlay has its own linked list of directory offsets separate from main one. This will make it possible to mix normal updates and overlays

diff --git a/Engine/Source/Programs/UnrealBuildAccelerator/Common/Private/UbaDirectoryTableHolder.cpp b/Engine/Source/Programs/UnrealBuildAccelerator/Common/Private/UbaDirectoryTableHolder.cpp
--- a/Engine/Source/Programs/UnrealBuildAccelerator/Common/Private/UbaDirectoryTableHolder.cpp
+++ b/Engine/Source/Programs/UnrealBuildAccelerator/Common/Private/UbaDirectoryTableHolder.cpp
@@ -249,7 +249,7 @@ namespace uba
 
 	bool DirectoryTableHolder::PopulateOverlayFile(DirectoryTableOverlay& overlay, StringView fileName, const Function<bool(FileInformation&, StringKey)>& getFileInformation)
 	{
-		auto& dirTable = m_directoryTable;
+		//auto& dirTable = m_directoryTable;
 
 		StringKey dirKey;
 		const tchar* lastSlash = nullptr;
@@ -260,7 +260,13 @@ namespace uba
 		if (!getFileInformation(fileInfo, fileNameKey))
 			return false;
 
+		SCOPED_WRITE_LOCK(overlay.lock, lock);
+
 		auto dirOffsetIns = overlay.lookup.try_emplace(dirKey);
+		
+		// TODO: we should use the normal dir table to decide how far up the tree we need to go to make sure we have an entire chain of directories
+		
+		/*
 		if (dirOffsetIns.second)
 		{
 			StringView dirName(fixedFileName.data, u32(lastSlash - fixedFileName.data));
@@ -277,7 +283,7 @@ namespace uba
 				if (findIt == dirTable.m_lookup.end())
 					return m_logger.Error(TC("OverlayFiles can only exist in existing directories (%.*s)"), fixedFileName.count, fixedFileName.data);
 			}
-			DirectoryTable::Directory& dir = findIt->second;
+			//DirectoryTable::Directory& dir = findIt->second;
 			lookupLock.Leave();
 
 			SCOPED_WRITE_LOCK(dir.lock, dirLock);
@@ -297,11 +303,18 @@ namespace uba
 			}
 			dirOffsetIns.first->second = dir.latestOffset;
 		}
+		*/
 
 		u8 temp[1024];
 		BinaryWriter writer(temp, 0, sizeof(temp));
 		writer.WriteStringKey(dirKey);
 		writer.Write7BitEncoded(dirOffsetIns.first->second); // Previous entry for same directory
+		if (!dirOffsetIns.first->second)
+		{
+			writer.WriteFileAttributes(1);
+			writer.WriteVolumeSerial(fileInfo.volumeSerialNumber); // Reuse file info serial
+			writer.WriteFileIndex(++m_fileIndexCounter);
+		}
 		writer.Write7BitEncoded(1); // Count
 		writer.WriteString(lastSlash + 1);
 		writer.WriteFileAttributes(fileInfo.attributes);
diff --git a/Engine/Source/Programs/UnrealBuildAccelerator/Common/Private/UbaSessionClient.cpp b/Engine/Source/Programs/UnrealBuildAccelerator/Common/Private/UbaSessionClient.cpp
--- a/Engine/Source/Programs/UnrealBuildAccelerator/Common/Private/UbaSessionClient.cpp
+++ b/Engine/Source/Programs/UnrealBuildAccelerator/Common/Private/UbaSessionClient.cpp
@@ -2393,7 +2393,7 @@ namespace uba
 		lock2.Leave();
 		if (newMemPosition == m_dirtableParsedPosition)
 			return false;
-		m_directoryTable.ParseDirectoryTableNoLock(m_directoryTable.m_memory, m_dirtableParsedPosition, newMemPosition);
+		m_directoryTable.ParseDirectoryTableNoLock(m_directoryTable.m_memory, m_dirtableParsedPosition, newMemPosition, false);
 		m_dirtableParsedPosition = newMemPosition;
 		return true;
 	}
@@ -2988,7 +2988,7 @@ namespace uba
 		}
 
 		u32 foundCount = 0;
-		dirTable.TraverseAllFilesNoLock([&](const DirectoryTable::EntryInformation& info, const StringBufferBase& path, u32 dirOffset)
+		dirTable.TraverseAllFilesNoLock([&](const DirectoryTable::EntryInformation& info, const StringBufferBase& path, DirTableOffset dirOffset)
 			{
 				if (!path.EndsWith(searchString))
 					return;
diff --git a/Engine/Source/Programs/UnrealBuildAccelerator/Core/Public/UbaDirectoryTable.h b/Engine/Source/Programs/UnrealBuildAccelerator/Core/Public/UbaDirectoryTable.h
--- a/Engine/Source/Programs/UnrealBuildAccelerator/Core/Public/UbaDirectoryTable.h
+++ b/Engine/Source/Programs/UnrealBuildAccelerator/Core/Public/UbaDirectoryTable.h
@@ -58,6 +58,8 @@ namespace uba
 			Directory(MemoryBlock& block) : files(block) {}
 			u32 latestOffset = InvalidTableOffset;
 			u32 populatedOffset = InvalidTableOffset;
+			u32 latestOverlayOffset = InvalidTableOffset;
+			u32 populatedOverlayOffset = InvalidTableOffset;
 			EntryLookup files;
 			ReaderWriterLock lock;
 		};
@@ -68,7 +70,7 @@ namespace uba
 		inline void DeinitOverlay(const u8* oldMem, u32 size);
 		inline void ParseDirectoryTable(DirTableSize size);
 		inline void ParseDirectoryTableNoLock(DirTableSize size);
-		inline void ParseDirectoryTableNoLock(const u8* memory, u32 from, u32 to, u32 additionalOffset = 0);
+		inline void ParseDirectoryTableNoLock(const u8* memory, u32 from, u32 to, bool isOverlay);
 		inline void PopulateDirectory(const StringKeyHasher& hasher, Directory& dir);
 		inline bool IsDirectoryUpToDate(Directory& dir);
 		inline void PopulateDirectoryNoLock(const StringKeyHasher& hasher, Directory& dir);
@@ -109,7 +111,8 @@ namespace uba
 		template<typename Func> inline void TraverseOriginalFilesNoPopulate(DirTableOffset tableOffset, const Func& func);
 
 		#if PLATFORM_WINDOWS
-		template<typename Func> inline void TraverseFilesRecursiveNoLock(const StringBufferBase& path, const Func& func);
+		// Func is (const DirectoryTable::EntryInformation& info, StringView fileName, DirTableOffset fileOffset)
+		template<typename Func> inline void TraverseFilesRecursiveNoLock(StringView path, const Func& func);
 		template<typename Func> inline void TraverseAllFilesNoLock(const Func& func);
 		#endif
 
@@ -126,6 +129,7 @@ namespace uba
 		u32 m_memorySize = 0;
 
 		const u8* m_overlayMemory = nullptr;
+		u32 m_overlaySize = 0;
 	};
 
 	////////////////////////////////////////////////////////////////////////////////////////////////////
@@ -141,7 +145,8 @@ namespace uba
 	void DirectoryTable::InitOverlay(const u8* mem, u32 tableSize)
 	{
 		m_overlayMemory = mem;
-		ParseDirectoryTableNoLock(mem, 0, tableSize, OverlayTableOffset);
+		ParseDirectoryTableNoLock(mem, 0, tableSize, true);
+		m_overlaySize = tableSize;
 	}
 
 	void DirectoryTable::DeinitOverlay(const u8* oldMem, u32 size)
@@ -185,6 +190,7 @@ namespace uba
 			++i;
 		}
 		m_overlayMemory = nullptr;
+		m_overlaySize = 0;
 	}
 
 	void DirectoryTable::ParseDirectoryTable(DirTableSize size)
@@ -195,14 +201,21 @@ namespace uba
 
 	void DirectoryTable::ParseDirectoryTableNoLock(DirTableSize size)
 	{
-		if (size.main <= m_memorySize)
-			return;
-		ParseDirectoryTableNoLock(m_memory, m_memorySize, size.main);
-		m_memorySize = size.main;
+		if (size.main > m_memorySize)
+		{
+			ParseDirectoryTableNoLock(m_memory, m_memorySize, size.main, false);
+			m_memorySize = size.main;
+		}
+		if (size.overlay > m_overlaySize)
+		{
+			ParseDirectoryTableNoLock(m_overlayMemory, m_overlaySize, size.overlay, true);
+			m_overlaySize = size.overlay;
+		}
 	}
 
-	void DirectoryTable::ParseDirectoryTableNoLock(const u8* memory, u32 from, u32 to, u32 additionalOffset)
+	void DirectoryTable::ParseDirectoryTableNoLock(const u8* memory, u32 from, u32 to, bool isOverlay)
 	{
+		u32 additionalOffset = isOverlay ? OverlayTableOffset : 0;
 		BinaryReader reader(memory, from, to);
 		while (true)
 		{
@@ -213,7 +226,7 @@ namespace uba
 			u64 storageSize = reader.Read7BitEncoded();
 			StringKey dirKey = reader.ReadStringKey();
 			auto insres = m_lookup.try_emplace(dirKey, m_memoryBlock); // Note that this is allowed to overwrite
-			u32& tableOffset = insres.first->second.latestOffset;
+			u32& tableOffset = isOverlay ? insres.first->second.latestOverlayOffset : insres.first->second.latestOffset;
 			UBA_ASSERTF(tableOffset < OverlayTableOffset || additionalOffset == OverlayTableOffset, TC("Can't add non-overlay directory items after overlay items"));
 			tableOffset = u32(reader.GetPosition()) + additionalOffset;
 			reader.Skip(storageSize - sizeof(dirKey));
@@ -228,7 +241,7 @@ namespace uba
 
 	bool DirectoryTable::IsDirectoryUpToDate(Directory& dir)
 	{
-		return dir.populatedOffset == dir.latestOffset;
+		return dir.populatedOffset == dir.latestOffset && dir.populatedOverlayOffset == dir.latestOverlayOffset;
 	}
 
 	void DirectoryTable::PopulateDirectoryNoLock(const StringKeyHasher& hasher, Directory& dir)
@@ -383,6 +396,15 @@ namespace uba
 						parentDir.populatedOffset = parentDir.latestOffset;
 					}
 				}
+				if (parentDir.populatedOverlayOffset != parentDir.latestOverlayOffset)
+				{
+					SCOPED_WRITE_LOCK(parentDir.lock, lock);
+					if (parentDir.populatedOverlayOffset != parentDir.latestOverlayOffset)
+					{
+						PopulateDirectoryRecursive(ancestorHasher, parentDir.latestOverlayOffset, parentDir.populatedOverlayOffset, parentDir.files);
+						parentDir.populatedOverlayOffset = parentDir.latestOverlayOffset;
+					}
+				}
 
 				SCOPED_READ_LOCK(parentDir.lock, lock);
 				auto entryIt = parentDir.files.find(entryKey);
@@ -600,7 +622,7 @@ namespace uba
 
 #if PLATFORM_WINDOWS
 	template<typename Func>
-	void DirectoryTable::TraverseFilesRecursiveNoLock(const StringBufferBase& path, const Func& func)
+	void DirectoryTable::TraverseFilesRecursiveNoLock(StringView path, const Func& func)
 	{
 		auto findIt = m_lookup.find(ToStringKey(path));
 		if (findIt == m_lookup.end())
@@ -613,7 +635,7 @@ namespace uba
 			DirectoryTable::EntryInformation info;
 			StringBuffer<> fileName(path);
 			fileName.Append(PathSeparator);
-			u32 fileOffset = fileKv.second;
+			DirTableOffset fileOffset = fileKv.second;
 			GetEntryInformation(info, fileOffset, fileName.data + fileName.count, fileName.capacity - fileName.count);
 			fileName.count = TStrlen(fileName.data);
 			if (CaseInsensitiveFs)
diff --git a/Engine/Source/Programs/UnrealBuildAccelerator/Detours/Private/PosixOS/UbaDetoursFunctionsPosix.cpp b/Engine/Source/Programs/UnrealBuildAccelerator/Detours/Private/PosixOS/UbaDetoursFunctionsPosix.cpp
--- a/Engine/Source/Programs/UnrealBuildAccelerator/Detours/Private/PosixOS/UbaDetoursFunctionsPosix.cpp
+++ b/Engine/Source/Programs/UnrealBuildAccelerator/Detours/Private/PosixOS/UbaDetoursFunctionsPosix.cpp
@@ -3167,6 +3167,7 @@ extern "C"
 			{
 				const u8* overlayMem = g_directoryTable.m_overlayMemory;
 				g_directoryTable.DeinitOverlay(oldOverlayFiles.data(), g_overlayFilesSize);
+				directoryTableSize.overlay = 0; // Want to just parse normal directory first
 				g_directoryTable.ParseDirectoryTableNoLock(directoryTableSize);
 				g_directoryTable.InitOverlay(overlayMem, overlaySize);
 				g_overlayFilesSize = overlaySize;
diff --git a/Engine/Source/Programs/UnrealBuildAccelerator/Detours/Private/Windows/UbaDetoursFunctionsWin.cpp b/Engine/Source/Programs/UnrealBuildAccelerator/Detours/Private/Windows/UbaDetoursFunctionsWin.cpp
--- a/Engine/Source/Programs/UnrealBuildAccelerator/Detours/Private/Windows/UbaDetoursFunctionsWin.cpp
+++ b/Engine/Source/Programs/UnrealBuildAccelerator/Detours/Private/Windows/UbaDetoursFunctionsWin.cpp
@@ -1349,6 +1349,7 @@ extern "C"
 			{
 				const u8* overlayMem = g_directoryTable.m_overlayMemory;
 				g_directoryTable.DeinitOverlay(oldOverlayFiles.data(), g_overlayFilesSize);
+				directoryTableSize.overlay = 0; // Want to just parse normal directory first
 				g_directoryTable.ParseDirectoryTableNoLock(directoryTableSize);
 				g_directoryTable.InitOverlay(overlayMem, overlaySize);
 				g_overlayFilesSize = overlaySize;
diff --git a/Engine/Source/Programs/UnrealBuildAccelerator/Test/Private/UbaTestDirectoryTable.cpp b/Engine/Source/Programs/UnrealBuildAccelerator/Test/Private/UbaTestDirectoryTable.cpp
--- a/Engine/Source/Programs/UnrealBuildAccelerator/Test/Private/UbaTestDirectoryTable.cpp
+++ b/Engine/Source/Programs/UnrealBuildAccelerator/Test/Private/UbaTestDirectoryTable.cpp
@@ -19,23 +19,68 @@ namespace uba
 		
 		if (dir.EntryExists(testRootDir, false) != DirectoryTable::Exists_Yes)
 			return false;
-
 		if (dir.EntryExists(testRootDir, true) != DirectoryTable::Exists_Yes)
 			return false;
 
 		StringBuffer<> fooFile;
-		if (!CreateTestFile(fooFile, logger, testRootDir, TCV("foo.txt"), TCV("")))
+		if (!CreateTestFile(fooFile, logger, testRootDir, TCV("foo.txt"), TCV("foo")))
 			return false;
 		if (!holder.RegisterNewFile(fooFile.data))
 			return false;
 
+		if (dir.EntryExists(fooFile, true) != DirectoryTable::Exists_No)
+			return false;
+		dir.ParseDirectoryTable(holder.GetDirectoryTableSize(nullptr));
+		if (dir.EntryExists(fooFile, true) != DirectoryTable::Exists_Yes)
+			return false;
+
+
+		DirectoryTableOverlay overlay;
+		holder.CreateOverlay(overlay);
 
+		dir.InitOverlay(overlay.memory, overlay.size);
 
 		StringBuffer<> barFile;
-		if (!CreateTestFile(barFile, logger, testRootDir, TCV("bar.txt"), TCV("")))
+		if (!CreateTestFile(barFile, logger, testRootDir, TCV("bar.txt"), TCV("bar")))
+			return false;
+
+		holder.PopulateOverlayFile(overlay, barFile, [](FileInformation& out, StringKey) { out.size = 6; out.attributes = DefaultAttributes(); return true; });
+
+		if (dir.EntryExists(barFile, true) != DirectoryTable::Exists_No)
+			return false;
+		dir.ParseDirectoryTable(holder.GetDirectoryTableSize(&overlay));
+		if (dir.EntryExists(barFile, true) != DirectoryTable::Exists_Yes)
+			return false;
+
+		StringBuffer<> mehFile;
+		if (!CreateTestFile(mehFile, logger, testRootDir, TCV("meh.txt"), TCV("meh")))
+			return false;
+		if (!holder.RegisterNewFile(mehFile.data))
+			return false;
+
+		if (dir.EntryExists(mehFile, true) != DirectoryTable::Exists_No)
+			return false;
+		dir.ParseDirectoryTable(holder.GetDirectoryTableSize(&overlay));
+		if (dir.EntryExists(mehFile, true) != DirectoryTable::Exists_Yes)
+			return false;
+
+#if PLATFORM_WINDOWS
+		u32 counter = 0;
+		StringBuffer<> testRootDirLower(StringView(testRootDir.data, testRootDir.count - 1));
+		testRootDirLower.MakeLower();
+		bool success = true;
+		dir.TraverseFilesRecursiveNoLock(testRootDirLower, [&](const DirectoryTable::EntryInformation& info, StringView fileName, DirTableOffset fileOffset)
+			{
+				++counter;
+				success = success && info.size == 6;
+				success = success && info.attributes != 0;
+				success = success && fileName.GetFileName().count == 7;
+			});
+		if (!success)
 			return false;
-		if (!holder.RegisterNewFile(barFile.data))
+		if (counter != 3)
 			return false;
+#endif
 
 
 		return true;
