[UBA]
* Improved assert when reading garbage string from stream

diff --git a/Engine/Source/Programs/UnrealBuildAccelerator/Core/Public/UbaBinaryReaderWriter.h b/Engine/Source/Programs/UnrealBuildAccelerator/Core/Public/UbaBinaryReaderWriter.h
--- a/Engine/Source/Programs/UnrealBuildAccelerator/Core/Public/UbaBinaryReaderWriter.h
+++ b/Engine/Source/Programs/UnrealBuildAccelerator/Core/Public/UbaBinaryReaderWriter.h
@@ -393,11 +393,12 @@ namespace uba
 			}
 			if (a >= 240 && a <= 253)
 			{
-				UBA_ASSERT(false); // Wide chars cannot exceed 16 bits
+				// Wide chars cannot exceed 16 bits
+				UBA_ASSERTF(false, TC("This should not happen, data is corrupt (Reader Start:0x%llx Pos:%llu Length:%llu Data:%.*s)"), u64(m_begin), u64(pos - m_begin), u64(m_end - m_begin), u32(it - str), str); // Wide chars cannot exceed 16 bits
 				*it++ = tchar(~0);
 				continue;
 			}
-			UBA_ASSERT(false); // Wide chars cannot exceed 16 bits
+			UBA_ASSERTF(false, TC("This should not happen, data is corrupt (Reader Start:0x%llx Pos:%llu Length:%llu Data:%.*s)"), u64(m_begin), u64(pos - m_begin), u64(m_end - m_begin), u32(it - str), str); // Wide chars cannot exceed 16 bits
 			*it++ = tchar(~0);
 		}
 		*it = 0;
