[UBA]
* Fixed so coordinator status is showing in visualizer when running through ubacli etc

diff --git a/Engine/Source/Developer/UbaCoordinatorHorde/Private/UbaCoordinatorHorde.cpp b/Engine/Source/Developer/UbaCoordinatorHorde/Private/UbaCoordinatorHorde.cpp
--- a/Engine/Source/Developer/UbaCoordinatorHorde/Private/UbaCoordinatorHorde.cpp
+++ b/Engine/Source/Developer/UbaCoordinatorHorde/Private/UbaCoordinatorHorde.cpp
@@ -47,6 +47,11 @@ public:
 		m_manager.SetTargetCoreCount(count);
 	}
 
+	virtual void SetUpdateStatusCallback(UpdateStatusCallback* callback, void* userData) override
+	{
+		m_manager.SetUpdateStatusCallback(callback, userData);
+	}
+
 	FUbaHordeAgentManager m_manager;
 };
 
diff --git a/Engine/Source/Programs/UnrealBuildAccelerator/Common/Public/UbaCoordinator.h b/Engine/Source/Programs/UnrealBuildAccelerator/Common/Public/UbaCoordinator.h
--- a/Engine/Source/Programs/UnrealBuildAccelerator/Common/Public/UbaCoordinator.h
+++ b/Engine/Source/Programs/UnrealBuildAccelerator/Common/Public/UbaCoordinator.h
@@ -22,6 +22,9 @@ namespace uba
 		using AddClientCallback = bool(void* userData, const tchar* ip, u16 port, const tchar* crypto);
 		virtual void SetAddClientCallback(AddClientCallback* callback, void* userData) = 0;
 		virtual void SetTargetCoreCount(u32 count) = 0;
+
+		using UpdateStatusCallback = void(void* userData, const TCHAR* status);
+		virtual void SetUpdateStatusCallback(UpdateStatusCallback* callback, void* userData) {}
 	};
 
 	// This is how the function signatures creating/destroying the coordinator module needs to exported
diff --git a/Engine/Source/Programs/UnrealBuildAccelerator/Common/Public/UbaCoordinatorWrapper.h b/Engine/Source/Programs/UnrealBuildAccelerator/Common/Public/UbaCoordinatorWrapper.h
--- a/Engine/Source/Programs/UnrealBuildAccelerator/Common/Public/UbaCoordinatorWrapper.h
+++ b/Engine/Source/Programs/UnrealBuildAccelerator/Common/Public/UbaCoordinatorWrapper.h
@@ -6,6 +6,7 @@
 #include "UbaFile.h"
 #include "UbaNetworkServer.h"
 #include "UbaScheduler.h"
+#include "UbaSessionServer.h"
 #include "UbaStringBuffer.h"
 
 #if !PLATFORM_WINDOWS
@@ -63,6 +64,9 @@ namespace uba
 			m_networkServer = &networkServer;
 			m_scheduler = scheduler;
 
+			if (m_scheduler)
+				m_scheduler->GetSession().GetTrace().StatusUpdate(0, 1, ToView(coordinatorType), LogEntryType_Info, {});
+
 			m_coordinatorThread.Start([this, mcc = info.maxCoreCount]() { ThreadUpdate(mcc); return 0; }, TC("UbaCoordWrap"));
 
 			return true;
@@ -70,6 +74,13 @@ namespace uba
 
 		void ThreadUpdate(u32 maxCoreCount)
 		{
+			m_coordinator->SetUpdateStatusCallback([](void* userData, const TCHAR* status)
+				{
+					auto& cw = *(CoordinatorWrapper*)userData;
+					if (cw.m_scheduler)
+						cw.m_scheduler->GetSession().GetTrace().StatusUpdate(0, 6, ToView(status), LogEntryType_Info, {});
+				}, this);
+
 			m_coordinator->SetAddClientCallback([](void* userData, const tchar* ip, u16 port, const tchar* crypto)
 				{
 					auto& cw = *(CoordinatorWrapper*)userData;
