[UBA]
* Fixed linux compile error in UbaCoordinator

diff --git a/Engine/Source/Programs/UnrealBuildAccelerator/Common/Public/UbaCoordinator.h b/Engine/Source/Programs/UnrealBuildAccelerator/Common/Public/UbaCoordinator.h
--- a/Engine/Source/Programs/UnrealBuildAccelerator/Common/Public/UbaCoordinator.h
+++ b/Engine/Source/Programs/UnrealBuildAccelerator/Common/Public/UbaCoordinator.h
@@ -23,7 +23,7 @@ namespace uba
 		virtual void SetAddClientCallback(AddClientCallback* callback, void* userData) = 0;
 		virtual void SetTargetCoreCount(u32 count) = 0;
 
-		using UpdateStatusCallback = void(void* userData, const TCHAR* status);
+		using UpdateStatusCallback = void(void* userData, const tchar* status);
 		virtual void SetUpdateStatusCallback(UpdateStatusCallback* callback, void* userData) {}
 	};
 
