[UBA]
* Fixed linux compile error in UbaCoordinatorWrapper

diff --git a/Engine/Source/Programs/UnrealBuildAccelerator/Common/Public/UbaCoordinatorWrapper.h b/Engine/Source/Programs/UnrealBuildAccelerator/Common/Public/UbaCoordinatorWrapper.h
--- a/Engine/Source/Programs/UnrealBuildAccelerator/Common/Public/UbaCoordinatorWrapper.h
+++ b/Engine/Source/Programs/UnrealBuildAccelerator/Common/Public/UbaCoordinatorWrapper.h
@@ -74,7 +74,7 @@ namespace uba
 
 		void ThreadUpdate(u32 maxCoreCount)
 		{
-			m_coordinator->SetUpdateStatusCallback([](void* userData, const TCHAR* status)
+			m_coordinator->SetUpdateStatusCallback([](void* userData, const tchar* status)
 				{
 					auto& cw = *(CoordinatorWrapper*)userData;
 					if (cw.m_scheduler)
