[UBA]
* Added UBA_FOR_ASSERT macro that can be used to put code inside that is only used with asserts

diff --git a/Engine/Source/Programs/UnrealBuildAccelerator/Common/Private/UbaCacheServer.cpp b/Engine/Source/Programs/UnrealBuildAccelerator/Common/Private/UbaCacheServer.cpp
--- a/Engine/Source/Programs/UnrealBuildAccelerator/Common/Private/UbaCacheServer.cpp
+++ b/Engine/Source/Programs/UnrealBuildAccelerator/Common/Private/UbaCacheServer.cpp
@@ -1368,8 +1368,7 @@ namespace uba
 					auto compressor = OodleLZ_Compressor_Kraken;
 					auto compressionLevel = OodleLZ_CompressionLevel_SuperFast;
 					u64 toCompress = Min(memLeft, u64(256*1024 - 128));
-					auto compressedBlockSize = OodleLZ_Compress(compressor, mem, (OO_SINTa)toCompress, buffer, compressionLevel);
-					(void)compressedBlockSize;
+					OodleLZ_Compress(compressor, mem, (OO_SINTa)toCompress, buffer, compressionLevel);
 					memLeft -= toCompress;
 				}
 				#endif
@@ -1880,8 +1879,8 @@ namespace uba
 		else
 		{
 			SCOPED_FUTEX(connectionBucket.deferredCacheEntryLookupLock, lock3);
-			bool res = connectionBucket.deferredCacheEntryLookup.try_emplace(cmdKey, std::move(newEntry)).second;
-			UBA_ASSERT(res);(void)res;
+			UBA_FOR_ASSERT(bool res =) connectionBucket.deferredCacheEntryLookup.try_emplace(cmdKey, std::move(newEntry)).second;
+			UBA_ASSERT(res);
 		}
 
 		//m_logger.Info(TC("Added new cache entry (%u inputs and %u outputs)"), u32(inputs.size()), outputCount);
diff --git a/Engine/Source/Programs/UnrealBuildAccelerator/Common/Private/UbaFileMapping.cpp b/Engine/Source/Programs/UnrealBuildAccelerator/Common/Private/UbaFileMapping.cpp
--- a/Engine/Source/Programs/UnrealBuildAccelerator/Common/Private/UbaFileMapping.cpp
+++ b/Engine/Source/Programs/UnrealBuildAccelerator/Common/Private/UbaFileMapping.cpp
@@ -44,7 +44,7 @@ namespace uba
 		if (res)
 		{
 			SCOPED_FUTEX(g_fileMappingsLock, lock);
-			auto insres = g_fileMappings.try_emplace(res, hint);(void)insres;
+			UBA_FOR_ASSERT(auto insres =) g_fileMappings.try_emplace(res, hint);
 			UBA_ASSERT(insres.second);
 			//logger.Info(TC("FILEMAPPING CREAT: %s %llu (%llu)"), hint, u64(res), g_fileMappings.size());
 		}
diff --git a/Engine/Source/Programs/UnrealBuildAccelerator/Core/Public/UbaPlatform.h b/Engine/Source/Programs/UnrealBuildAccelerator/Core/Public/UbaPlatform.h
--- a/Engine/Source/Programs/UnrealBuildAccelerator/Core/Public/UbaPlatform.h
+++ b/Engine/Source/Programs/UnrealBuildAccelerator/Core/Public/UbaPlatform.h
@@ -79,11 +79,13 @@ namespace uba
 	#define UBA_ASSERT(x) UBA_SHIPPING_ASSERT(x)
 	#define UBA_ASSERTF(x, fmt, ...) UBA_SHIPPING_ASSERTF(x, fmt, ##__VA_ARGS__)
 	#define UBA_ASSERT_MESSAGEBOX 0
+	#define UBA_FOR_ASSERT(...) __VA_ARGS__
 	inline constexpr bool IsDebug = true;
 	#else
 	#define UBA_ASSERT(x) do { } while(false)
 	#define UBA_ASSERTF(x, ...) do { } while(false)
 	#define UBA_ASSERT_MESSAGEBOX 0
+	#define UBA_FOR_ASSERT(...)
 	inline constexpr bool IsDebug = false;
 	#endif
 
