[UBA]
* Moved sleep in scheduler for when about to spawn new processes for command and noticing that there are processes already around that are requesting command

diff --git a/Engine/Source/Programs/UnrealBuildAccelerator/Common/Private/UbaScheduler.cpp b/Engine/Source/Programs/UnrealBuildAccelerator/Common/Private/UbaScheduler.cpp
--- a/Engine/Source/Programs/UnrealBuildAccelerator/Common/Private/UbaScheduler.cpp
+++ b/Engine/Source/Programs/UnrealBuildAccelerator/Common/Private/UbaScheduler.cpp
@@ -1100,7 +1100,6 @@ namespace uba
 				if (m_getNextQueueCount)
 				{
 					m_updateThreadWakeup.Set();
-					Sleep(1);
 					return ~0u;
 				}
 
@@ -1242,6 +1241,9 @@ namespace uba
 			if (indexToRun == ~0u)
 			{
 				lock.Leave();
+
+				Sleep(1); // We want to sleep a little bit in case this PopProcess failed because there is a process sitting on reuse request
+
 				if (!isLocal)
 					return false;
 				return RaceRemoteProcess();
