[UBA]
* Changed code in BinaryWriter::Flush that checks if parent process is alive to instead use True_NtQueryInformationProcess just to reduce the log spam when logging is enabled

diff --git a/Engine/Source/Programs/UnrealBuildAccelerator/Detours/Private/Windows/UbaDetoursMainWin.cpp b/Engine/Source/Programs/UnrealBuildAccelerator/Detours/Private/Windows/UbaDetoursMainWin.cpp
--- a/Engine/Source/Programs/UnrealBuildAccelerator/Detours/Private/Windows/UbaDetoursMainWin.cpp
+++ b/Engine/Source/Programs/UnrealBuildAccelerator/Detours/Private/Windows/UbaDetoursMainWin.cpp
@@ -134,11 +134,12 @@ namespace uba
 			if (res != WAIT_TIMEOUT)
 				TerminateCurrentProcess(1337);
 
-			DWORD exitCode;
-			if (!True_GetExitCodeProcess(g_hostProcess, &exitCode))
+			// Using True_NtQueryInformationProcess to get rid of spam in logs
+			PROCESS_BASIC_INFORMATION pbi;
+			if (True_NtQueryInformationProcess(g_hostProcess, ProcessBasicInformation, &pbi, sizeof(pbi), NULL) != STATUS_SUCCESS)
 				TerminateCurrentProcess(1338);
-
-			if (exitCode != STILL_ACTIVE)
+			auto exitCode = (NTSTATUS)(ULONG_PTR)pbi.Reserved1;
+			if (exitCode != STATUS_PENDING)
 				TerminateCurrentProcess(1339);
 
 			DWORD cancelRes = True_WaitForSingleObject(g_cancelEvent, 0);
