[UBA]
* Working on relying 100% on directory table to check if files exist. This is a large refactoring which is not fully implemented yet. Files created by local detoured processes might end up in main directory table even if they are just temporary. Key change is that NtCreateFile now always rely on the directory table. This fixes a lot of hacks and works better with complex process hierarchies.
* Fixed deadlock in racing code if cancelling build causing InternalRunProcess to exit with a failure
* Added unit tests

diff --git a/Engine/Source/Programs/UnrealBuildAccelerator/Common/Private/UbaDirectoryTableHolder.cpp b/Engine/Source/Programs/UnrealBuildAccelerator/Common/Private/UbaDirectoryTableHolder.cpp
--- a/Engine/Source/Programs/UnrealBuildAccelerator/Common/Private/UbaDirectoryTableHolder.cpp
+++ b/Engine/Source/Programs/UnrealBuildAccelerator/Common/Private/UbaDirectoryTableHolder.cpp
@@ -249,61 +249,55 @@ namespace uba
 
 	bool DirectoryTableHolder::PopulateOverlayFile(DirectoryTableOverlay& overlay, StringView fileName, const Function<bool(FileInformation&, StringKey)>& getFileInformation)
 	{
-		//auto& dirTable = m_directoryTable;
-
 		StringKey dirKey;
 		const tchar* lastSlash = nullptr;
 		StringBuffer<> fixedFileName;
 		StringKey fileNameKey = GetKeyAndFixedName(fixedFileName, fileName.data, &dirKey, &lastSlash);
-
 		FileInformation fileInfo;
 		if (!getFileInformation(fileInfo, fileNameKey))
 			return false;
+		return WriteOverlayEntry(overlay, dirKey, ToView(lastSlash + 1), fileInfo);
+	}
 
-		SCOPED_WRITE_LOCK(overlay.lock, lock);
+	bool DirectoryTableHolder::AddOverlayFile(DirectoryTableOverlay& overlay, StringView fileName)
+	{
+		FileInformation fi;
+		fi.size = InvalidValue; // TODO: This is very wrong
+		fi.attributes = DefaultAttributes();
+		fi.index = ++m_fileIndexCounter;
+		return WriteOverlayEntry(overlay, fileName, fi);
+	}
 
-		auto dirOffsetIns = overlay.lookup.try_emplace(dirKey);
-		
-		// TODO: we should use the normal dir table to decide how far up the tree we need to go to make sure we have an entire chain of directories
-		
-		/*
-		if (dirOffsetIns.second)
+	bool DirectoryTableHolder::MoveOverlayFile(DirectoryTableOverlay& overlay, StringView from, StringView to)
+	{
+		if (!WriteOverlayEntry(overlay, from, {}))
+			return false;
+		return AddOverlayFile(overlay, to);
+	}
+
+	bool DirectoryTableHolder::WriteOverlayEntry(DirectoryTableOverlay& overlay, StringView fileName, const FileInformation& fileInfo)
+	{
+		StringKey dirKey;
+		const tchar* lastSlash = nullptr;
+		StringBuffer<> fixedFileName;
+		StringKey fileNameKey = GetKeyAndFixedName(fixedFileName, fileName.data, &dirKey, &lastSlash);
+
+		// We need to make sure as far down as possible that the directories are created.. overlay entries must always somehow be attached to the end of the normal table
+		// TODO: This must ALSO be made for remotes
+		if (!m_runningRemote)
 		{
 			StringView dirName(fixedFileName.data, u32(lastSlash - fixedFileName.data));
+			WriteDirectoryEntries(dirKey, dirName);
+		}
 
-			SCOPED_READ_LOCK(dirTable.m_lookupLock, lookupLock);
-			auto findIt = dirTable.m_lookup.find(dirKey);
-			if (findIt == dirTable.m_lookup.end())
-			{
-				lookupLock.Leave();
-				u32 res = WriteDirectoryEntries(dirKey, dirName);
-				UBA_ASSERTF(res, TC("Failed to write directory entries for %.*s"), dirName.count, dirName.data); (void)res;
-				lookupLock.Enter();
-				findIt = dirTable.m_lookup.find(dirKey);
-				if (findIt == dirTable.m_lookup.end())
-					return m_logger.Error(TC("OverlayFiles can only exist in existing directories (%.*s)"), fixedFileName.count, fixedFileName.data);
-			}
-			//DirectoryTable::Directory& dir = findIt->second;
-			lookupLock.Leave();
+		return WriteOverlayEntry(overlay, dirKey, ToView(lastSlash + 1), fileInfo);
+	}
 
-			SCOPED_WRITE_LOCK(dir.lock, dirLock);
-			if (dir.populatedOffset != 3)
-			{
-				while (dir.populatedOffset == 0)
-				{
-					dirLock.Leave();
-					Sleep(1);
-					dirLock.Enter();
-				}
-				dir.populatedOffset = 3; // 3 means it is locked, it is now a directory owned by overlays
+	bool DirectoryTableHolder::WriteOverlayEntry(DirectoryTableOverlay& overlay, StringKey dirKey, StringView fileNameWithoutSlash, const FileInformation& fileInfo)
+	{
+		SCOPED_WRITE_LOCK(overlay.lock, lock);
 
-#if UBA_DEBUG_TRACK_DIR
-				m_debugLogger->Info(TC("TRACKOVERL %.*s"), dirName.count, dirName.data);
-#endif
-			}
-			dirOffsetIns.first->second = dir.latestOffset;
-		}
-		*/
+		auto dirOffsetIns = overlay.lookup.try_emplace(dirKey);
 
 		u8 temp[1024];
 		BinaryWriter writer(temp, 0, sizeof(temp));
@@ -311,12 +305,19 @@ namespace uba
 		writer.Write7BitEncoded(dirOffsetIns.first->second); // Previous entry for same directory
 		if (!dirOffsetIns.first->second)
 		{
-			writer.WriteFileAttributes(1);
+			u32 dirAttributes = 0;
+			#if PLATFORM_WINDOWS
+			dirAttributes = FILE_ATTRIBUTE_DIRECTORY;
+			#else
+			dirAttributes = S_IFDIR;
+			#endif
+
+			writer.WriteFileAttributes(dirAttributes);
 			writer.WriteVolumeSerial(fileInfo.volumeSerialNumber); // Reuse file info serial
 			writer.WriteFileIndex(++m_fileIndexCounter);
 		}
 		writer.Write7BitEncoded(1); // Count
-		writer.WriteString(lastSlash + 1);
+		writer.WriteString(fileNameWithoutSlash);
 		writer.WriteFileAttributes(fileInfo.attributes);
 		writer.WriteVolumeSerial(fileInfo.volumeSerialNumber);
 		writer.WriteFileIndex(fileInfo.index);
@@ -329,7 +330,7 @@ namespace uba
 
 		BinaryWriter memWriter(overlay.memory, overlay.size, overlay.committed);
 		memWriter.Write7BitEncoded(written); // Storage size
-		dirOffsetIns.first->second = u32(memWriter.GetPosition()) + sizeof(StringKey) + OverlayTableOffset;
+		dirOffsetIns.first->second = u32(memWriter.GetPosition()) + sizeof(StringKey);
 		memWriter.WriteBytes(writer.GetData(), written);
 		overlay.size = u32(memWriter.GetPosition());
 		return true;
@@ -337,10 +338,6 @@ namespace uba
 
 	bool DirectoryTableHolder::RegisterCreateFileForWrite(StringKey fileNameKey, StringView fileName, DirectoryTableOverlay* overlay, u64 fileSize, u64 lastWriteTime, bool registerRealFile, bool invalidateStorage)
 	{
-		// Remote is not updating its own directory table
-		if (m_runningRemote)
-			return true;
-
 		auto& dirTable = m_directoryTable;
 
 		StringKey dirKey;
@@ -349,6 +346,22 @@ namespace uba
 		if (!GetDirKey(dirKey, dirName, lastSlash, fileName))
 			return true;
 
+		// Remote is not updating its own directory table
+		if (m_runningRemote)
+		{
+			if (overlay)
+			{
+				FileInformation info;
+				info.attributes = DefaultAttributes();
+				info.index = ++m_fileIndexCounter;
+				info.volumeSerialNumber = 1;
+				info.size = fileSize;
+				info.lastWriteTime = lastWriteTime;
+				WriteOverlayEntry(*overlay, fileName, info);
+			}
+			return true;
+		}
+
 #if 0//_DEBUG  // Bring this back, turned off right now because a few lines above the call to this method we add a mapping
 		{
 			SCOPED_FUTEX(m_fileMappingTableLookupLock, lookupLock);
@@ -368,8 +381,8 @@ namespace uba
 		// When not writing to disk we need to populate lookup before adding non-written files.. otherwise they will be lost once lookup is actually populated
 		if (!shouldWriteToDisk)
 		{
-			u32 res = WriteDirectoryEntries(dirKey, dirName);
-			UBA_ASSERTF(res, TC("Failed to write directory entries for %s"), dirName.data); (void)res;
+			UBA_FOR_ASSERT(u32 res =) WriteDirectoryEntries(dirKey, dirName);
+			UBA_ASSERTF(res, TC("Failed to write directory entries for %s"), dirName.data);
 		}
 
 		SCOPED_READ_LOCK(dirTable.m_lookupLock, lookupCs);
@@ -395,8 +408,8 @@ namespace uba
 		if (dir.populatedOffset == 2)
 		{
 			dirLock.Leave();
-			u32 res = WriteDirectoryEntries(dirKey, dirName);
-			UBA_ASSERT(res); (void)res;
+			UBA_FOR_ASSERT(u32 res =) WriteDirectoryEntries(dirKey, dirName);
+			UBA_ASSERT(res);
 			dirLock.Enter();
 		}
 		UBA_ASSERTF(dir.populatedOffset == 1 || (!ShouldWriteToDisk(fileName) && dir.populatedOffset != 3), TC("Registering create file for write %s with unexpect dir.populatedOffset %u "), fileName.data, dir.populatedOffset);
@@ -477,6 +490,7 @@ namespace uba
 
 		FileEntryAdded(fileNameKey, lastWriteTime, fileSize);
 
+		u32 volumeIndex = m_volumeCache.GetSerialIndex(volumeSerial);
 		u8 temp[1024];
 		u64 written;
 		u64 entryPos;
@@ -488,7 +502,7 @@ namespace uba
 			writer.WriteString(lastSlash + 1);
 			entryPos = writer.GetPosition();
 			writer.WriteFileAttributes(attributes);
-			writer.WriteVolumeSerial(m_volumeCache.GetSerialIndex(volumeSerial));
+			writer.WriteVolumeSerial(volumeIndex);
 			writer.WriteFileIndex(fileIndex);
 			if (!isDirectory)
 			{
@@ -502,6 +516,18 @@ namespace uba
 		m_debugLogger->Info(TC("TRACKADD    %s (Size: %llu, Attr: %u, Key: %s, Id: %llu)"), fileName.data, fileSize, attributes, KeyToString(fileNameKey).data, fileIndex);
 #endif
 
+		if (overlay && !isDirectory)
+		{
+			FileInformation info;
+			info.attributes = attributes;
+			info.index = fileIndex;
+			info.lastWriteTime = lastWriteTime;
+			info.size = fileSize;
+			info.volumeSerialNumber = volumeIndex;
+			if (!WriteOverlayEntry(*overlay, dirKey, ToView(lastSlash + 1), info))
+				return false;
+		}
+
 
 		SCOPED_WRITE_LOCK(dirTable.m_memoryLock, memoryLock);
 		u32 writePos = dirTable.m_memorySize;
@@ -516,24 +542,31 @@ namespace uba
 		return true;
 	}
 
-	DirTableSize DirectoryTableHolder::RegisterDeleteFile(StringKey fileNameKey, StringView fileName, DirectoryTableOverlay* overlay)
+	bool DirectoryTableHolder::RegisterDeleteFile(StringKey fileNameKey, StringView fileName, DirectoryTableOverlay* overlay)
 	{
-		// Remote is not updating its own directory table
-		if (m_runningRemote)
-			return GetDirectoryTableSize(overlay);
-
-		auto& dirTable = m_directoryTable;
-
 		StringKey dirKey;
 		const tchar* lastSlash;
 		StringBuffer<> dirName;
 		if (!GetDirKey(dirKey, dirName, lastSlash, fileName))
-			return {InvalidTableOffset, InvalidTableOffset};
+			return false;
+
+		if (overlay)
+		{
+			// TODO: If it does not exist in overlay nor main it should not create any entry
+			FileInformation fileInfo;
+			if (!WriteOverlayEntry(*overlay, dirKey, ToView(lastSlash+1), fileInfo))
+				return false;
+		}
+
+		// Remote is not updating its own directory table
+		if (m_runningRemote)
+			return true;
 
+		auto& dirTable = m_directoryTable;
 		SCOPED_READ_LOCK(dirTable.m_lookupLock, lookupCs);
 		auto res = dirTable.m_lookup.find(dirKey);
 		if (res == dirTable.m_lookup.end())
-			return {};
+			return true;
 		DirectoryTable::Directory& dir = res->second;
 		lookupCs.Leave();
 		SCOPED_WRITE_LOCK(dir.lock, dirLock);
@@ -546,7 +579,7 @@ namespace uba
 		}
 
 		if (dir.populatedOffset == 2) // Invalid directory, don't need to register this
-			return {};
+			return true;
 
 		UBA_ASSERTF(dir.populatedOffset == 1, TC("Registering deleted file %s with unexpect dir.populatedOffset %u "), fileName.data, dir.populatedOffset);
 
@@ -561,7 +594,7 @@ namespace uba
 
 		// Does not exist, no need adding to file table
 		if (dir.files.erase(fileNameKey) == 0)
-			return {};
+			return true;
 
 		u8 temp[1024];
 		u64 written;
@@ -597,12 +630,12 @@ namespace uba
 			dir.latestOffset = dirTable.m_memorySize + tableOffset;
 			dirTable.m_memorySize += u32(writer.GetPosition());
 		}
-
-		return GetDirectoryTableSize(overlay);
+		return true;
 	}
 
 	u32 DirectoryTableHolder::WriteDirectoryEntries(const StringKey& dirKey, StringView dirPath, u32* outTableOffset)
 	{
+		UBA_ASSERT(!m_runningRemote);
 		auto& dirTable = m_directoryTable;
 		u32 temp;
 		if (!outTableOffset)
diff --git a/Engine/Source/Programs/UnrealBuildAccelerator/Common/Private/UbaProcess.cpp b/Engine/Source/Programs/UnrealBuildAccelerator/Common/Private/UbaProcess.cpp
--- a/Engine/Source/Programs/UnrealBuildAccelerator/Common/Private/UbaProcess.cpp
+++ b/Engine/Source/Programs/UnrealBuildAccelerator/Common/Private/UbaProcess.cpp
@@ -1059,6 +1059,16 @@ namespace uba
 		return true;
 	}
 
+	bool ProcessImpl::HandleRegisterFileForWrite(BinaryReader& reader, BinaryWriter& writer)
+	{
+		StringBuffer<> fileName;
+		reader.ReadString(fileName);
+		StringKey fileNameKey = reader.ReadStringKey();
+		m_session.AddOverlayFile(m_shared.overlay, fileName);
+		writer.WriteU64(ToU64(m_session.GetDirectoryTableSize(&m_shared.overlay)));
+		return true;
+	}
+
 	bool ProcessImpl::HandleGetFullFileName(BinaryReader& reader, BinaryWriter& writer)
 	{
 		GetFullFileNameMessage msg { *this };
@@ -1201,7 +1211,7 @@ namespace uba
 		reader.ReadString(msg.directoryName);
 		msg.directoryNameKey = reader.ReadStringKey();
 		ListDirectoryResponse response;
-		CallSession(MessageType_ListDirectory, m_session.GetListDirectoryInfo(response, msg.directoryName, msg.directoryNameKey));
+		CallSession(MessageType_ListDirectory, m_session.GetListDirectoryInfo(response, msg.directoryName, msg.directoryNameKey, &m_shared.overlay));
 		writer.WriteU64(ToU64(response.directoryTableSize));
 		writer.WriteU32(response.tableOffset);
 		return true;
diff --git a/Engine/Source/Programs/UnrealBuildAccelerator/Common/Private/UbaSession.cpp b/Engine/Source/Programs/UnrealBuildAccelerator/Common/Private/UbaSession.cpp
--- a/Engine/Source/Programs/UnrealBuildAccelerator/Common/Private/UbaSession.cpp
+++ b/Engine/Source/Programs/UnrealBuildAccelerator/Common/Private/UbaSession.cpp
@@ -2785,19 +2785,35 @@ namespace uba
 					return false;
 		}
 
+		//
+		// This entire code below here needs to be revisited
+		// We want to only use overlay for everything here.
+		// And then transfer info back to the main directory table once things are done
+		//
+
+		if (m_runningRemote && registerRealFile && msg.fileName.count && !msg.deleteOnClose)
+		{
+			FileBasicInformation info;
+			if (GetFileBasicInformation(info, m_logger, msg.fileName.data))
+			{
+				lastWriteTime = info.lastWriteTime;
+				fileSize = info.size;
+			}
+		}
+
 		if (!rules.IsInvisible(activeFile.name))
 		{
 			if (!msg.newName.IsEmpty())
 			{
 				RegisterDeleteFile(activeFile.nameKey, activeFile.name, &shared.overlay);
-				if (RegisterCreateFileForWrite(msg.newNameKey, msg.newName, nullptr, fileSize, lastWriteTime, registerRealFile, true) && registerRealFile)
+				if (RegisterCreateFileForWrite(msg.newNameKey, msg.newName, &shared.overlay, fileSize, lastWriteTime, registerRealFile, true) && registerRealFile)
 					TraceWrittenFile(msg.process.m_id, msg.newName, fileSize);
 			}
 			else if (msg.deleteOnClose)
 				RegisterDeleteFile(activeFile.nameKey, activeFile.name, &shared.overlay);
 			else
 			{
-				if (RegisterCreateFileForWrite(activeFile.nameKey, activeFile.name, nullptr, fileSize, lastWriteTime, registerRealFile, true) && registerRealFile)
+				if (RegisterCreateFileForWrite(activeFile.nameKey, activeFile.name, &shared.overlay, fileSize, lastWriteTime, registerRealFile, true) && registerRealFile)
 					TraceWrittenFile(msg.process.m_id, activeFile.name, fileSize);
 			}
 		}
@@ -2844,7 +2860,10 @@ namespace uba
 		if (deleteRealFile)
 			out.result = uba::DeleteFileW(msg.fileName.data);
 		out.errorCode = GetLastError();
-		out.directoryTableSize = RegisterDeleteFile(msg.fileNameKey, msg.fileName, &shared.overlay);
+
+		RegisterDeleteFile(msg.fileNameKey, msg.fileName, &shared.overlay);
+
+		out.directoryTableSize = GetDirectoryTableSize(&msg.process.m_shared.overlay);
 		return true;
 	}
 
@@ -2900,6 +2919,10 @@ namespace uba
 				{
 					out.errorCode = ERROR_SUCCESS;
 					out.result = true;
+					
+					if (!MoveOverlayFile(msg.process.m_shared.overlay, msg.fromName, msg.toName))
+						return false;
+					out.directoryTableSize = GetDirectoryTableSize(&msg.process.m_shared.overlay);
 					return true;
 				}
 			}
@@ -2919,8 +2942,13 @@ namespace uba
 			if (RegisterCreateFileForWrite(msg.toKey, msg.toName, nullptr, 0, 0, true, true))
 				TraceWrittenFile(process.m_id, msg.toName, 0);
 
+		if (!AddOverlayFile(msg.process.m_shared.overlay, msg.toName))
+			return false;
+
 		out.errorCode = ERROR_SUCCESS;
-		out.directoryTableSize = RegisterDeleteFile(msg.fromKey, msg.fromName, &shared.overlay);
+		RegisterDeleteFile(msg.fromKey, msg.fromName, &shared.overlay);
+
+		out.directoryTableSize = GetDirectoryTableSize(&msg.process.m_shared.overlay);
 		return true;
 	}
 
@@ -2991,12 +3019,19 @@ namespace uba
 		return false;
 	}
 
-	bool Session::GetListDirectoryInfo(ListDirectoryResponse& out, const StringView& dirName, const StringKey& dirKey)
+	bool Session::GetListDirectoryInfo(ListDirectoryResponse& out, const StringView& dirName, const StringKey& dirKey, DirectoryTableOverlay* overlay)
 	{
 		u32 tableOffset;
 		u32 tableSize = WriteDirectoryEntries(dirKey, dirName, &tableOffset);
 		out.tableOffset = tableOffset;
-		out.directoryTableSize = {tableSize, 0u};
+		u32 overlaySize = 0;
+		if (overlay)
+		{
+			SCOPED_READ_LOCK(overlay->lock, lock);
+			overlaySize = overlay->size;
+		}
+
+		out.directoryTableSize = {tableSize, overlaySize};
 		return true;
 	}
 
diff --git a/Engine/Source/Programs/UnrealBuildAccelerator/Common/Private/UbaSessionClient.cpp b/Engine/Source/Programs/UnrealBuildAccelerator/Common/Private/UbaSessionClient.cpp
--- a/Engine/Source/Programs/UnrealBuildAccelerator/Common/Private/UbaSessionClient.cpp
+++ b/Engine/Source/Programs/UnrealBuildAccelerator/Common/Private/UbaSessionClient.cpp
@@ -906,10 +906,11 @@ namespace uba
 		const tchar* fromName = msg.fromName.data;
 		const tchar* toName = msg.toName.data;
 		auto& process = msg.process;
+		auto& shared = process.m_shared;
 
 		{
-			SCOPED_FUTEX(process.m_shared.writtenFilesLock, lock);
-			auto& writtenFiles = process.m_shared.writtenFiles;
+			SCOPED_FUTEX(shared.writtenFilesLock, lock);
+			auto& writtenFiles = shared.writtenFiles;
 			auto findIt = writtenFiles.find(msg.fromKey);
 			if (findIt != writtenFiles.end())
 			{
@@ -935,6 +936,9 @@ namespace uba
 			}
 		}
 
+		MoveOverlayFile(shared.overlay, msg.fromName, msg.toName);
+		out.directoryTableSize = GetDirectoryTableSize(&shared.overlay);
+
 		bool sendMove = true;
 		{
 			SCOPED_FUTEX(m_outputFilesLock, lock);
@@ -1079,7 +1083,7 @@ namespace uba
 		return true;
 	}
 
-	bool SessionClient::GetListDirectoryInfo(ListDirectoryResponse& out, const StringView& dirName, const StringKey& dirKey)
+	bool SessionClient::GetListDirectoryInfo(ListDirectoryResponse& out, const StringView& dirName, const StringKey& dirKey, DirectoryTableOverlay* overlay)
 	{
 		TrackWorkScope tws(m_client, AsView(TC("GetListDir")), ColorWork);
 		StackBinaryWriter<1024> writer;
@@ -1098,15 +1102,14 @@ namespace uba
 		u32 oldTableSize = GetDirectoryTableSize(nullptr).main;
 		if (!UpdateDirectoryTableFromServer(reader))
 			return false;
-		u32 newTableSize = GetDirectoryTableSize(nullptr).main;
+		DirTableSize newTableSize = GetDirectoryTableSize(overlay);
 
 		// Ask for a refresh of hashes straight away since they will likely be asked for by the process doing this query
-		if (oldTableSize != newTableSize)
+		if (oldTableSize != newTableSize.main)
 			m_waitToSendEvent.Set();
 
 		out.tableOffset = tableOffset;
-		out.directoryTableSize = {newTableSize, 0u};
-
+		out.directoryTableSize = newTableSize;
 		return true;
 	}
 
@@ -2432,7 +2435,7 @@ namespace uba
 					if (!entry.handled)
 					{
 						ListDirectoryResponse out;
-						if (!GetListDirectoryInfo(out, dirName, dirKey))
+						if (!GetListDirectoryInfo(out, dirName, dirKey, nullptr))
 							return false;
 						ParseDirectoryTable();
 						entry.handled = true;
diff --git a/Engine/Source/Programs/UnrealBuildAccelerator/Common/Private/UbaSessionServer.cpp b/Engine/Source/Programs/UnrealBuildAccelerator/Common/Private/UbaSessionServer.cpp
--- a/Engine/Source/Programs/UnrealBuildAccelerator/Common/Private/UbaSessionServer.cpp
+++ b/Engine/Source/Programs/UnrealBuildAccelerator/Common/Private/UbaSessionServer.cpp
@@ -352,23 +352,6 @@ namespace uba
 			aligned_free(s);
 		}
 		m_clientSessions.clear();
-
-		#if 0
-		for (auto& kv : m_directoryTable.m_lookup)
-		{
-			DirectoryTable::Directory& dir = kv.second;
-
-			DirectoryTable::EntryLookup files(m_directoryTable.m_memoryBlock);
-			m_directoryTable.PopulateDirectoryRecursive(StringKeyHasher(), dir.tableOffset, 0, files);
-			for (auto& fileKv : files)
-			{
-				BinaryReader reader(m_directoryTable.m_memory, fileKv.second);
-				StringBuffer<> filename;
-				reader.ReadString(filename);
-				m_logger.Info(filename.data);
-			}
-		}
-		#endif
 	}
 
 	ProcessHandle SessionServer::RunProcessRacing(const ProcessToRaceCallback& callback)
@@ -421,30 +404,39 @@ namespace uba
 		psi.exitedFunc = [](void* userData, const ProcessHandle& h, ProcessExitedResponse& r)
 			{
 				auto& rpi = *(RacingProcessExitInfo*)userData;
-				rpi.allowedToExit.IsSet();
+				auto& racingProcess = *(ProcessImpl*)h.m_process;
+
+				auto exitGuard = MakeGuard([&]()
+					{
+						rpi.exitedFunc(racingProcess.m_startInfo.userData, h, r);
+						delete &rpi;
+					});
+
+				// We spawn the process async, and if we get here before m_messageThread is set that means
+				if (!racingProcess.m_messageThread)
+					return;
 
+				rpi.allowedToExit.IsSet();
 				auto& remoteProcess = *(RemoteProcess*)rpi.remoteProcess.m_process;
-				auto& racingProcess = *(ProcessImpl*)h.m_process;
 
 				if (racingProcess.IsCancelled())
 				{
 					remoteProcess.m_racingProcess = {};
+					return;
 				}
-				else
-				{
-					// We got here before remote process.. need to make sure process exit callback is called first
-					rpi.server->OnRaceLost(remoteProcess);
-					remoteProcess.m_exitCode = ProcessCancelExitCode;
-					remoteProcess.m_done.Set();
-					remoteProcess.CallProcessExit(rpi.remoteProcess);
-				}
-
-				rpi.exitedFunc(racingProcess.m_startInfo.userData, h, r);
 
-				delete &rpi;
+				// We got here before remote process.. need to make sure process exit callback is called first
+				rpi.server->OnRaceLost(remoteProcess);
+				remoteProcess.m_exitCode = ProcessCancelExitCode;
+				remoteProcess.m_done.Set();
+				remoteProcess.CallProcessExit(rpi.remoteProcess);
 			};
 
+		// This function can call directly into psi.exitedFunc if it fails in certain ways (
+
 		ProcessHandle racingProcess = InternalRunProcess(psi, true, nullptr, m_allowLocalDetour, rpi);
+		if (!racingProcess.IsValid())
+			return {};
 
 		m_trace.ProcessRace(racingProcess.GetId(), newest->GetId());
 
@@ -1660,7 +1652,7 @@ namespace uba
 		reader.ReadString(dirName);
 		StringKey dirKey = reader.ReadStringKey();
 		ListDirectoryResponse out;
-		GetListDirectoryInfo(out, dirName, dirKey);
+		GetListDirectoryInfo(out, dirName, dirKey, nullptr);
 		writer.WriteU32(out.tableOffset);
 		WriteDirectoryTable(session, reader, writer);
 		return true;
diff --git a/Engine/Source/Programs/UnrealBuildAccelerator/Common/Public/UbaDirectoryTableHolder.h b/Engine/Source/Programs/UnrealBuildAccelerator/Common/Public/UbaDirectoryTableHolder.h
--- a/Engine/Source/Programs/UnrealBuildAccelerator/Common/Public/UbaDirectoryTableHolder.h
+++ b/Engine/Source/Programs/UnrealBuildAccelerator/Common/Public/UbaDirectoryTableHolder.h
@@ -49,9 +49,13 @@ namespace uba
 	//protected:
 
 		bool PopulateOverlayFile(DirectoryTableOverlay& overlay, StringView fileName, const Function<bool(FileInformation& out, StringKey)>& getFileInformation);
+		bool AddOverlayFile(DirectoryTableOverlay& overlay, StringView fileName);
+		bool MoveOverlayFile(DirectoryTableOverlay& overlay, StringView from, StringView to);
+		bool WriteOverlayEntry(DirectoryTableOverlay& overlay, StringView fileName, const FileInformation& fileInfo);
+		bool WriteOverlayEntry(DirectoryTableOverlay& overlay, StringKey dirKey, StringView fileNameWithoutSlash, const FileInformation& fileInfo);
 
 		bool RegisterCreateFileForWrite(StringKey fileNameKey, StringView fileName, DirectoryTableOverlay* overlay, u64 fileSize, u64 lastWriteTime, bool registerRealFile, bool invalidateStorage);
-		DirTableSize RegisterDeleteFile(StringKey fileNameKey, StringView fileName, DirectoryTableOverlay* overlay);
+		bool RegisterDeleteFile(StringKey fileNameKey, StringView fileName, DirectoryTableOverlay* overlay);
 
 		u32 WriteDirectoryEntries(const StringKey& dirKey, StringView dirPath, u32* outTableOffset = nullptr);
 		void WriteDirectoryEntriesRecursive(const StringKey& dirKey, StringView dirPath, u32& outTableOffset);
diff --git a/Engine/Source/Programs/UnrealBuildAccelerator/Common/Public/UbaFile.h b/Engine/Source/Programs/UnrealBuildAccelerator/Common/Public/UbaFile.h
--- a/Engine/Source/Programs/UnrealBuildAccelerator/Common/Public/UbaFile.h
+++ b/Engine/Source/Programs/UnrealBuildAccelerator/Common/Public/UbaFile.h
@@ -46,20 +46,20 @@ namespace uba
 
 	struct FileBasicInformation
 	{
-		u32 attributes;
-		u64 lastWriteTime;
-		u64 size;
+		u32 attributes = 0;
+		u64 lastWriteTime = 0;
+		u64 size = 0;
 	};
 	bool GetFileBasicInformationByHandle(FileBasicInformation& out, Logger& logger, const tchar* fileName, FileHandle hFile, bool errorOnFail = true);
 	bool GetFileBasicInformation(FileBasicInformation& out, Logger& logger, const tchar* fileName, bool errorOnFail = true);
 
 	struct FileInformation
 	{
-		u32 attributes;
-		u32 volumeSerialNumber;
-		u64 lastWriteTime;
-		u64 size;
-		u64 index;
+		u32 attributes = 0;
+		u32 volumeSerialNumber = 0;
+		u64 lastWriteTime = 0;
+		u64 size = 0;
+		u64 index = 0;
 	};
 	bool GetFileInformationByHandle(FileInformation& out, Logger& logger, const tchar* fileName, FileHandle hFile);
 	bool GetFileInformation(FileInformation& out, Logger& logger, const tchar* fileName);
diff --git a/Engine/Source/Programs/UnrealBuildAccelerator/Common/Public/UbaSession.h b/Engine/Source/Programs/UnrealBuildAccelerator/Common/Public/UbaSession.h
--- a/Engine/Source/Programs/UnrealBuildAccelerator/Common/Public/UbaSession.h
+++ b/Engine/Source/Programs/UnrealBuildAccelerator/Common/Public/UbaSession.h
@@ -177,7 +177,7 @@ namespace uba
 		virtual bool RemoveDirectory(RemoveDirectoryResponse& out, const RemoveDirectoryMessage& msg, DirectoryTableOverlay* overlay);
 		virtual bool GetFullFileName(GetFullFileNameResponse& out, const GetFullFileNameMessage& msg);
 		virtual bool GetLongPathName(GetLongPathNameResponse& out, const GetLongPathNameMessage& msg);
-		virtual bool GetListDirectoryInfo(ListDirectoryResponse& out, const StringView& dirName, const StringKey& dirKey);
+		virtual bool GetListDirectoryInfo(ListDirectoryResponse& out, const StringView& dirName, const StringKey& dirKey, DirectoryTableOverlay* overlay);
 		virtual bool WriteFilesToDisk(ProcessImpl& process, WrittenFile** files, u32 fileCount);
 		virtual bool AllocFailed(Process& process, const tchar* allocType, u32 error);
 		virtual bool GetNextProcess(Process& process, bool& outNewProcess, bool& outShouldExit, NextProcessInfo& outNextProcess, u32 prevExitCode, BinaryReader& prevStats, u32 timeoutMs);
diff --git a/Engine/Source/Programs/UnrealBuildAccelerator/Common/Public/UbaSessionClient.h b/Engine/Source/Programs/UnrealBuildAccelerator/Common/Public/UbaSessionClient.h
--- a/Engine/Source/Programs/UnrealBuildAccelerator/Common/Public/UbaSessionClient.h
+++ b/Engine/Source/Programs/UnrealBuildAccelerator/Common/Public/UbaSessionClient.h
@@ -69,7 +69,7 @@ namespace uba
 		virtual bool RemoveDirectory(RemoveDirectoryResponse& out, const RemoveDirectoryMessage& msg, DirectoryTableOverlay* overlay) override;
 		virtual bool GetFullFileName(GetFullFileNameResponse& out, const GetFullFileNameMessage& msg) override;
 		virtual bool GetLongPathName(GetLongPathNameResponse& out, const GetLongPathNameMessage& msg) override;
-		virtual bool GetListDirectoryInfo(ListDirectoryResponse& out, const StringView& dirName, const StringKey& dirKey) override;
+		virtual bool GetListDirectoryInfo(ListDirectoryResponse& out, const StringView& dirName, const StringKey& dirKey, DirectoryTableOverlay* overlay) override;
 		virtual bool WriteFilesToDisk(ProcessImpl& process, WrittenFile** files, u32 fileCount) override;
 		virtual bool AllocFailed(Process& process, const tchar* allocType, u32 error) override;
 		virtual void PrintSessionStats(Logger& logger) override;
diff --git a/Engine/Source/Programs/UnrealBuildAccelerator/Core/Public/UbaDirectoryTable.h b/Engine/Source/Programs/UnrealBuildAccelerator/Core/Public/UbaDirectoryTable.h
--- a/Engine/Source/Programs/UnrealBuildAccelerator/Core/Public/UbaDirectoryTable.h
+++ b/Engine/Source/Programs/UnrealBuildAccelerator/Core/Public/UbaDirectoryTable.h
@@ -5,6 +5,17 @@
 #include "UbaPathUtils.h"
 #include "UbaBinaryReaderWriter.h"
 
+#if 0
+#define WRITE_DEBUG(format, ...) \
+{ \
+	static int debugFile = open("/dev/tty", O_WRONLY); \
+	StringBuffer<> str; \
+	str.Appendf(format "\n", __VA_ARGS__); \
+	write(debugFile, str.data, str.count); \
+}
+#endif
+
+
 namespace uba
 {
 	////////////////////////////////////////////////////////////////////////////////////////////////////
@@ -35,7 +46,7 @@ namespace uba
 	////////////////////////////////////////////////////////////////////////////////////////////////////
 
 	enum : u64 { InvalidTableOffset = 0 }; // Important that it is zero to compress enc7 as good as possible
-	enum : u64 { OverlayTableOffset = 256*1024*1024 };
+	enum : u32 { OverlayTableFlag = 0x10000000 };
 	enum : u32 { DirectoryEntryFlag = 0x80000000 };
 
 	struct DirTableSize { u32 main = 0; u32 overlay = 0; };
@@ -45,7 +56,7 @@ namespace uba
 	struct DirTableOffset { u32 internal = InvalidTableOffset; };
 	inline bool IsDirectoryEntry(DirTableOffset d) { return (d.internal & DirectoryEntryFlag) != 0; }
 	inline bool IsValidEntry(DirTableOffset d) { return d.internal != InvalidTableOffset; }
-
+	inline bool IsOverlayOffset(u32 offset) { return (offset & OverlayTableFlag) != 0; }
 	////////////////////////////////////////////////////////////////////////////////////////////////////
 
 	class DirectoryTable
@@ -70,13 +81,13 @@ namespace uba
 		inline void DeinitOverlay(const u8* oldMem, u32 size);
 		inline void ParseDirectoryTable(DirTableSize size);
 		inline void ParseDirectoryTableNoLock(DirTableSize size);
-		inline void ParseDirectoryTableNoLock(const u8* memory, u32 from, u32 to, bool isOverlay);
+		inline void ParseDirectoryTableNoLock(const u8* memory, u32 from, u32 to, u32 isOverlayFlag);
 		inline void PopulateDirectory(const StringKeyHasher& hasher, Directory& dir);
 		inline bool IsDirectoryUpToDate(Directory& dir);
 		inline void PopulateDirectoryNoLock(const StringKeyHasher& hasher, Directory& dir);
 		inline BinaryReader GetReader(u32 tableOffset);
-		inline void PopulateDirectoryRecursive(const StringKeyHasher& hasher, u32 latestOffset, u32 populatedOffset, EntryLookup& files);
-		inline void PopulateDirectoryWithFiles(BinaryReader& reader, const StringKeyHasher& hasher, EntryLookup& files);
+		inline void PopulateDirectoryRecursive(const StringKeyHasher& hasher, u32 latestOffset, u32 populatedOffset, EntryLookup& files, u32 isOverlayFlag);
+		inline void PopulateDirectoryWithFiles(BinaryReader& reader, const StringKeyHasher& hasher, EntryLookup& files, u32 isOverlayFlag);
 
 		enum Exists
 		{
@@ -90,7 +101,7 @@ namespace uba
 		// NOTE, even if checkIfDir is true you can still get back a normal "file entry" representing the dir. "IsDirectoryEntry(tableOffset)" can be used to check which one it is
 		inline Exists EntryExists(StringKey entryKey, StringView entryName, bool checkIfDir = false, DirTableOffset* outTableOffset = nullptr);
 		inline Exists EntryExistsNoLock(StringKey entryKey, StringView entryName, bool checkIfDir = false, DirTableOffset* outTableOffset = nullptr);
-		inline Exists EntryExists(StringView str, bool checkIfDir = false);
+		inline Exists EntryExists(StringView str, bool checkIfDir = false, DirTableOffset* outTableOffset = nullptr);
 
 		struct EntryInformation
 		{
@@ -105,16 +116,15 @@ namespace uba
 		inline u32 GetEntryInformation(EntryInformation& outInfo, DirTableOffset tableOffset, tchar* outFileName = nullptr, u32 fileNameCapacity = 0);
 		inline void GetFinalPath(StringBufferBase& out, const tchar* path);
 		inline bool GetLatestOffset(DirTableOffset& out, const Directory& dir);
+		inline bool ContainsFile(const StringKeyHasher& hasher, Directory& dir, StringKey fileNameKey);
 
 		// This function only traverses the entries that was in the original populate and does not populate the lookup table.
 		// Can be used for scenarios where we know we only want the files that existed before the session started and not files that were added while active
 		template<typename Func> inline void TraverseOriginalFilesNoPopulate(DirTableOffset tableOffset, const Func& func);
 
-		#if PLATFORM_WINDOWS
 		// Func is (const DirectoryTable::EntryInformation& info, StringView fileName, DirTableOffset fileOffset)
 		template<typename Func> inline void TraverseFilesRecursiveNoLock(StringView path, const Func& func);
 		template<typename Func> inline void TraverseAllFilesNoLock(const Func& func);
-		#endif
 
 		inline DirectoryTable(MemoryBlock& block);
 		DirectoryTable(const DirectoryTable&) = delete;
@@ -144,8 +154,9 @@ namespace uba
 
 	void DirectoryTable::InitOverlay(const u8* mem, u32 tableSize)
 	{
+		UBA_ASSERT(u64(mem) > OverlayTableFlag);
 		m_overlayMemory = mem;
-		ParseDirectoryTableNoLock(mem, 0, tableSize, true);
+		ParseDirectoryTableNoLock(mem, 0, tableSize, OverlayTableFlag);
 		m_overlaySize = tableSize;
 	}
 
@@ -158,7 +169,7 @@ namespace uba
 		for (auto i=m_lookup.begin(), e=m_lookup.end(); i!=e;)
 		{
 			auto& dir = i->second;
-			if (dir.latestOffset < OverlayTableOffset)
+			if (!IsOverlayOffset(dir.latestOffset))
 			{
 				++i;
 				continue;
@@ -168,7 +179,7 @@ namespace uba
 			{
 				BinaryReader reader(GetReader(dir.latestOffset));
 				dir.latestOffset = u32(reader.Read7BitEncoded());
-				if (dir.latestOffset < OverlayTableOffset)
+				if (!IsOverlayOffset(dir.latestOffset))
 					break;
 			}
 			if (dir.latestOffset == InvalidTableOffset)
@@ -181,7 +192,7 @@ namespace uba
 
 			for (auto j=dir.files.begin(); j!=dir.files.end();)
 			{
-				if (j->second.internal < OverlayTableOffset)
+				if (!IsOverlayOffset(j->second.internal))
 					++j;
 				else
 					j = dir.files.erase(j);
@@ -203,19 +214,18 @@ namespace uba
 	{
 		if (size.main > m_memorySize)
 		{
-			ParseDirectoryTableNoLock(m_memory, m_memorySize, size.main, false);
+			ParseDirectoryTableNoLock(m_memory, m_memorySize, size.main, 0);
 			m_memorySize = size.main;
 		}
 		if (size.overlay > m_overlaySize)
 		{
-			ParseDirectoryTableNoLock(m_overlayMemory, m_overlaySize, size.overlay, true);
+			ParseDirectoryTableNoLock(m_overlayMemory, m_overlaySize, size.overlay, OverlayTableFlag);
 			m_overlaySize = size.overlay;
 		}
 	}
 
-	void DirectoryTable::ParseDirectoryTableNoLock(const u8* memory, u32 from, u32 to, bool isOverlay)
+	void DirectoryTable::ParseDirectoryTableNoLock(const u8* memory, u32 from, u32 to, u32 isOverlayFlag)
 	{
-		u32 additionalOffset = isOverlay ? OverlayTableOffset : 0;
 		BinaryReader reader(memory, from, to);
 		while (true)
 		{
@@ -225,10 +235,10 @@ namespace uba
 			UBA_ASSERTF(pos < to, TC("Should never read past size (pos: %llu, size: %u)"), pos, to);
 			u64 storageSize = reader.Read7BitEncoded();
 			StringKey dirKey = reader.ReadStringKey();
-			auto insres = m_lookup.try_emplace(dirKey, m_memoryBlock); // Note that this is allowed to overwrite
-			u32& tableOffset = isOverlay ? insres.first->second.latestOverlayOffset : insres.first->second.latestOffset;
-			UBA_ASSERTF(tableOffset < OverlayTableOffset || additionalOffset == OverlayTableOffset, TC("Can't add non-overlay directory items after overlay items"));
-			tableOffset = u32(reader.GetPosition()) + additionalOffset;
+			Directory& dir = m_lookup.try_emplace(dirKey, m_memoryBlock).first->second; // Note that this is allowed to overwrite
+			u32& tableOffset = isOverlayFlag ? dir.latestOverlayOffset : dir.latestOffset;
+			UBA_ASSERTF(!IsOverlayOffset(tableOffset) || IsOverlayOffset(isOverlayFlag), TC("Can't add non-overlay directory items after overlay items"));
+			tableOffset = u32(reader.GetPosition()) | isOverlayFlag;
 			reader.Skip(storageSize - sizeof(dirKey));
 		}
 	}
@@ -246,23 +256,29 @@ namespace uba
 
 	void DirectoryTable::PopulateDirectoryNoLock(const StringKeyHasher& hasher, Directory& dir)
 	{
-		if (dir.populatedOffset == dir.latestOffset)
-			return;
-		PopulateDirectoryRecursive(hasher, dir.latestOffset, dir.populatedOffset, dir.files);
-		dir.populatedOffset = dir.latestOffset;
+		if (dir.populatedOffset != dir.latestOffset)
+		{
+			PopulateDirectoryRecursive(hasher, dir.latestOffset, dir.populatedOffset, dir.files, 0);
+			dir.populatedOffset = dir.latestOffset;
+		}
+		if (dir.populatedOverlayOffset != dir.latestOverlayOffset)
+		{
+			PopulateDirectoryRecursive(hasher, dir.latestOverlayOffset, dir.populatedOverlayOffset, dir.files, OverlayTableFlag);
+			dir.populatedOverlayOffset = dir.latestOverlayOffset;
+		}
 	}
 
 	BinaryReader DirectoryTable::GetReader(u32 tableOffset)
 	{
-		if (tableOffset < OverlayTableOffset)
-			return BinaryReader(m_memory, tableOffset);
+		if (IsOverlayOffset(tableOffset))
+			return BinaryReader(m_overlayMemory, tableOffset - OverlayTableFlag);
 		else
-			return { m_overlayMemory - OverlayTableOffset, tableOffset};
+			return BinaryReader(m_memory, tableOffset);
 	}
 
-	void DirectoryTable::PopulateDirectoryRecursive(const StringKeyHasher& hasher, u32 latestOffset, u32 populatedOffset, EntryLookup& files)
+	void DirectoryTable::PopulateDirectoryRecursive(const StringKeyHasher& hasher, u32 latestOffset, u32 populatedOffset, EntryLookup& files, u32 isOverlayFlag)
 	{
-		BinaryReader reader(GetReader(latestOffset));
+		BinaryReader reader(isOverlayFlag?m_overlayMemory:m_memory, latestOffset&~isOverlayFlag);
 
 		u32 prevTableOffset = u32(reader.Read7BitEncoded());
 
@@ -278,7 +294,7 @@ namespace uba
 				firstIsRoot = prevTableOffset == InvalidTableOffset;
 				break;
 			}
-			reader = GetReader(prevTableOffset);
+			reader.SetPosition(prevTableOffset);
 			prevTableOffset = u32(reader.Read7BitEncoded());
 			readerOffsets[count++] = u32(reader.GetPosition());
 			if (count == sizeof_array(buffer))
@@ -290,7 +306,7 @@ namespace uba
 
 		for (u32 i=count; i>0; --i)
 		{
-			reader = GetReader(readerOffsets[i-1]);
+			reader.SetPosition(readerOffsets[i-1]);
 			if (firstIsRoot)
 			{
 				firstIsRoot = false;
@@ -301,14 +317,14 @@ namespace uba
 				reader.ReadFileIndex(); // Directory file index
 			}
 
-			PopulateDirectoryWithFiles(reader, hasher, files);
+			PopulateDirectoryWithFiles(reader, hasher, files, isOverlayFlag);
 		}
 
 		if (readerOffsets != buffer)
 			delete[] readerOffsets;
 	}
 
-	void DirectoryTable::PopulateDirectoryWithFiles(BinaryReader& reader, const StringKeyHasher& hasher, EntryLookup& files)
+	void DirectoryTable::PopulateDirectoryWithFiles(BinaryReader& reader, const StringKeyHasher& hasher, EntryLookup& files, u32 isOverlayFlag)
 	{
 		u64 itemCount = reader.Read7BitEncoded();
 
@@ -334,7 +350,10 @@ namespace uba
 			}
 
 			StringKey filenameKey = ToStringKey(hasher, filename.data, filename.count);
-			files[filenameKey] = {offset}; // Always write, since same file might have been added with new info
+			
+			DirTableOffset& fileOffset = files[filenameKey];
+			//UBA_ASSERTF(fileOffset.internal == InvalidTableOffset || !IsOverlayOffset(fileOffset.internal) || IsOverlayOffset(offset), TC("File %s is updated from main table while existing in overlay"), filename.data);
+			fileOffset.internal = offset | isOverlayFlag; // Always write, since same file might have been added with new info
 		}
 	}
 
@@ -385,23 +404,17 @@ namespace uba
 			if (dirIt != m_lookup.end())
 			{
 				DirectoryTable::Directory& parentDir = dirIt->second;
-				if (parentDir.latestOffset == -1)
-					return Exists_No;
-				if (parentDir.populatedOffset != parentDir.latestOffset)
+				if (parentDir.populatedOffset != parentDir.latestOffset || parentDir.populatedOverlayOffset != parentDir.latestOverlayOffset)
 				{
 					SCOPED_WRITE_LOCK(parentDir.lock, lock);
 					if (parentDir.populatedOffset != parentDir.latestOffset)
 					{
-						PopulateDirectoryRecursive(ancestorHasher, parentDir.latestOffset, parentDir.populatedOffset, parentDir.files);
+						PopulateDirectoryRecursive(ancestorHasher, parentDir.latestOffset, parentDir.populatedOffset, parentDir.files, 0);
 						parentDir.populatedOffset = parentDir.latestOffset;
 					}
-				}
-				if (parentDir.populatedOverlayOffset != parentDir.latestOverlayOffset)
-				{
-					SCOPED_WRITE_LOCK(parentDir.lock, lock);
 					if (parentDir.populatedOverlayOffset != parentDir.latestOverlayOffset)
 					{
-						PopulateDirectoryRecursive(ancestorHasher, parentDir.latestOverlayOffset, parentDir.populatedOverlayOffset, parentDir.files);
+						PopulateDirectoryRecursive(ancestorHasher, parentDir.latestOverlayOffset, parentDir.populatedOverlayOffset, parentDir.files, OverlayTableFlag);
 						parentDir.populatedOverlayOffset = parentDir.latestOverlayOffset;
 					}
 				}
@@ -412,7 +425,7 @@ namespace uba
 					return Exists_No;
 				if (inAncestor)
 				{
-					BinaryReader reader(m_memory, entryIt->second.internal, m_memorySize);
+					BinaryReader reader(GetReader(entryIt->second.internal));
 					reader.SkipString();
 					if (!IsDirectory(reader.ReadFileAttributes()))
 						return Exists_No;
@@ -430,21 +443,21 @@ namespace uba
 		return Exists_Maybe;
 	}
 
-	DirectoryTable::Exists DirectoryTable::EntryExists(StringView str, bool checkIfDir)
+	DirectoryTable::Exists DirectoryTable::EntryExists(StringView str, bool checkIfDir, DirTableOffset* outTableOffset)
 	{
 		StringBuffer<> str2(str);
 		if (str2[str2.count-1] == PathSeparator)
 			str2.Resize(str2.count-1);
 		if (CaseInsensitiveFs)
 			str2.MakeLower();
-		return EntryExists(ToStringKey(str2), str2, checkIfDir, nullptr);
+		return EntryExists(ToStringKey(str2), str2, checkIfDir, outTableOffset);
 	}
 
 	u32 DirectoryTable::GetAttributes(DirTableOffset tableOffset)
 	{
 		if (IsDirectoryEntry(tableOffset))
 		{
-			BinaryReader reader(m_memory, tableOffset.internal & ~DirectoryEntryFlag);
+			BinaryReader reader(GetReader(tableOffset.internal & ~DirectoryEntryFlag));
 			u64 prevTableOffset = reader.Read7BitEncoded();
 			while (prevTableOffset != InvalidTableOffset)
 			{
@@ -453,7 +466,7 @@ namespace uba
 			}
 			return reader.ReadFileAttributes();
 		}
-		BinaryReader reader(m_memory, tableOffset.internal);
+		BinaryReader reader(GetReader(tableOffset.internal));
 		reader.SkipString();
 		return reader.ReadFileAttributes();
 	}
@@ -574,10 +587,24 @@ namespace uba
 
 	bool DirectoryTable::GetLatestOffset(DirTableOffset& out, const Directory& dir)
 	{
-		if (dir.latestOffset == InvalidFileHandle)
-			return false;
-		out.internal = dir.latestOffset | DirectoryEntryFlag;
-		return true;
+		if (dir.latestOverlayOffset != InvalidTableOffset)
+		{
+			out.internal = dir.latestOverlayOffset | DirectoryEntryFlag | OverlayTableFlag;
+			return true;
+		}
+		if (dir.latestOffset != InvalidTableOffset)
+		{
+			out.internal = dir.latestOffset | DirectoryEntryFlag;
+			return true;
+		}
+		return false;
+	}
+
+	bool DirectoryTable::ContainsFile(const StringKeyHasher& hasher, Directory& dir, StringKey fileNameKey)
+	{
+		SCOPED_WRITE_LOCK(dir.lock, lock);
+		PopulateDirectoryNoLock(hasher, dir);
+		return dir.files.find(fileNameKey) != dir.files.end();
 	}
 
 	template<typename Func>
@@ -620,7 +647,6 @@ namespace uba
 		}
 	}
 
-#if PLATFORM_WINDOWS
 	template<typename Func>
 	void DirectoryTable::TraverseFilesRecursiveNoLock(StringView path, const Func& func)
 	{
@@ -648,14 +674,17 @@ namespace uba
 	template<typename Func>
 	void DirectoryTable::TraverseAllFilesNoLock(const Func& func)
 	{
+		#if PLATFORM_WINDOWS
 		for (tchar l='a';l!='z'; ++l)
 		{
 			StringBuffer<4> drive;
 			drive.Append(l).Append(':');
 			TraverseFilesRecursiveNoLock(drive, func);
 		}
+		#else
+		TraverseFilesRecursiveNoLock(TCV("/"), func);
+		#endif
 	}
-#endif
 
 	DirectoryTable::DirectoryTable(MemoryBlock& block)
 	:	m_memoryBlock(block)
diff --git a/Engine/Source/Programs/UnrealBuildAccelerator/Core/Public/UbaProtocol.h b/Engine/Source/Programs/UnrealBuildAccelerator/Core/Public/UbaProtocol.h
--- a/Engine/Source/Programs/UnrealBuildAccelerator/Core/Public/UbaProtocol.h
+++ b/Engine/Source/Programs/UnrealBuildAccelerator/Core/Public/UbaProtocol.h
@@ -9,6 +9,7 @@ namespace uba
 	#define UBA_PROCESS_MESSAGES \
 		UBA_PROCESS_MESSAGE(Init) \
 		UBA_PROCESS_MESSAGE(CreateFile) \
+		UBA_PROCESS_MESSAGE(RegisterFileForWrite) \
 		UBA_PROCESS_MESSAGE(GetFullFileName) \
 		UBA_PROCESS_MESSAGE(GetLongPathName) \
 		UBA_PROCESS_MESSAGE(CloseFile) \
diff --git a/Engine/Source/Programs/UnrealBuildAccelerator/Core/Public/UbaWinBinDependencyParser.h b/Engine/Source/Programs/UnrealBuildAccelerator/Core/Public/UbaWinBinDependencyParser.h
--- a/Engine/Source/Programs/UnrealBuildAccelerator/Core/Public/UbaWinBinDependencyParser.h
+++ b/Engine/Source/Programs/UnrealBuildAccelerator/Core/Public/UbaWinBinDependencyParser.h
@@ -7,7 +7,6 @@
 #include <algorithm>
 
 #if !defined(UBA_IS_DETOURED_INCLUDE)
-#define True_CreateFileW ::CreateFileW
 #define True_CreateFileMappingW ::CreateFileMappingW
 #define True_MapViewOfFile ::MapViewOfFile
 #define True_UnmapViewOfFile ::UnmapViewOfFile
@@ -152,7 +151,7 @@ namespace uba
 			outError.Appendf(TC("Batch file has no dependencies: %.*s"), filePath.count, filePath.data);
 			return true;
 		}
-		HANDLE fileHandle = True_CreateFileW(filePath.data, GENERIC_READ, FILE_SHARE_READ, 0, OPEN_EXISTING, FILE_ATTRIBUTE_NORMAL | FILE_FLAG_SEQUENTIAL_SCAN, 0);
+		HANDLE fileHandle = ::CreateFileW(filePath.data, GENERIC_READ, FILE_SHARE_READ, 0, OPEN_EXISTING, FILE_ATTRIBUTE_NORMAL | FILE_FLAG_SEQUENTIAL_SCAN, 0);
 		if (fileHandle == INVALID_HANDLE_VALUE)
 		{
 			outError.Appendf(TC("File not found: %.*s"), filePath.count, filePath.data);
diff --git a/Engine/Source/Programs/UnrealBuildAccelerator/Detours/Private/PosixOS/UbaDetoursFunctionsPosix.cpp b/Engine/Source/Programs/UnrealBuildAccelerator/Detours/Private/PosixOS/UbaDetoursFunctionsPosix.cpp
--- a/Engine/Source/Programs/UnrealBuildAccelerator/Detours/Private/PosixOS/UbaDetoursFunctionsPosix.cpp
+++ b/Engine/Source/Programs/UnrealBuildAccelerator/Detours/Private/PosixOS/UbaDetoursFunctionsPosix.cpp
@@ -2961,7 +2961,7 @@ namespace uba
 		g_directoryTable.Init(dirTableMem, directoryTableCount, directoryTableSize.main);
 
 		StringBuffer<128> overlayTableUidName;
-		GetMappingHandleName(overlayTableUidName, directoryTableHandle);
+		GetMappingHandleName(overlayTableUidName, overlayTableHandle);
 		int overlayTableFd = shm_open(overlayTableUidName.data, O_RDONLY, S_IRUSR | S_IWUSR);
 		UBA_ASSERTF(overlayTableFd != -1, "shm_open failed (%s)", strerror(errno));
 		u8* overlayTableMem = (u8*)mmap(NULL, OverlayTableMemSize, PROT_READ, MAP_SHARED, overlayTableFd, 0);
diff --git a/Engine/Source/Programs/UnrealBuildAccelerator/Detours/Private/UbaDetoursFileMappingTable.cpp b/Engine/Source/Programs/UnrealBuildAccelerator/Detours/Private/UbaDetoursFileMappingTable.cpp
--- a/Engine/Source/Programs/UnrealBuildAccelerator/Detours/Private/UbaDetoursFileMappingTable.cpp
+++ b/Engine/Source/Programs/UnrealBuildAccelerator/Detours/Private/UbaDetoursFileMappingTable.cpp
@@ -107,6 +107,17 @@ namespace uba
 		g_directoryTable.ParseDirectoryTable(directoryTableSize);
 	}
 
+	void Rpc_RegisterFileForWrite(StringView fileName, const StringKey& fileNameKey)
+	{
+		RPC_MESSAGE(RegisterFileForWrite, createFile)
+		writer.WriteString(fileName);
+		writer.WriteStringKey(fileNameKey);
+		BinaryReader reader = writer.Flush();
+		DirTableSize directoryTableSize = FromU64(reader.ReadU64());
+		pcs.Leave();
+		g_directoryTable.ParseDirectoryTable(directoryTableSize);
+	}
+
 	void Rpc_CheckRemapping(const StringView& fileName, const StringKey& fileNameKey)
 	{
 		RPC_MESSAGE(CheckRemapping, createFile)
@@ -118,7 +129,7 @@ namespace uba
 		g_mappedFileTable.ParseNoLock(mappedFileTableSize);
 	}
 
-	u32 Rpc_UpdateDirectory(const StringKey& dirKey, const tchar* dirName, u64 dirNameLen, bool lockDirTable)
+	bool Rpc_UpdateDirectory(const StringKey& dirKey, const tchar* dirName, u64 dirNameLen, bool lockDirTable)
 	{
 		DirTableSize directoryTableSize;
 		u32 tableOffset;
@@ -128,8 +139,7 @@ namespace uba
 			writer.WriteStringKey(dirKey);
 			BinaryReader reader = writer.Flush();
 			directoryTableSize = FromU64(reader.ReadU64());
-			u64 to = reader.ReadU32();
-			tableOffset = to == InvalidTableOffset ? ~u32(0) : (u32)to;
+			tableOffset = reader.ReadU32();
 			pcs.Leave();
 			DEBUG_LOG_PIPE(L"ListDirectory", L"(%ls)", dirName);
 		}
@@ -137,7 +147,7 @@ namespace uba
 			g_directoryTable.ParseDirectoryTable(directoryTableSize);
 		else
 			g_directoryTable.ParseDirectoryTableNoLock(directoryTableSize);
-		return tableOffset;
+		return tableOffset != InvalidTableOffset;
 	}
 
 	void Rpc_UpdateCloseHandle(const tchar* handleName, u32 closeId, bool deleteOnClose, const tchar* newName, SharedMemoryHandle memoryHandle, u64 memoryWritten, bool success)
@@ -363,7 +373,7 @@ namespace uba
 		u32 dirNameLen = u32(lastPathSeparator - entryName.data);
 		DirHash hash(StringView(entryNameForKey.data, dirNameLen));
 
-		if (Rpc_UpdateDirectory(hash.key, entryName.data, dirNameLen) == ~u32(0))
+		if (!Rpc_UpdateDirectory(hash.key, entryName.data, dirNameLen))
 			return false;
 
 		SCOPED_WRITE_LOCK(g_directoryTable.m_lookupLock, lookLock);
diff --git a/Engine/Source/Programs/UnrealBuildAccelerator/Detours/Private/UbaDetoursFileMappingTable.h b/Engine/Source/Programs/UnrealBuildAccelerator/Detours/Private/UbaDetoursFileMappingTable.h
--- a/Engine/Source/Programs/UnrealBuildAccelerator/Detours/Private/UbaDetoursFileMappingTable.h
+++ b/Engine/Source/Programs/UnrealBuildAccelerator/Detours/Private/UbaDetoursFileMappingTable.h
@@ -94,8 +94,9 @@ namespace uba
 	enum : u8 { AccessFlag_Read = 1, AccessFlag_Write = 2 };
 
 	void Rpc_CreateFileW(const StringView& fileName, const StringKey& fileNameKey, u8 accessFlags, tchar* outNewName, u64 newNameCapacity, u64& outSize, u32& outCloseId, bool lock);
+	void Rpc_RegisterFileForWrite(StringView fileName, const StringKey& fileNameKey);
 	void Rpc_CheckRemapping(const StringView& fileName, const StringKey& fileNameKey);
-	u32  Rpc_UpdateDirectory(const StringKey& dirKey, const tchar* dirName, u64 dirNameLen, bool lockDirTable = true);
+	bool Rpc_UpdateDirectory(const StringKey& dirKey, const tchar* dirName, u64 dirNameLen, bool lockDirTable = true);
 	void Rpc_UpdateCloseHandle(const tchar* handleName, u32 closeId, bool deleteOnClose, const tchar* newName, SharedMemoryHandle memoryHandle, u64 memoryWritten, bool success);
 	void Rpc_UpdateTables();
 	bool Rpc_GetEntryOffset(DirTableOffset& out, const StringKey& entryNameKey, StringView entryName, bool checkIfDir = false);
diff --git a/Engine/Source/Programs/UnrealBuildAccelerator/Detours/Private/Windows/UbaDetoursFunctionsKernelBase.inl b/Engine/Source/Programs/UnrealBuildAccelerator/Detours/Private/Windows/UbaDetoursFunctionsKernelBase.inl
--- a/Engine/Source/Programs/UnrealBuildAccelerator/Detours/Private/Windows/UbaDetoursFunctionsKernelBase.inl
+++ b/Engine/Source/Programs/UnrealBuildAccelerator/Detours/Private/Windows/UbaDetoursFunctionsKernelBase.inl
@@ -210,35 +210,6 @@ BOOL Detoured_DuplicateHandle(HANDLE hSourceProcessHandle, HANDLE hSourceHandle,
 	return res;
 }
 
-HANDLE Detoured_CreateFileW(LPCWSTR lpFileName, DWORD dwDesiredAccess, DWORD dwShareMode, LPSECURITY_ATTRIBUTES lpSecurityAttributes, DWORD dwCreationDisposition, DWORD dwFlagsAndAttributes, HANDLE hTemplateFile)
-{
-	DETOURED_CALL(CreateFileW);
-	DEBUG_LOG_PROXY(L"CreateFileW", L"%ls", lpFileName);
-	u32 disallowDetour = !lpFileName || Equals(lpFileName, L"nul");
-	t_disallowDetour += disallowDetour;
-	t_hackInsideCreateFile = true;
-	HANDLE h = True_CreateFileW(lpFileName, dwDesiredAccess, dwShareMode, lpSecurityAttributes, dwCreationDisposition, dwFlagsAndAttributes, hTemplateFile);
-	t_hackInsideCreateFile = false;
-	t_disallowDetour -= disallowDetour;
-	return h;
-}
-
-// Calls directly to NtCreateFile so need to be detoured
-HANDLE Detoured_CreateFileA(LPCSTR lpFileName, DWORD dwDesiredAccess, DWORD dwShareMode, LPSECURITY_ATTRIBUTES lpSecurityAttributes, DWORD dwCreationDisposition, DWORD dwFlagsAndAttributes, HANDLE hTemplateFile)
-{
-	DETOURED_CALL(CreateFileA);
-	DEBUG_LOG_TRUE(L"CreateFileA", L"%hs", lpFileName);
-	StringBuffer<> fileName;
-	fileName.Appendf(L"%hs", lpFileName);
-	u32 disallowDetour = fileName.Equals(L"nul");
-	t_disallowDetour += disallowDetour;
-	t_hackInsideCreateFile = true;
-	HANDLE h = True_CreateFileA(lpFileName, dwDesiredAccess, dwShareMode, lpSecurityAttributes, dwCreationDisposition, dwFlagsAndAttributes, hTemplateFile);
-	t_hackInsideCreateFile = false;
-	t_disallowDetour -= disallowDetour;
-	return h;
-}
-
 BOOL Detoured_CreateDirectoryW(LPCWSTR lpPathName, LPSECURITY_ATTRIBUTES lpSecurityAttributes)
 {
 	DETOURED_CALL(CreateDirectoryW);
@@ -502,8 +473,7 @@ BOOL Detoured_GetVolumeInformationByHandleW(HANDLE hFile, LPWSTR lpVolumeNameBuf
 	else if (isListDirectoryHandle(hFile))
 	{
 		auto& listHandle = asListDirectoryHandle(hFile);
-		if (!g_directoryTable.GetLatestOffset(entryOffset, listHandle.dir))
-			UBA_ASSERT(false);
+		entryOffset = listHandle.dirTableOffset;
 		trueHandle = INVALID_HANDLE_VALUE;
 	}
 
@@ -1850,6 +1820,9 @@ BOOL Detoured_GetFileInformationByHandleEx(HANDLE hFile, FILE_INFO_BY_HANDLE_CLA
 	{
 		DirectoryTable::EntryInformation entryInfo;
 		g_directoryTable.GetEntryInformation(entryInfo, entryOffset);
+
+		UBA_ASSERT(entryInfo.attributes);
+
 		if (fileInformationClass == FileBasicInfo)
 		{
 			auto& data = *(FILE_BASIC_INFO*)lpFileInformation;
@@ -1858,7 +1831,7 @@ BOOL Detoured_GetFileInformationByHandleEx(HANDLE hFile, FILE_INFO_BY_HANDLE_CLA
 			data.LastWriteTime = ToLargeInteger(entryInfo.lastWrite);
 			data.ChangeTime = ToLargeInteger(entryInfo.lastWrite);
 			data.FileAttributes = entryInfo.attributes;
-			DEBUG_LOG_DETOURED(L"GetFileInformationByHandleEx", L"(DIRTABLE) (FileBasicInfo) %llu (%ls)", uintptr_t(hFile), HandleToName(hFile));
+			DEBUG_LOG_DETOURED(L"GetFileInformationByHandleEx", L"(DIRTABLE) (FileBasicInfo) %llu (%ls) -> %s", uintptr_t(hFile), HandleToName(hFile), ToString(TRUE));
 			return TRUE;
 		}
 		else if (fileInformationClass == FileIdInfo)
@@ -1955,8 +1928,11 @@ BOOL Detoured_GetFileInformationByHandle(HANDLE hFile, LPBY_HANDLE_FILE_INFORMAT
 			lpFileInformation->nNumberOfLinks = 1;//~u32(0); // TODO
 			u64 fileSize = fi.size;
 			if (fileSize == InvalidValue)
+			{
 				fileSize = entryInfo.size;
-
+				if (fileSize == InvalidValue) // File is still open?
+					fileSize = 0;
+			}
 #if UBA_DEBUG_VALIDATE
 			if (g_validateFileAccess && !(entryInfo.attributes & FILE_ATTRIBUTE_DIRECTORY))
 			{
@@ -3578,6 +3554,20 @@ BOOL Detoured_GetVolumePathNameA(LPCSTR lpszFileName, LPSTR lpszVolumePathName,
 	return True_GetVolumePathNameA(lpszFileName, lpszVolumePathName, cchBufferLength);
 }
 
+HANDLE Detoured_CreateFileW(LPCWSTR lpFileName, DWORD dwDesiredAccess, DWORD dwShareMode, LPSECURITY_ATTRIBUTES lpSecurityAttributes, DWORD dwCreationDisposition, DWORD dwFlagsAndAttributes, HANDLE hTemplateFile)
+{
+	DETOURED_CALL(CreateFileW);
+	DEBUG_LOG_PROXY(L"CreateFileW", L"%ls", lpFileName);
+	return True_CreateFileW(lpFileName, dwDesiredAccess, dwShareMode, lpSecurityAttributes, dwCreationDisposition, dwFlagsAndAttributes, hTemplateFile);
+}
+
+HANDLE Detoured_CreateFileA(LPCSTR lpFileName, DWORD dwDesiredAccess, DWORD dwShareMode, LPSECURITY_ATTRIBUTES lpSecurityAttributes, DWORD dwCreationDisposition, DWORD dwFlagsAndAttributes, HANDLE hTemplateFile)
+{
+	DETOURED_CALL(CreateFileA);
+	DEBUG_LOG_TRUE(L"CreateFileA", L"%hs", lpFileName);
+	return True_CreateFileA(lpFileName, dwDesiredAccess, dwShareMode, lpSecurityAttributes, dwCreationDisposition, dwFlagsAndAttributes, hTemplateFile);
+}
+
 DWORD Detoured_GetFileSize(HANDLE hFile, LPDWORD lpFileSizeHigh)
 {
 	DETOURED_CALL(GetFileSize);
@@ -3836,7 +3826,7 @@ BOOL Detoured_CreateHardLinkA(LPCSTR lpFileName, LPCSTR lpExistingFileName, LPSE
 BOOL Detoured_DeleteFileA(LPCSTR lpFileName)
 {
 	DETOURED_CALL(DeleteFileA);
-	DEBUG_LOG_TRUE(L"DeleteFileA", L"");
+	DEBUG_LOG_PROXY(L"DeleteFileA", L"");
 	t_calledDeleteFileW = false;
 	BOOL res = True_DeleteFileA(lpFileName);
 	UBA_ASSERT(t_calledDeleteFileW);
diff --git a/Engine/Source/Programs/UnrealBuildAccelerator/Detours/Private/Windows/UbaDetoursFunctionsNtDll.inl b/Engine/Source/Programs/UnrealBuildAccelerator/Detours/Private/Windows/UbaDetoursFunctionsNtDll.inl
--- a/Engine/Source/Programs/UnrealBuildAccelerator/Detours/Private/Windows/UbaDetoursFunctionsNtDll.inl
+++ b/Engine/Source/Programs/UnrealBuildAccelerator/Detours/Private/Windows/UbaDetoursFunctionsNtDll.inl
@@ -1286,8 +1286,6 @@ NTSTATUS NTAPI Shared_NtCreateFile(bool IsCreateFunc, PHANDLE hFileHandle, ACCES
 	if (ntPath.count && ntPath[1] != '?') // we can have an empty ntPath but RootDirectory set
 		return trueFunc();
 
-	bool allowDirectoryCache = g_allowDirectoryCache && t_hackInsideCreateFile; // TODO: This needs to be fixed.. t_hackInsideCreateFile just happens to work but is wrong!
-
 	HANDLE rootDir = ObjectAttributes->RootDirectory;
 	StringBuffer<> fileName;
 
@@ -1384,8 +1382,6 @@ NTSTATUS NTAPI Shared_NtCreateFile(bool IsCreateFunc, PHANDLE hFileHandle, ACCES
 
 	DevirtualizePath(fileName);
 
-	DirTableOffset dirTableOffset;
-
 	//UBA_ASSERT(CreateDisposition != FILE_SUPERSEDE);
 	bool isDeleteOnClose = (CreateOptions & FILE_DELETE_ON_CLOSE) != 0; // clang is using CreateFile with DeleteOnClose to delete files after build errors
 
@@ -1396,17 +1392,29 @@ NTSTATUS NTAPI Shared_NtCreateFile(bool IsCreateFunc, PHANDLE hFileHandle, ACCES
 	
 	keepInMemory = keepInMemory || ((isWrite || isDeleteOnClose) && g_rules->IsOutputFile(fileName, g_systemTemp) && g_allowKeepFilesInMemory) || (isWrite && isThrowAway);
 
-#if UBA_DEBUG_LOG_ENABLED
+	#if UBA_DEBUG_LOG_ENABLED
 	const wchar_t* isWriteStr = isWrite ? L"WRITE " : L""; (void)isWriteStr;
-#endif
+	#endif
 
 	bool isSystemFile = fileName.StartsWith(g_systemRoot.data);
-	bool checkIfDir = false;
+	if (g_runningRemote)
+	{
+		if (fileName.EndsWith(TCV(".tlb")))  // We want the tbl files to be detoured
+			isSystemFile = false;
+		if (fileName.StartsWith(g_exeDir))
+			isSystemFile = true;
+	}
+
+	//bool checkIfDir = false;
 	// This is here just to avoid getting a NtQueryVolumeInformationFile to get volume information 
 	if (fileName[3] == 0 && fileName[1] == ':')
 	{
-		isSystemFile = ToLower(fileName[0]) == g_systemRoot[0];
-		checkIfDir = true;
+		isSystemFile = true;
+		if (g_runningRemote)
+		{
+			// TODO: Should we create a detoured handle so name calls are correct?
+			ObjectAttributes->ObjectName->Buffer[4] = g_systemRoot[0];
+		}
 	}
 
 	bool isSystemOrTempFile = isSystemFile || fileName.StartsWith(g_systemTemp.data);
@@ -1432,64 +1440,109 @@ NTSTATUS NTAPI Shared_NtCreateFile(bool IsCreateFunc, PHANDLE hFileHandle, ACCES
 	fileNameLower.MakeLower();
 	StringKey fileNameKey = ToStringKey(fileNameLower);
 
-	if (allowDirectoryCache)
+	DirTableOffset dirTableOffset;
+
+	// All NtCreateFile that needs existing file is checked first
+	if (CreateDisposition == FILE_OPEN || CreateDisposition == FILE_OVERWRITE)
 	{
-		// This is an optimization where we populate directory table and use that to figure out if file exists or not..
-		// .. in msvc's case it doesn't matter much because these tables are already up to date when msvc use CreateFile.
-		// .. clang otoh is using CreateFile with tooons of different paths trying to open files.. in remote worker case this becomes super expensive
-		if (!isWrite && !isSystemOrTempFile) // We need to skip SystemTemp.. lots of stuff going on there.
+		if (!Rpc_GetEntryOffset(dirTableOffset, fileNameKey, fileName, false))
 		{
-			CHECK_PATH(fileNameLower);
+			IoStatusBlock->Status = STATUS_OBJECT_NAME_NOT_FOUND;
+			IoStatusBlock->Information = FILE_DOES_NOT_EXIST;
+			DEBUG_LOG_DETOURED(funcName, L"(DIRTABLE) %llu, (%ls) -> STATUS_OBJECT_NAME_NOT_FOUND", uintptr_t(*hFileHandle), fileName.data);
+			return STATUS_OBJECT_NAME_NOT_FOUND;
+		}
+	}
+	else if (isSystemOrTempFile)
+	{
+		// TODO: For now, ignore these when writing..
+		// Once overlay is fully implemented we want all written files to end up in overlay to simplify all the other functions
+	}
+	else if (const tchar* lastSeparator = fileNameLower.Last(PathSeparator))
+	{
+		StringView dirName(fileNameLower.data, u32(lastSeparator - fileNameLower.data));
+		DirHash hash(dirName);
 
-			bool allowEarlyOut = true;
-			if (!Rpc_GetEntryOffset(dirTableOffset, fileNameKey, fileName, checkIfDir))
-			{
-				// This could be a written file not reported to server yet
-				{
-					SCOPED_READ_LOCK(g_mappedFileTable.m_lookupLock, lock);
-					auto findIt = g_mappedFileTable.m_lookup.find(fileNameKey);
-					if (findIt != g_mappedFileTable.m_lookup.end())
-						allowEarlyOut = findIt->second.deleted;
-				}
-				if (allowEarlyOut)
-				{
-					//SetLastError(ERROR_FILE_NOT_FOUND); // Don't think this is needed
-					*hFileHandle = INVALID_HANDLE_VALUE;
-					DEBUG_LOG_DETOURED(funcName, L"NOTFOUND_USINGTABLE (%ls) -> Error", fileName.data);
+		SCOPED_WRITE_LOCK(g_directoryTable.m_lookupLock, lookupLock);
+		auto insres = g_directoryTable.m_lookup.try_emplace(hash.key, g_memoryBlock);
+		DirectoryTable::Directory& dir = insres.first->second;
+		if (insres.second)
+			if (g_directoryTable.EntryExistsNoLock(hash.key, dirName) != DirectoryTable::Exists_No)
+				Rpc_UpdateDirectory(hash.key, dirName.data, dirName.count, false);
+		DirTableOffset dirTableOffset2;
+		if (!g_directoryTable.GetLatestOffset(dirTableOffset2, dir))
+		{
+			IoStatusBlock->Status = STATUS_OBJECT_PATH_NOT_FOUND;
+			IoStatusBlock->Information = 0;
+			DEBUG_LOG_DETOURED(funcName, L"(DIRTABLE) %llu, (%s) -> STATUS_OBJECT_PATH_NOT_FOUND", uintptr_t(*hFileHandle), fileName.data);
+			return STATUS_OBJECT_PATH_NOT_FOUND;
+		}
 
-#if UBA_DEBUG_VALIDATE
-					if (g_validateFileAccess)
-					{
-						SuppressDetourScope _;
-						UBA_ASSERTF(True_GetFileAttributesW(fileName.data) == INVALID_FILE_ATTRIBUTES, L"DIRTABLE claims file %ls does not exist but it does", fileName.data);
-					}
-#endif
+		if (CreateDisposition == FILE_CREATE && g_directoryTable.ContainsFile(hash.open, dir, fileNameKey))
+		{
+			IoStatusBlock->Information = FILE_EXISTS;
+			IoStatusBlock->Status = STATUS_OBJECT_NAME_COLLISION;
+			DEBUG_LOG_DETOURED(funcName, L"(DIRTABLE) %llu, (%s) -> STATUS_OBJECT_NAME_COLLISION", uintptr_t(*hFileHandle), fileName.data);
+			return STATUS_OBJECT_NAME_COLLISION;
+		}
+	}
+	else
+	{
+		UBA_ASSERT(CreateDisposition != FILE_CREATE); // It is extremely unlikely we end up here.
+	}
 
-					return STATUS_OBJECT_NAME_NOT_FOUND;
-				}
-			}
-			else if (!checkIfDir)
-			{
-				// File could have been deleted.
-				DirectoryTable::EntryInformation entryInfo;
-				g_directoryTable.GetEntryInformation(entryInfo, dirTableOffset);
-				if (entryInfo.attributes == 0)
-				{
-					DEBUG_LOG_DETOURED(funcName, L"DELETED %llu, (%ls) -> Success", uintptr_t(*hFileHandle), fileName.data);
-					return STATUS_OBJECT_NAME_NOT_FOUND;
-				}
-				else if (useContent && IsDirectory(entryInfo.attributes))
-				{
-					DEBUG_LOG_DETOURED(funcName, L"%llu, (%ls) -> STATUS_FILE_IS_A_DIRECTORY", uintptr_t(*hFileHandle), fileName.data);
-					return STATUS_FILE_IS_A_DIRECTORY;
-				}
-			}
+	static constexpr GENERIC_MAPPING genericMapping = { FILE_GENERIC_READ, FILE_GENERIC_WRITE, FILE_GENERIC_EXECUTE, FILE_ALL_ACCESS };
+	ACCESS_MASK desiredAccess = DesiredAccess;
+	if (desiredAccess & MAXIMUM_ALLOWED)
+		desiredAccess = FILE_ALL_ACCESS | SYNCHRONIZE;
+	else
+		RtlMapGenericMask(&desiredAccess, &genericMapping);
+
+	//bool wantsWrite = (desiredAccess & (FILE_WRITE_DATA | FILE_APPEND_DATA | FILE_WRITE_ATTRIBUTES | FILE_WRITE_EA | DELETE | WRITE_DAC | WRITE_OWNER)) != 0;
+	//bool wantsRead  = (desiredAccess & (FILE_READ_DATA | FILE_READ_ATTRIBUTES | FILE_READ_EA)) != 0;
+
+
+	// Almost all non-writes end up here except the weird FILE_OPEN_IF with read-only
+	if (!isWrite && CreateDisposition != FILE_OPEN_IF)
+	{
+		UBA_ASSERT(IsValidEntry(dirTableOffset));
+		DirectoryTable::EntryInformation entryInfo;
+		g_directoryTable.GetEntryInformation(entryInfo, dirTableOffset);
+
+		// File/directory has been deleted
+		if (entryInfo.attributes == 0)
+		{
+			IoStatusBlock->Status = STATUS_OBJECT_NAME_NOT_FOUND;
+			IoStatusBlock->Information = FILE_DOES_NOT_EXIST;
+			DEBUG_LOG_DETOURED(funcName, L"DELETED %llu, (%ls) -> STATUS_OBJECT_NAME_NOT_FOUND", uintptr_t(*hFileHandle), fileName.data);
+			return STATUS_OBJECT_NAME_NOT_FOUND;
+		}
 
-			bool isWriteAttributes = (DesiredAccess & FILE_WRITE_ATTRIBUTES) != 0;
 
+		bool isDirectory = IsDirectory(entryInfo.attributes);
+
+		if (isDirectory && (CreateOptions & FILE_NON_DIRECTORY_FILE) != 0)
+		{
+			IoStatusBlock->Information = 0;
+			IoStatusBlock->Status = STATUS_FILE_IS_A_DIRECTORY;
+			DEBUG_LOG_DETOURED(funcName, L"%llu, (%ls) -> STATUS_FILE_IS_A_DIRECTORY", uintptr_t(*hFileHandle), fileName.data);
+			return STATUS_FILE_IS_A_DIRECTORY;
+		}
+
+		if (!isDirectory && (CreateOptions & FILE_DIRECTORY_FILE) != 0)
+		{
+			IoStatusBlock->Information = 0;
+			IoStatusBlock->Status = STATUS_NOT_A_DIRECTORY;
+			DEBUG_LOG_DETOURED(funcName, L"%llu, (%ls) -> STATUS_NOT_A_DIRECTORY", uintptr_t(*hFileHandle), fileName.data);
+			return STATUS_NOT_A_DIRECTORY;
+		}
+
+		// Only meta data, no content.
+		if ((desiredAccess & (FILE_READ_DATA | FILE_WRITE_DATA | FILE_APPEND_DATA)) == 0)
+		{
 			// If file is an output file we still allow this path and accept wrong (compressed) size.
 			// This is a bit hacky but we don't want to transfer and decompress file just to get the size
-			if (allowEarlyOut && !useContent && !isWriteAttributes && (!CouldBeCompressedFile(fileName) || g_rules->IsOutputFile(fileName, g_systemTemp)))
+			if (!CouldBeCompressedFile(fileName) || g_rules->IsOutputFile(fileName, g_systemTemp))
 			{
 				auto dh = new DetouredHandle(HandleType_File);
 				dh->fileObject = new FileObject();
@@ -1509,90 +1562,55 @@ NTSTATUS NTAPI Shared_NtCreateFile(bool IsCreateFunc, PHANDLE hFileHandle, ACCES
 				return STATUS_SUCCESS;
 			}
 		}
-	}
 
-	if (isSystemOrTempFile)
-	{
-	}
-	else if ((DesiredAccess & FILE_LIST_DIRECTORY) != 0 && (CreateOptions & FILE_DIRECTORY_FILE) != 0)
-	{
-		if (isWrite || !g_allowListDirectoryHandle)
+		// We can't fail desiredAccess & FILE_LIST_DIRECTORY here, this needs to be set with the handle (FILE_READ_DATA is the same as FILE_LIST_DIRECTORY)
+		if (isDirectory && (desiredAccess & FILE_LIST_DIRECTORY) != 0)
 		{
-			TimerScope ts(g_kernelStats.createFile);
-			UBA_ASSERT(!g_runningRemote);
-			NTSTATUS res = True_NtCreateFile(hFileHandle, DesiredAccess, ObjectAttributes, IoStatusBlock, AllocationSize, FileAttributes, ShareAccess, CreateDisposition, CreateOptions, EaBuffer, EaLength);
+			// I want to defer this because it doesn't mean that it will access the directory content
+			UBA_ASSERT(fileNameLower.data[fileNameLower.count - 1] != '\\');
+			DirHash hash(fileNameLower);
+			UBA_ASSERT(!isDeleteOnClose);
+			SCOPED_WRITE_LOCK(g_directoryTable.m_lookupLock, lookupLock);
+			auto insres = g_directoryTable.m_lookup.try_emplace(hash.key, g_memoryBlock);
+			DirectoryTable::Directory& dir = insres.first->second;
+			if (insres.second)
+			{
+				auto existsResult = g_directoryTable.EntryExistsNoLock(hash.key, fileNameLower);
+				if (existsResult != DirectoryTable::Exists_No)
+					Rpc_UpdateDirectory(hash.key, fileNameLower.data, fileNameLower.count, false);
+			}
+			g_directoryTable.PopulateDirectory(hash.open, dir);
 
-			DEBUG_LOG_DETOURED(funcName, L"(CREATE_DIR) %llu, (%ls) -> %ls", uintptr_t(*hFileHandle), fileName.data, ToString(res).data);
-			return res;
-		}
+			auto listHandle = new ListDirectoryHandle{ hash.key, dir };
+			listHandle->dirTableOffset = dirTableOffset;
+			listHandle->it = 0;
 
-		UBA_ASSERT(fileNameLower.data[fileNameLower.count - 1] != '\\');
-		DirHash hash(fileNameLower);
+			SCOPED_READ_LOCK(dir.lock, lock);
+			listHandle->fileTableOffsets.resize(dir.files.size());
+			u32 it = 0;
+			for (auto& pair : dir.files)
+				listHandle->fileTableOffsets[it++] = pair.second;
+			lock.Leave();
 
-		SCOPED_WRITE_LOCK(g_directoryTable.m_lookupLock, lookupLock);
-		auto insres = g_directoryTable.m_lookup.try_emplace(hash.key, g_memoryBlock);
-		DirectoryTable::Directory& dir = insres.first->second;
-		if (insres.second)
-		{
-			auto existsResult = g_directoryTable.EntryExistsNoLock(hash.key, fileNameLower);
-			if (existsResult != DirectoryTable::Exists_No)
-				Rpc_UpdateDirectory(hash.key, fileNameLower.data, fileNameLower.count, false);
-		}
+			#if UBA_DEBUG_VALIDATE
+			if (g_validateFileAccess)
+				listHandle->validateHandle = validateHandle;
+			#endif
 
-		bool exists = false;
-		if (dir.latestOffset != InvalidTableOffset)
-		{
-			DirTableOffset entryOffset;
-			g_directoryTable.GetLatestOffset(entryOffset, dir);
-			DirectoryTable::EntryInformation entryInfo;
-			g_directoryTable.GetEntryInformation(entryInfo, entryOffset);
-			exists = entryInfo.attributes != 0;
-		}
+			*hFileHandle = makeListDirectoryHandle(listHandle);
 
-#if UBA_DEBUG_VALIDATE
-		NTSTATUS res = exists ? 0 : -1; (void)res;
-		HANDLE validateHandle = INVALID_HANDLE_VALUE;
-		if (g_validateFileAccess && !isListDirectoryHandle(rootDir))
-		{
-			IO_STATUS_BLOCK IoStatusBlock2;
-			ObjectAttributes->RootDirectory = rootDir;
-			NTSTATUS res2 = True_NtCreateFile(&validateHandle, DesiredAccess, ObjectAttributes, &IoStatusBlock2, AllocationSize, FileAttributes, ShareAccess, CreateDisposition, CreateOptions, EaBuffer, EaLength); (void)res2;
-			UBA_ASSERT(res < 0 && res2 < 0 || res >= 0 && res2 >= 0);
-			ObjectAttributes->RootDirectory = nullptr;
-		}
-#endif
+			listHandle->originalName = g_memoryBlock.Strdup(fileName).data;
 
-		if (!exists)
-		{
-			DEBUG_LOG_DETOURED(funcName, L"(AS_DIRECTORY) (%ls) -> NOT EXISTS", fileName.data);
-			return STATUS_OBJECT_NAME_NOT_FOUND;
+			IoStatusBlock->Information = 1;
+			IoStatusBlock->Pointer = nullptr;
+			IoStatusBlock->Status = 0;
+			DEBUG_LOG_DETOURED(funcName, L"(AS_DIRECTORY) (%ls) -> %llu", fileName.data, uintptr_t(*hFileHandle));
+			return STATUS_SUCCESS;
 		}
-		g_directoryTable.PopulateDirectory(hash.open, dir);
-
-		auto listHandle = new ListDirectoryHandle{ hash.key, insres.first->second };
-		listHandle->it = 0;
-
-		SCOPED_READ_LOCK(dir.lock, lock);
-		listHandle->fileTableOffsets.resize(dir.files.size());
-		u32 it = 0;
-		for (auto& pair : dir.files)
-			listHandle->fileTableOffsets[it++] = pair.second;
-		lock.Leave();
-
-#if UBA_DEBUG_VALIDATE
-		if (g_validateFileAccess)
-			listHandle->validateHandle = validateHandle;
-#endif
-		*hFileHandle = makeListDirectoryHandle(listHandle);
-
-		listHandle->originalName = g_memoryBlock.Strdup(fileName).data;
-
-		IoStatusBlock->Information = 1;
-		IoStatusBlock->Pointer = nullptr;
-		IoStatusBlock->Status = 0;
-		DEBUG_LOG_DETOURED(funcName, L"(AS_DIRECTORY) (%ls) -> %llu", fileName.data, uintptr_t(*hFileHandle));
-
-		return STATUS_SUCCESS;
+	}
+	else
+	{
+		UBA_ASSERT((CreateOptions & FILE_DIRECTORY_FILE) == 0); // Create directory not supported atm.. will implement if needed
 	}
 
 	if (!keepInMemory || !isWrite) // we might get \\pipe\ here... 
@@ -1629,8 +1647,8 @@ NTSTATUS NTAPI Shared_NtCreateFile(bool IsCreateFunc, PHANDLE hFileHandle, ACCES
 			info.originalName = g_memoryBlock.Strdup(fileName).data;
 		if (isWrite)
 		{
-			bool lastWasWrite = IsContentWrite(info.lastDesiredAccess, 0);
-			bool shouldReport = !lastWasWrite || info.deleted || isDeleteOnClose;
+			//bool lastWasWrite = IsContentWrite(info.lastDesiredAccess, 0);
+			bool shouldReport = true;//!lastWasWrite || info.deleted || isDeleteOnClose;
 			shouldReport = shouldReport && !keepInMemory;
 
 			if (info.isFileMap)
@@ -1707,6 +1725,8 @@ NTSTATUS NTAPI Shared_NtCreateFile(bool IsCreateFunc, PHANDLE hFileHandle, ACCES
 			auto dh = new DetouredHandle(HandleType_File, th);
 			dh->dirTableOffset = dirTableOffset;
 			dh->fileObject = fo;
+			if (isWrite)
+				Rpc_RegisterFileForWrite(fileName, fileNameKey);
 			return makeDetouredHandle(dh);
 		};
 
@@ -2363,6 +2383,7 @@ NTSTATUS NTAPI Detoured_NtClose(HANDLE handle)
 		}
 #endif
 
+		DEBUG_LOG_DETOURED(L"NtClose", L"%llu (%ls) -> %ls", uintptr_t(handle), HandleToName(handle), ToString(STATUS_SUCCESS).data);
 		delete& listHandle;
 		return STATUS_SUCCESS;
 	}
@@ -2438,6 +2459,7 @@ NTSTATUS NTAPI Detoured_NtClose(HANDLE handle)
 	FileInfo& fi = *fo->fileInfo;
 	const wchar_t* path = fi.name;
 	wchar_t temp[512];
+
 	if (MemoryFile* mf = fi.memoryFile)
 	{
 		if (IsWrite(fo->desiredAccess, 0))
diff --git a/Engine/Source/Programs/UnrealBuildAccelerator/Detours/Private/Windows/UbaDetoursFunctionsWin.cpp b/Engine/Source/Programs/UnrealBuildAccelerator/Detours/Private/Windows/UbaDetoursFunctionsWin.cpp
--- a/Engine/Source/Programs/UnrealBuildAccelerator/Detours/Private/Windows/UbaDetoursFunctionsWin.cpp
+++ b/Engine/Source/Programs/UnrealBuildAccelerator/Detours/Private/Windows/UbaDetoursFunctionsWin.cpp
@@ -172,6 +172,7 @@ struct ListDirectoryHandle
 
 	StringKey dirNameKey;
 	DirectoryTable::Directory& dir;
+	DirTableOffset dirTableOffset;
 	int it;
 	Vector<DirTableOffset> fileTableOffsets;
 	HANDLE validateHandle;
@@ -530,7 +531,6 @@ void OnModuleLoaded(HMODULE moduleHandle, const StringView& name);
 
 // Variables used to communicate state from kernelbase functions to ntdll functions
 thread_local const wchar_t* t_renameFileNewName;
-thread_local bool t_hackInsideCreateFile;
 
 #include "UbaDetoursFunctionsMiMalloc.inl"
 #include "UbaDetoursFunctionsNtDll.inl"
diff --git a/Engine/Source/Programs/UnrealBuildAccelerator/Test/Private/UbaTestDirectoryTable.cpp b/Engine/Source/Programs/UnrealBuildAccelerator/Test/Private/UbaTestDirectoryTable.cpp
--- a/Engine/Source/Programs/UnrealBuildAccelerator/Test/Private/UbaTestDirectoryTable.cpp
+++ b/Engine/Source/Programs/UnrealBuildAccelerator/Test/Private/UbaTestDirectoryTable.cpp
@@ -10,29 +10,22 @@ namespace uba
 		DirectoryTableHolder holder(logger.m_writer, TC(""), false);
 		//holder.SetTreatTempDirAsEmpty();
 
-		if (!holder.RefreshDirectory(testRootDir.data, true))
-			return false;
+		CHECK_TRUE(holder.RefreshDirectory(testRootDir.data, true));
 
 		MemoryBlock dirMem(TC("Mem"));
 		DirectoryTable dir(dirMem);
 		dir.Init(holder.GetDirectoryTableMemory(), 0, holder.GetDirectoryTableSize(nullptr).main);
 		
-		if (dir.EntryExists(testRootDir, false) != DirectoryTable::Exists_Yes)
-			return false;
-		if (dir.EntryExists(testRootDir, true) != DirectoryTable::Exists_Yes)
-			return false;
+		CHECK_TRUE(dir.EntryExists(testRootDir, false) == DirectoryTable::Exists_Yes);
+		CHECK_TRUE(dir.EntryExists(testRootDir, true) == DirectoryTable::Exists_Yes);
 
 		StringBuffer<> fooFile;
-		if (!CreateTestFile(fooFile, logger, testRootDir, TCV("foo.txt"), TCV("foo")))
-			return false;
-		if (!holder.RegisterNewFile(fooFile.data))
-			return false;
+		CHECK_TRUE(CreateTestFile(fooFile, logger, testRootDir, TCV("foo.txt"), TCV("foo")));
+		CHECK_TRUE(holder.RegisterNewFile(fooFile.data));
 
-		if (dir.EntryExists(fooFile, true) != DirectoryTable::Exists_No)
-			return false;
+		CHECK_TRUE(dir.EntryExists(fooFile, true) == DirectoryTable::Exists_No);
 		dir.ParseDirectoryTable(holder.GetDirectoryTableSize(nullptr));
-		if (dir.EntryExists(fooFile, true) != DirectoryTable::Exists_Yes)
-			return false;
+		CHECK_TRUE(dir.EntryExists(fooFile, true) == DirectoryTable::Exists_Yes);
 
 
 		DirectoryTableOverlay overlay;
@@ -41,47 +34,47 @@ namespace uba
 		dir.InitOverlay(overlay.memory, overlay.size);
 
 		StringBuffer<> barFile;
-		if (!CreateTestFile(barFile, logger, testRootDir, TCV("bar.txt"), TCV("bar")))
-			return false;
+		CHECK_TRUE(CreateTestFile(barFile, logger, testRootDir, TCV("bar.txt"), TCV("bar")));
+		CHECK_TRUE(holder.PopulateOverlayFile(overlay, barFile, [](FileInformation& out, StringKey) { out.size = sizeof(tchar)*3; out.attributes = DefaultAttributes(); return true; }));
 
-		holder.PopulateOverlayFile(overlay, barFile, [](FileInformation& out, StringKey) { out.size = 6; out.attributes = DefaultAttributes(); return true; });
-
-		if (dir.EntryExists(barFile, true) != DirectoryTable::Exists_No)
-			return false;
+		CHECK_TRUE(dir.EntryExists(barFile, true) == DirectoryTable::Exists_No);
 		dir.ParseDirectoryTable(holder.GetDirectoryTableSize(&overlay));
-		if (dir.EntryExists(barFile, true) != DirectoryTable::Exists_Yes)
-			return false;
+		CHECK_TRUE(dir.EntryExists(barFile, true) == DirectoryTable::Exists_Yes);
 
 		StringBuffer<> mehFile;
-		if (!CreateTestFile(mehFile, logger, testRootDir, TCV("meh.txt"), TCV("meh")))
-			return false;
-		if (!holder.RegisterNewFile(mehFile.data))
-			return false;
+		CHECK_TRUE(CreateTestFile(mehFile, logger, testRootDir, TCV("meh.txt"), TCV("meh")));
+		CHECK_TRUE(holder.RegisterNewFile(mehFile.data));
 
-		if (dir.EntryExists(mehFile, true) != DirectoryTable::Exists_No)
-			return false;
+		CHECK_TRUE(dir.EntryExists(mehFile, true) == DirectoryTable::Exists_No);
 		dir.ParseDirectoryTable(holder.GetDirectoryTableSize(&overlay));
-		if (dir.EntryExists(mehFile, true) != DirectoryTable::Exists_Yes)
-			return false;
+		CHECK_TRUE(dir.EntryExists(mehFile, true) == DirectoryTable::Exists_Yes);
 
-#if PLATFORM_WINDOWS
 		u32 counter = 0;
-		StringBuffer<> testRootDirLower(StringView(testRootDir.data, testRootDir.count - 1));
-		testRootDirLower.MakeLower();
+		StringBuffer<> testRootDirForLookup(StringView(testRootDir.data, testRootDir.count - 1));
+		if (CaseInsensitiveFs)
+			testRootDirForLookup.MakeLower();
+
+		auto findIt = dir.m_lookup.find(ToStringKey(testRootDirForLookup));
+		CHECK_TRUE(findIt != dir.m_lookup.end());
+		CHECK_TRUE(findIt->second.files.size() == 3);
+
 		bool success = true;
-		dir.TraverseFilesRecursiveNoLock(testRootDirLower, [&](const DirectoryTable::EntryInformation& info, StringView fileName, DirTableOffset fileOffset)
+		dir.TraverseFilesRecursiveNoLock(testRootDirForLookup, [&](const DirectoryTable::EntryInformation& info, StringView fileName, DirTableOffset fileOffset)
 			{
 				++counter;
-				success = success && info.size == 6;
+				success = success && info.size == sizeof(tchar)*3;
 				success = success && info.attributes != 0;
 				success = success && fileName.GetFileName().count == 7;
 			});
-		if (!success)
-			return false;
-		if (counter != 3)
-			return false;
-#endif
+		CHECK_TRUE(success);
+		CHECK_TRUE(counter == 3);
 
+		// Remove file by writing entry with attributes=0
+		CHECK_TRUE(holder.WriteOverlayEntry(overlay, barFile, {}));
+		dir.ParseDirectoryTable(holder.GetDirectoryTableSize(&overlay));
+		DirTableOffset entryOffset;
+		CHECK_TRUE(dir.EntryExists(barFile, false, &entryOffset) == DirectoryTable::Exists_Yes);
+		CHECK_TRUE(!dir.GetAttributes(entryOffset));
 
 		return true;
 	}
