[UBA]
* Added option for enabling/disabling fill lock (UseFillLock) in server session (default is on like before)

diff --git a/Engine/Source/Programs/UnrealBuildAccelerator/Common/Private/UbaSessionServer.cpp b/Engine/Source/Programs/UnrealBuildAccelerator/Common/Private/UbaSessionServer.cpp
--- a/Engine/Source/Programs/UnrealBuildAccelerator/Common/Private/UbaSessionServer.cpp
+++ b/Engine/Source/Programs/UnrealBuildAccelerator/Common/Private/UbaSessionServer.cpp
@@ -153,7 +153,7 @@ namespace uba
 			table->GetValueAsBool(nameToHashTableEnabled, TC("NameToHashTableEnabled"));
 			table->GetValueAsBool(traceIOEnabled, TC("TraceIOEnabled"));
 			table->GetValueAsBool(customCasKeysEnabled, TC("CustomCasKeysEnabled"));
-
+			table->GetValueAsBool(useFillLock, TC("UseFillLock"));
 		}
 	}
 
@@ -165,7 +165,7 @@ namespace uba
 			dir.Resize(dir.count - 7); // Remove native and slash
 		const tchar* archPath[2] = { TC("x64"), TC("arm64") };
 		if (!dir.EndsWith(archPath[IsArmBinary]))
-			return reportError ? logger.Error(TC("Module dir is not under supported folder (%s) to be able to run cross architecutres, can't figure out matching x64/arm64 folder"), dir.data) : false;
+			return reportError ? logger.Error(TC("Module dir is not under supported folder (%s) to be able to run cross architectures, can't figure out matching x64/arm64 folder"), dir.data) : false;
 		dir.Resize(dir.count - TStrlen(archPath[IsArmBinary])).Append(archPath[!IsArmBinary]);
 		if (isUbtPath)
 			dir.Append(PathSeparator).Append(TCV("native"));
@@ -243,6 +243,7 @@ namespace uba
 		m_remoteTraceEnabled = info.remoteTraceEnabled;
 		m_traceIOEnabled = info.traceIOEnabled;
 		m_customCasKeysEnabled = info.customCasKeysEnabled;
+		m_useFillLock = info.useFillLock;
 
 		if (m_resetCas)
 			m_storage.Reset();
@@ -1717,7 +1718,11 @@ namespace uba
 
 		float weightLeft = availableWeight;
 		u32 addCount = 0;
-		SCOPED_FUTEX(m_fillUpOneAtTheTimeLock, fillLock); // This is a lock to group files better (all clients connect at the same time);
+
+		if (m_useFillLock)
+			m_fillUpOneAtTheTimeLock.Enter();
+		auto exitFileLock = MakeGuard([&]() { if (m_useFillLock) m_fillUpOneAtTheTimeLock.Leave(); });
+
 		while (weightLeft > 0)
 		{
 			RemoteProcess* process = DequeueProcess(session, sessionId, connectionInfo.GetId());
@@ -1814,7 +1819,7 @@ namespace uba
 
 			weightLeft -= startInfo.weight;
 		}
-		fillLock.Leave();
+		exitFileLock.Execute();
 
 		u32 neededDirectoryTableSize = GetDirectoryTableSize(nullptr).main;
 		u32 neededHashTableSize;
diff --git a/Engine/Source/Programs/UnrealBuildAccelerator/Common/Public/UbaSessionServer.h b/Engine/Source/Programs/UnrealBuildAccelerator/Common/Public/UbaSessionServer.h
--- a/Engine/Source/Programs/UnrealBuildAccelerator/Common/Public/UbaSessionServer.h
+++ b/Engine/Source/Programs/UnrealBuildAccelerator/Common/Public/UbaSessionServer.h
@@ -232,6 +232,7 @@ namespace uba
 		bool m_remoteTraceEnabled = false;
 		bool m_traceIOEnabled = false;
 		bool m_customCasKeysEnabled = false;
+		bool m_useFillLock = false;
 
 		Scheduler* m_outerScheduler = nullptr;
 
diff --git a/Engine/Source/Programs/UnrealBuildAccelerator/Common/Public/UbaSessionServerCreateInfo.h b/Engine/Source/Programs/UnrealBuildAccelerator/Common/Public/UbaSessionServerCreateInfo.h
--- a/Engine/Source/Programs/UnrealBuildAccelerator/Common/Public/UbaSessionServerCreateInfo.h
+++ b/Engine/Source/Programs/UnrealBuildAccelerator/Common/Public/UbaSessionServerCreateInfo.h
@@ -22,5 +22,6 @@ namespace uba
 		bool remoteTraceEnabled = false; // If this is true, the agents will run trace and send the .uba file back to server
 		bool traceIOEnabled = true; // Enable IO information for storage 
 		bool customCasKeysEnabled = false; // Need to set this to true to be able to use SetCustomCasKeyFromTrackedInputs
+		bool useFillLock = true; // Will fill one process request at the time. Useful if files that are similar are queued together
 	};
 }
