[UBA]
* Changed UbaAgent to by-default use crawler for dependencies

diff --git a/Engine/Source/Programs/UnrealBuildAccelerator/Agent/Private/UbaAgent.cpp b/Engine/Source/Programs/UnrealBuildAccelerator/Agent/Private/UbaAgent.cpp
--- a/Engine/Source/Programs/UnrealBuildAccelerator/Agent/Private/UbaAgent.cpp
+++ b/Engine/Source/Programs/UnrealBuildAccelerator/Agent/Private/UbaAgent.cpp
@@ -118,7 +118,7 @@ namespace uba
 		logger.Info(TC("  -resendcas              Will try to send same cas multiple times (set this to true if server is allowed to remove cas files)"));
 		logger.Info(TC("  -populateCas=<dir>      Prepopulate cas database with files in dir. If files needed exists on machine this can be an optimization"));
 		logger.Info(TC("  -description            Add more info about the agent that will show in the trace log when hovering over session"));
-		logger.Info(TC("  -usecrawler             Enables include crawler for known process types (clang/msvc)"));
+		logger.Info(TC("  -disablecrawler         Disables file crawler for known process types (clang/msvc)"));
 		#if PLATFORM_MAC
 		logger.Info(TC("  -killtcphogs            If failing to bind listen socket UbaAgent will attempt to kill processes holding it and then retry"));
 		logger.Info(TC("  -populateCasFromXcodeVersion=<version>   Prepopulate cas database with files from local xcode installation that matches the version."));
@@ -380,7 +380,7 @@ namespace uba
 		bool verbose = false;
 		bool printSummary = false;
 		bool killRandom = false;
-		bool useCrawler = false;
+		bool useCrawler = true;
 		bool useOverlappedSend = true;
 		u32 iocpWorkerCount = 0;
 		bool downloadDetoursLib = true;
@@ -617,7 +617,11 @@ namespace uba
 			{
 				killRandom = true;
 			}
-			else if (name.Equals(TCV("-usecrawler")))
+			else if (name.Equals(TCV("-disablecrawler")))
+			{
+				useCrawler = false;
+			}
+			else if (name.Equals(TCV("-enablecrawler"))) // Ignored, flag has been inverted
 			{
 				useCrawler = true;
 			}
