[UBA]
* Readded UseCrawler command line option to fix agent not starting when on command line

diff --git a/Engine/Source/Programs/UnrealBuildAccelerator/Agent/Private/UbaAgent.cpp b/Engine/Source/Programs/UnrealBuildAccelerator/Agent/Private/UbaAgent.cpp
--- a/Engine/Source/Programs/UnrealBuildAccelerator/Agent/Private/UbaAgent.cpp
+++ b/Engine/Source/Programs/UnrealBuildAccelerator/Agent/Private/UbaAgent.cpp
@@ -621,7 +621,7 @@ namespace uba
 			{
 				useCrawler = false;
 			}
-			else if (name.Equals(TCV("-enablecrawler"))) // Ignored, flag has been inverted
+			else if (name.Equals(TCV("-usecrawler"))) // Ignored, flag has been inverted
 			{
 				useCrawler = true;
 			}
