[UBA]
* Added exposed variable useDependencyCrawler to Client

diff --git a/Engine/Source/Programs/UnrealBuildAccelerator/Cli/Private/UbaClient.h b/Engine/Source/Programs/UnrealBuildAccelerator/Cli/Private/UbaClient.h
--- a/Engine/Source/Programs/UnrealBuildAccelerator/Cli/Private/UbaClient.h
+++ b/Engine/Source/Programs/UnrealBuildAccelerator/Cli/Private/UbaClient.h
@@ -20,6 +20,7 @@ namespace uba
 		u32 index = 0;
 		bool createSession = true;
 		bool addDirSuffix = true;
+		bool useDependencyCrawler = false;
 	};
 
 	class Client
@@ -60,6 +61,7 @@ namespace uba
 				sessionClientInfo.rootDir = clientRootDir.data;
 				sessionClientInfo.deleteSessionsOlderThanSeconds = 1;
 				sessionClientInfo.name.Append("Agent").AppendValue(info.index);
+				sessionClientInfo.useDependencyCrawler = info.useDependencyCrawler;
 
 				sessionClient = new SessionClient(sessionClientInfo);
 				sessionClient->Start();
