[UBA]
* Fixed so dependency crawler can handle command lines as well as response files

diff --git a/Engine/Source/Programs/UnrealBuildAccelerator/Common/Private/UbaSession.cpp b/Engine/Source/Programs/UnrealBuildAccelerator/Common/Private/UbaSession.cpp
--- a/Engine/Source/Programs/UnrealBuildAccelerator/Common/Private/UbaSession.cpp
+++ b/Engine/Source/Programs/UnrealBuildAccelerator/Common/Private/UbaSession.cpp
@@ -3637,10 +3637,6 @@ namespace uba
 		if (crawlerType == DependencyCrawlerType_None)
 			return;
 
-		const tchar* at = TStrchr(startInfo.arguments, '@');
-		if (!at)
-			return;
-
 		bool isLinkCrawler = crawlerType == DependencyCrawlerType_ClangLinker || crawlerType == DependencyCrawlerType_MsvcLinker;
 
 		auto CreateFileFunc = [this, ph = ProcessHandle(&process), rules = startInfo.rules, isLinkCrawler](TrackWorkScope& tracker, const StringView& fileName, const DependencyCrawler::AccessFileFunc& func, bool& stopCrawling)
@@ -3705,7 +3701,10 @@ namespace uba
 
 		auto DevirtualizePathFunc = [this, rootsHandle = startInfo.rootsHandle](StringBufferBase& inOut) { return DevirtualizePath(inOut, rootsHandle, false); };
 
-		m_dependencyCrawler.Add(at + 1, startInfo.workingDir, CreateFileFunc, DevirtualizePathFunc, startInfo.application, crawlerType, startInfo.rules->index);
+		if (const tchar* at = TStrchr(startInfo.arguments, '@'))
+			m_dependencyCrawler.AddRsp(at + 1, startInfo.workingDir, CreateFileFunc, DevirtualizePathFunc, startInfo.application, crawlerType, startInfo.rules->index);
+		else if (!isLinkCrawler)
+			m_dependencyCrawler.AddCommandLine(startInfo.arguments, startInfo.workingDir, CreateFileFunc, DevirtualizePathFunc, startInfo.application, crawlerType, startInfo.rules->index);
 	}
 
 #if UBA_USE_NATIVE_MAC_SEMAPHORES
diff --git a/Engine/Source/Programs/UnrealBuildAccelerator/Common/Public/UbaDependencyCrawler.h b/Engine/Source/Programs/UnrealBuildAccelerator/Common/Public/UbaDependencyCrawler.h
--- a/Engine/Source/Programs/UnrealBuildAccelerator/Common/Public/UbaDependencyCrawler.h
+++ b/Engine/Source/Programs/UnrealBuildAccelerator/Common/Public/UbaDependencyCrawler.h
@@ -38,7 +38,8 @@ namespace uba
 		using CreateFileFunc = Function<bool(TrackWorkScope& tracker, const StringView& fileName, const AccessFileFunc& func, bool& stopCrawling)>;
 
 		inline void Init(FileExistsFunc&& fileExistsFunc, TraverseFilesFunc&& traverseFilesFunc, bool useBloomFilter = true);
-		inline bool Add(const tchar* rsp, const tchar* workDir, CreateFileFunc&& func, DevirtualizePathFunc&& devirtualizePathFunc, const tchar* app, DependencyCrawlerType type, u32 ruleIndex);
+		inline bool AddRsp(const tchar* rsp, const tchar* workDir, CreateFileFunc&& createFileFunc, DevirtualizePathFunc&& devirtualizePathFunc, const tchar* app, DependencyCrawlerType type, u32 ruleIndex);
+		inline bool AddCommandLine(const tchar* commandLine, const tchar* workDir, CreateFileFunc&& createFileFunc, DevirtualizePathFunc&& devirtualizePathFunc, const tchar* app, DependencyCrawlerType type, u32 ruleIndex);
 
 		Logger& m_logger;
 		WorkManager& m_workManager;
@@ -102,12 +103,13 @@ namespace uba
 		using CodeFiles = List<CodeFile>;
 
 		inline bool ParseRsp(TrackWorkScope& tracker, Instance& instance, const StringView& rsp, CodeFiles& outCodeFiles);
-		inline bool ParseRsp2(TrackWorkScope& tracker, Instance& instance, const void* data, u64 dataSize, const tchar* rsp, CodeFiles& outCodeFiles);
+		template<typename CharType> inline bool ParseRsp2(TrackWorkScope& tracker, Instance& instance, const void* data, u64 dataSize, const tchar* rsp, CodeFiles& outCodeFiles);
 		inline bool ParsePch(Pch& pch, Instance& instance, const void* data, u64 dataSize);
 		inline bool ParseCodeFile(TrackWorkScope& tracker, Instance& instance, HandledFiles& handledFiles, const StringView& codeFile, bool parseDefines, const StringView& caller);
 		inline bool ParseCodeFile2(TrackWorkScope& tracker, Instance& instance, HandledFiles& handledFiles, const void* data, u64 dataSize, const StringView& codeFile, bool parseDefines);
 		inline bool HandleInclude(TrackWorkScope& tracker, Instance& instance, HandledFiles& handledFiles, const StringView& rootPath, const StringView& include, const StringView& codeFile, bool parseDefines);
 		inline bool TraverseInclude(TrackWorkScope& tracker, Instance& instance, const StringView& rootPathWithSlash);
+		inline bool TraverseCodeFiles(const WorkContext& context, Instance& instance, const CodeFiles& codeFiles, u32 ruleIndex);
 
 		inline u32 GetBottleneckMax(WorkManager& workManager, u32 bottleneckMax)
 		{
@@ -134,7 +136,7 @@ namespace uba
 		m_useBloomFilter = useBloomFilter;
 	}
 
-	bool DependencyCrawler::Add(const tchar* rsp, const tchar* workDir, CreateFileFunc&& createFileFunc, DevirtualizePathFunc&& devirtualizePathFunc, const tchar* app, DependencyCrawlerType type, u32 ruleIndex)
+	bool DependencyCrawler::AddRsp(const tchar* rsp, const tchar* workDir, CreateFileFunc&& createFileFunc, DevirtualizePathFunc&& devirtualizePathFunc, const tchar* app, DependencyCrawlerType type, u32 ruleIndex)
 	{
 		auto& instance = *new Instance();
 
@@ -154,84 +156,36 @@ namespace uba
 			{
 				auto& instance = ref.instance;
 				CodeFiles codeFiles;
-				if (!ParseRsp(context.tracker, instance, instance.rsp, codeFiles))
-					return;
-
-				for (auto& cf : codeFiles)
-				{
-					if (instance.usePch && cf.hasPch)
-					{
-						// Note, clang is always reading all the individual files that has been stored in the pch to validate content..
-						// So this code path is not implemented for clang.. if -fno-validate-ast-input-files-content is added we can implement this path
-						StringBuffer<> fixedFile;
-						FixPath(cf.path.c_str(), instance.workDir.data(), instance.workDir.size(), fixedFile);
-						if (CaseInsensitiveFs)
-							fixedFile.MakeLower();
-						StringKey pchKey = ToStringKey(fixedFile);
-
-						SCOPED_FUTEX(m_pchLookupLock, lock);
-						Pch& pch = m_pchLookup[pchKey];
-						lock.Leave();
-
-						SCOPED_FUTEX(pch.lock, lock2);
-						if (!pch.handled)
-						{
-							fixedFile.Append(TCV(".dep.json"));
-							bool stopCrawling = false;
-							instance.createFileFunc(context.tracker, fixedFile, [&](const void* data, u64 dataSize)
-							{
-								return ParsePch(pch, instance, data, dataSize);
-							}, stopCrawling);
-							if (stopCrawling)
-								return;
-							pch.handled = true;
-						}
-						instance.pch = &pch;
-					}
-					else if (instance.type == DependencyCrawlerType_MsvcLinker)
-					{
-						// If we end up here it means that we have compressed obj files enabled and want to decompress them in parallel
-						m_workManager.AddWork([this, ref, codeFile = cf.path](const WorkContext& context)
-							{
-								bool stopCrawling = false;
-								ref.instance.createFileFunc(context.tracker, codeFile, {}, stopCrawling);
-							}, 1, TC("CrawlForDecomp"), ColorWork, WorkPriority_Low, &m_bottleneck);
-					}
-					else
-					{
-						bool parseDefines = instance.platform.empty();
-						HandledFiles handledFiles;
-						ParseCodeFile(context.tracker, instance, parseDefines ? handledFiles : m_handledFiles, cf.path, parseDefines, instance.rsp);
-						if (instance.platform.empty())
-						{
-							if (!instance.overriddenPlatformName.empty())
-								instance.platform = instance.overriddenPlatformName;
-							else
-								instance.platform = instance.compiledPlatform;
-						}
-					}
-				}
+				if (ParseRsp(context.tracker, instance, instance.rsp, codeFiles))
+					TraverseCodeFiles(context, instance, codeFiles, ruleIndex);
+			}, 1, TC("CrawlRsp"), ColorWork, WorkPriority_Low, &m_bottleneck);
+		return true;
+	}
 
-				if (!instance.builtinIncludesDir.empty())
-				{
-					SCOPED_FUTEX(m_builtIncludesHandledLock, lock);
-					bool shouldHandle = m_builtIncludesHandled.insert(ruleIndex).second;
-					lock.Leave();
-					if (shouldHandle)
-					{
-						StringBuffer<> fixedPath;
-						FixPath(instance.builtinIncludesDir.c_str(), instance.workDir.data(), instance.workDir.size(), fixedPath);
-						if (!instance.devirtualizePathFunc(fixedPath))
-							UBA_LOG_DEVIRTUALIZATION_ERROR("Failed to devirtualize path %s found in builtin includes", fixedPath.data);
-						if (CaseInsensitiveFs)
-							fixedPath.MakeLower();
-						fixedPath.Append(PathSeparator);
+	bool DependencyCrawler::AddCommandLine(const tchar* commandLine, const tchar* workDir, CreateFileFunc&& createFileFunc, DevirtualizePathFunc&& devirtualizePathFunc, const tchar* app, DependencyCrawlerType type, u32 ruleIndex)
+	{
+		auto& instance = *new Instance();
 
-						TraverseInclude(context.tracker, instance, fixedPath);
-					}
-				}
+		instance.type = type;
+		instance.application = app;
+		//instance.rsp = rsp;
+		instance.platform = TC("Win64"); // Hard coded platform for now
+		instance.workDir = workDir;
+		UBA_ASSERT(*workDir);
+		if (CaseInsensitiveFs)
+			ToLower(instance.workDir.data());
+		if (instance.workDir.back() != PathSeparator)
+			instance.workDir += PathSeparator;
+		instance.createFileFunc = std::move(createFileFunc);
+		instance.devirtualizePathFunc = std::move(devirtualizePathFunc);
 
-			}, 1, TC("CrawlRsp"), ColorWork, WorkPriority_Low, &m_bottleneck);
+		m_workManager.AddWork([this, ref = InstanceRef(instance), ruleIndex, cmdLine = TString(commandLine)](const WorkContext& context)
+			{
+				auto& instance = ref.instance;
+				CodeFiles codeFiles;
+				if (ParseRsp2<tchar>(context.tracker, instance, cmdLine.data(), cmdLine.size(), TC(""), codeFiles))
+					TraverseCodeFiles(context, instance, codeFiles, ruleIndex);
+			}, 1, TC("CrawlCmd"), ColorWork, WorkPriority_Low, &m_bottleneck);
 		return true;
 	}
 
@@ -244,11 +198,12 @@ namespace uba
 		if (CaseInsensitiveFs)
 			fixedPath.MakeLower();
 		bool stopCrawling = false;
-		if (!instance.createFileFunc(tracker, fixedPath, [&](const void* data, u64 dataSize) { return ParseRsp2(tracker, instance, data, dataSize, rsp.data, outCodeFiles); }, stopCrawling))
+		if (!instance.createFileFunc(tracker, fixedPath, [&](const void* data, u64 dataSize) { return ParseRsp2<char>(tracker, instance, data, dataSize, rsp.data, outCodeFiles); }, stopCrawling))
 			return stopCrawling ? false : m_logger.Warning(TC("Failed to parse rsp %s"), rsp.data);
 		return true;
 	}
 
+	template<typename CharType>
 	bool DependencyCrawler::ParseRsp2(TrackWorkScope& tracker, Instance& instance, const void* data, u64 dataSize, const tchar* rsp, CodeFiles& outCodeFiles)
 	{
 		auto AddCodeFile = [&](StringView file, bool pushFront, bool hasPch, bool fixPath = true)
@@ -327,7 +282,7 @@ namespace uba
 			StringBuffer<> prevArg;
 			bool addRoot = false;
 			bool handled = false;
-			ParseArguments((const char*)data, dataSize, [&](const char* arg, u32 argLen)
+			ParseArguments((const CharType*)data, dataSize, [&](const CharType* arg, u32 argLen)
 				{
 					StringBuffer<> sb;
 					sb.Append(arg, argLen);
@@ -415,6 +370,10 @@ namespace uba
 							instance.builtinIncludesDir += TC("/include");
 							UBA_CRAWLER_DEBUG_PRINT("RESOURCEDIR", sb.data + 14);
 						}
+						else if (sb.StartsWith(TCV("-I"), false)) // clang-cl.exe allows this
+						{
+							AddRoot(StringView(sb).Skip(2));
+						}
 						else
 							IgnoreOption(sb);
 					}
@@ -633,6 +592,8 @@ namespace uba
 
 	bool DependencyCrawler::ParseCodeFile2(TrackWorkScope& tracker, Instance& instance, HandledFiles& handledFiles, const void* data, u64 dataSize, const StringView& codeFile, bool parseDefines)
 	{
+		UBA_CRAWLER_DEBUG_PRINT("PARSE", codeFile.data);
+
 		const char* it = (const char*)data;
 		const char* end = it + dataSize;
 		bool hasNonSpace = false;
@@ -931,4 +892,82 @@ namespace uba
 			});
 		return true;
 	}
+
+	inline bool DependencyCrawler::TraverseCodeFiles(const WorkContext& context, Instance& instance, const CodeFiles& codeFiles, u32 ruleIndex)
+	{
+		for (auto& cf : codeFiles)
+		{
+			if (instance.usePch && cf.hasPch)
+			{
+				// Note, clang is always reading all the individual files that has been stored in the pch to validate content..
+				// So this code path is not implemented for clang.. if -fno-validate-ast-input-files-content is added we can implement this path
+				StringBuffer<> fixedFile;
+				FixPath(cf.path.c_str(), instance.workDir.data(), instance.workDir.size(), fixedFile);
+				if (CaseInsensitiveFs)
+					fixedFile.MakeLower();
+				StringKey pchKey = ToStringKey(fixedFile);
+
+				SCOPED_FUTEX(m_pchLookupLock, lock);
+				Pch& pch = m_pchLookup[pchKey];
+				lock.Leave();
+
+				SCOPED_FUTEX(pch.lock, lock2);
+				if (!pch.handled)
+				{
+					fixedFile.Append(TCV(".dep.json"));
+					bool stopCrawling = false;
+					instance.createFileFunc(context.tracker, fixedFile, [&](const void* data, u64 dataSize)
+						{
+							return ParsePch(pch, instance, data, dataSize);
+						}, stopCrawling);
+					if (stopCrawling)
+						return true;
+					pch.handled = true;
+				}
+				instance.pch = &pch;
+			}
+			else if (instance.type == DependencyCrawlerType_MsvcLinker)
+			{
+				// If we end up here it means that we have compressed obj files enabled and want to decompress them in parallel
+				m_workManager.AddWork([this, ref = InstanceRef(instance), codeFile = cf.path](const WorkContext& context)
+					{
+						bool stopCrawling = false;
+						ref.instance.createFileFunc(context.tracker, codeFile, {}, stopCrawling);
+					}, 1, TC("CrawlForDecomp"), ColorWork, WorkPriority_Low, &m_bottleneck);
+			}
+			else
+			{
+				bool parseDefines = instance.platform.empty();
+				HandledFiles handledFiles;
+				ParseCodeFile(context.tracker, instance, parseDefines ? handledFiles : m_handledFiles, cf.path, parseDefines, instance.rsp);
+				if (instance.platform.empty())
+				{
+					if (!instance.overriddenPlatformName.empty())
+						instance.platform = instance.overriddenPlatformName;
+					else
+						instance.platform = instance.compiledPlatform;
+				}
+			}
+		}
+
+		if (!instance.builtinIncludesDir.empty())
+		{
+			SCOPED_FUTEX(m_builtIncludesHandledLock, lock);
+			bool shouldHandle = m_builtIncludesHandled.insert(ruleIndex).second;
+			lock.Leave();
+			if (shouldHandle)
+			{
+				StringBuffer<> fixedPath;
+				FixPath(instance.builtinIncludesDir.c_str(), instance.workDir.data(), instance.workDir.size(), fixedPath);
+				if (!instance.devirtualizePathFunc(fixedPath))
+					UBA_LOG_DEVIRTUALIZATION_ERROR("Failed to devirtualize path %s found in builtin includes", fixedPath.data);
+				if (CaseInsensitiveFs)
+					fixedPath.MakeLower();
+				fixedPath.Append(PathSeparator);
+
+				TraverseInclude(context.tracker, instance, fixedPath);
+			}
+		}
+		return true;
+	}
 }
