[UBA]
* Moved sleep in scheduler since it caused bottleneck

diff --git a/Engine/Source/Programs/UnrealBuildAccelerator/Common/Private/UbaScheduler.cpp b/Engine/Source/Programs/UnrealBuildAccelerator/Common/Private/UbaScheduler.cpp
--- a/Engine/Source/Programs/UnrealBuildAccelerator/Common/Private/UbaScheduler.cpp
+++ b/Engine/Source/Programs/UnrealBuildAccelerator/Common/Private/UbaScheduler.cpp
@@ -933,8 +933,10 @@ namespace uba
 		return entryToProcess;
 	}
 
-	u32 Scheduler::PopProcess(bool isLocal, bool isCrossArchitecture, bool isCrossPlatform, RunStatus& outPrevStatus, bool forReuse)
+	u32 Scheduler::PopProcess(bool isLocal, bool isCrossArchitecture, bool isCrossPlatform, RunStatus& outPrevStatus, bool forReuse, bool& outShouldSleep)
 	{
+		outShouldSleep = false;
+
 		float activeWeight = m_activeLocalProcessWeight;
 		float maxWeight = float(m_maxLocalProcessors);
 		bool atMaxLocalWeight = activeWeight >= maxWeight;
@@ -1100,6 +1102,7 @@ namespace uba
 				if (m_getNextQueueCount)
 				{
 					m_updateThreadWakeup.Set();
+					outShouldSleep = true;
 					return ~0u;
 				}
 
@@ -1237,12 +1240,14 @@ namespace uba
 			if (isCrossArchitecture || isCrossPlatform)
 				m_hadCrossSystemRequest = true;
 
-			u32 indexToRun = PopProcess(isLocal, isCrossArchitecture, isCrossPlatform, prevStatus, false);
+			bool shouldSleep;
+			u32 indexToRun = PopProcess(isLocal, isCrossArchitecture, isCrossPlatform, prevStatus, false, shouldSleep);
 			if (indexToRun == ~0u)
 			{
 				lock.Leave();
 
-				Sleep(1); // We want to sleep a little bit in case this PopProcess failed because there is a process sitting on reuse request
+				if (shouldSleep)
+					Sleep(1); // We want to sleep a little bit in case this PopProcess failed because there is a process sitting on reuse request
 
 				if (!isLocal)
 					return false;
@@ -1386,9 +1391,10 @@ namespace uba
 
 		while (true)
 		{
+			bool shouldSleep;
 			RunStatus prevStatus;
 			SCOPED_FUTEX(m_processEntriesLock, lock);
-			u32 indexToRun = PopProcess(isLocal, false, false, prevStatus, forReuse);
+			u32 indexToRun = PopProcess(isLocal, false, false, prevStatus, forReuse, shouldSleep);
 			if (indexToRun == ~0u)
 			{
 				u64 usedTimeMs = TimeToMs(GetTime() - startTime);
diff --git a/Engine/Source/Programs/UnrealBuildAccelerator/Common/Public/UbaScheduler.h b/Engine/Source/Programs/UnrealBuildAccelerator/Common/Public/UbaScheduler.h
--- a/Engine/Source/Programs/UnrealBuildAccelerator/Common/Public/UbaScheduler.h
+++ b/Engine/Source/Programs/UnrealBuildAccelerator/Common/Public/UbaScheduler.h
@@ -147,7 +147,7 @@ namespace uba
 		void ProcessExited(ExitProcessInfo* info, const ProcessHandle& handle);
 		bool CanRun();
 		u32 PopCache();
-		u32 PopProcess(bool isLocal, bool isCrossArchitecture, bool isCrossPlatform, RunStatus& outPrevStatus, bool forReuse);
+		u32 PopProcess(bool isLocal, bool isCrossArchitecture, bool isCrossPlatform, RunStatus& outPrevStatus, bool forReuse, bool& outShouldSleep);
 		bool RunCacheQuery();
 		bool RunQueuedProcess(bool isLocal, bool isCrossArchitecture, bool isCrossPlatform);
 		bool RaceRemoteProcess();
