[UBA]
* Fixed so detours memory is showing proper numbers on posix in trace log

diff --git a/Engine/Source/Programs/UnrealBuildAccelerator/Detours/Private/PosixOS/UbaDetoursFunctionsPosix.cpp b/Engine/Source/Programs/UnrealBuildAccelerator/Detours/Private/PosixOS/UbaDetoursFunctionsPosix.cpp
--- a/Engine/Source/Programs/UnrealBuildAccelerator/Detours/Private/PosixOS/UbaDetoursFunctionsPosix.cpp
+++ b/Engine/Source/Programs/UnrealBuildAccelerator/Detours/Private/PosixOS/UbaDetoursFunctionsPosix.cpp
@@ -3025,6 +3025,8 @@ namespace uba
 
 		SendInput();
 
+		g_stats.detoursMemory = g_memoryBlock.writtenSize;
+
 		RPC_MESSAGE(Exit, log);
 		writer.WriteU32(0); // Exit code
 		writer.WriteString(""); // Log name
@@ -3083,6 +3085,8 @@ namespace uba
 			if (!allowTerminate)
 				return;
 
+			g_stats.detoursMemory = g_memoryBlock.writtenSize;
+
 			RPC_MESSAGE(Exit, log);
 			writer.WriteU32(terminateCode); // Exit code
 			writer.WriteString(""); // Log name
