[UBA]
* Added more info to fatal memory error

diff --git a/Engine/Source/Programs/UnrealBuildAccelerator/Core/Private/UbaMemory.cpp b/Engine/Source/Programs/UnrealBuildAccelerator/Core/Private/UbaMemory.cpp
--- a/Engine/Source/Programs/UnrealBuildAccelerator/Core/Private/UbaMemory.cpp
+++ b/Engine/Source/Programs/UnrealBuildAccelerator/Core/Private/UbaMemory.cpp
@@ -101,7 +101,7 @@ namespace uba
 		u64 newPos = startPos + bytes;
 		
 		if (newPos > reserveSize)
-			FatalError(9882, TC("Ran out of reserved virtual address space in %s. Reserved %llu, Needed %llu (%s)"), name, reserveSize, newPos, hint);
+			FatalError(9882, TC("Ran out of reserved virtual address space in %s. Requesting %llu, Alignment %llu, Reserved %llu, Need %llu (%s)"), name, bytes, alignment, reserveSize, newPos, hint);
 
 		#if PLATFORM_WINDOWS
 		if (newPos > committedSize)
