[UBA]
* Fix build when UBA_USE_MIMALLOC is 0, TRACK_UCRT_ALLOC_ENABLED needs to be defined to 0 too

Pull request by lclara-BARB, thank you!

diff --git a/Engine/Source/Programs/UnrealBuildAccelerator/Detours/Private/Windows/UbaDetoursFunctionsMiMalloc.inl b/Engine/Source/Programs/UnrealBuildAccelerator/Detours/Private/Windows/UbaDetoursFunctionsMiMalloc.inl
--- a/Engine/Source/Programs/UnrealBuildAccelerator/Detours/Private/Windows/UbaDetoursFunctionsMiMalloc.inl
+++ b/Engine/Source/Programs/UnrealBuildAccelerator/Detours/Private/Windows/UbaDetoursFunctionsMiMalloc.inl
@@ -247,5 +247,5 @@ size_t Detoured__aligned_msize(void* p, size_t alignment, size_t offset)
 //}
 
 #endif
-
+#define TRACK_UCRT_ALLOC_ENABLED 0
 #endif
