[UBAObjTool] Ignore WX when parsing arguments (treat warnings as errors)

Pull request by lclara-BARB, thank you!

diff --git a/Engine/Source/Programs/UnrealBuildAccelerator/ObjTool/Private/UbaObjTool.cpp b/Engine/Source/Programs/UnrealBuildAccelerator/ObjTool/Private/UbaObjTool.cpp
--- a/Engine/Source/Programs/UnrealBuildAccelerator/ObjTool/Private/UbaObjTool.cpp
+++ b/Engine/Source/Programs/UnrealBuildAccelerator/ObjTool/Private/UbaObjTool.cpp
@@ -122,6 +122,9 @@ namespace uba
 					if (name.Equals(TCV("/NOLOGO")))
 					{
 					}
+					else if (name.Equals(TCV("/WX")))
+					{
+					}
 					else if (name.Equals(TCV("/errorReport")))
 					{
 					}
