[UBAVisualizer][Windows] Add support for drag/drop .uba file in the visualizer window

Pull request by lclara-BARB, thank you!

diff --git a/Engine/Source/Programs/UnrealBuildAccelerator/Visualizer/Private/Windows/UbaVisualizerWin.cpp b/Engine/Source/Programs/UnrealBuildAccelerator/Visualizer/Private/Windows/UbaVisualizerWin.cpp
--- a/Engine/Source/Programs/UnrealBuildAccelerator/Visualizer/Private/Windows/UbaVisualizerWin.cpp
+++ b/Engine/Source/Programs/UnrealBuildAccelerator/Visualizer/Private/Windows/UbaVisualizerWin.cpp
@@ -8,6 +8,7 @@
 #include <uxtheme.h>
 #include <dwmapi.h>
 #include <shellscalingapi.h>
+#include <shlobj_core.h>
 
 #pragma comment (lib, "Shcore.lib")
 #pragma comment (lib, "UxTheme.lib")
@@ -358,6 +359,13 @@ namespace uba
 		if (m_config.DarkMode)
 			UpdateTheme();
 
+		// Initialize OLE for Drag and Drop support.
+		HRESULT dragDropRes = ::OleInitialize(NULL);
+		UBA_ASSERT(SUCCEEDED(dragDropRes));
+		dragDropRes = ::RegisterDragDrop(hwnd, this);
+		UBA_ASSERT(SUCCEEDED(dragDropRes));
+
+
 		RECT r;
 		GetClientRect(m_hwnd, &r);
 		m_clientWidth = r.right;
@@ -648,6 +656,125 @@ namespace uba
 		DwmSetWindowAttribute(m_hwnd, attribute, &useDarkMode, sizeof(useDarkMode));
 	}
 
+	HRESULT VisualizerWin::QueryInterface(REFIID riid, void** ppvObject)
+	{
+		if (IID_IDropTarget == riid || IID_IUnknown == riid)
+		{
+			AddRef();
+			*ppvObject = (IDropTarget*)(this);
+			return S_OK;
+		}
+		else
+		{
+			*ppvObject = NULL;
+			return E_NOINTERFACE;
+		}
+	}
+
+	ULONG VisualizerWin::AddRef(void)
+	{
+		return InterlockedIncrement(&m_OLEReferenceCount);
+	}
+
+	ULONG VisualizerWin::Release(void)
+	{
+		InterlockedDecrement(&m_OLEReferenceCount);
+		return m_OLEReferenceCount;
+	}
+
+	HRESULT VisualizerWin::DragEnter(IDataObject* pDataObj, ::DWORD grfKeyState, POINTL pt, ::DWORD* pdwEffect)
+	{
+		if (pDataObj == nullptr)
+			return E_FAIL;
+		if (pdwEffect == nullptr)
+			return E_POINTER;
+		*pdwEffect = DROPEFFECT_COPY;
+		return S_OK;
+	}
+
+	HRESULT VisualizerWin::DragOver(::DWORD grfKeyState, POINTL pt, ::DWORD* pdwEffect)
+	{
+		if (pdwEffect == nullptr)
+			return E_POINTER;
+		*pdwEffect = DROPEFFECT_COPY;
+		return S_OK;
+	}
+
+	HRESULT VisualizerWin::DragLeave(void)
+	{
+		return S_OK;
+	}
+
+	HRESULT VisualizerWin::Drop(IDataObject* pDataObj, ::DWORD grfKeyState, POINTL pt, ::DWORD* pdwEffect)
+	{
+		if (pDataObj == nullptr)
+			return E_FAIL;
+		if (pdwEffect == nullptr)
+			return E_POINTER;
+
+		// Utility to ensure resource release
+		struct FOLEResourceGuard
+		{
+			STGMEDIUM& StorageMedium;
+			LPVOID DataPointer;
+			FOLEResourceGuard(STGMEDIUM& InStorage) : StorageMedium(InStorage), DataPointer(GlobalLock(InStorage.hGlobal)) {}
+			~FOLEResourceGuard() { GlobalUnlock(StorageMedium.hGlobal); ReleaseStgMedium(&StorageMedium); }
+		};
+
+		// Attempt to get plain text or unicode text from the data being dragged in
+		STGMEDIUM storageMedium;
+		FORMATETC FormatEtc_Ansii = { CF_TEXT, 0, DVASPECT_CONTENT, -1, TYMED_HGLOBAL };
+		const bool bHaveAnsiText = (pDataObj->QueryGetData(&FormatEtc_Ansii) == S_OK) ? true : false;
+
+		FORMATETC FormatEtc_UNICODE = { CF_UNICODETEXT, 0, DVASPECT_CONTENT, -1, TYMED_HGLOBAL };
+		const bool bHaveUnicodeText = (pDataObj->QueryGetData(&FormatEtc_UNICODE) == S_OK) ? true : false;
+
+		FORMATETC FormatEtc_File = { CF_HDROP, 0, DVASPECT_CONTENT, -1, TYMED_HGLOBAL };
+		const bool bHaveFiles = (pDataObj->QueryGetData(&FormatEtc_File) == S_OK) ? true : false;
+
+		StringBuffer<MaxPath> filePath;
+		if (bHaveUnicodeText && S_OK == pDataObj->GetData(&FormatEtc_UNICODE, &storageMedium))
+		{
+			FOLEResourceGuard ResourceGuard(storageMedium);
+			filePath.Append(static_cast<TCHAR*>(ResourceGuard.DataPointer));
+		}
+		if (bHaveAnsiText && S_OK == pDataObj->GetData(&FormatEtc_Ansii, &storageMedium))
+		{
+			// TODO: support ansi text
+			// FOLEResourceGuard ResourceGuard(storageMedium);
+			// filePath.Append(static_cast<TCHAR*>(ResourceGuard.DataPointer));
+		}
+		if (bHaveFiles && S_OK == pDataObj->GetData(&FormatEtc_File, &storageMedium))
+		{
+			FOLEResourceGuard ResourceGuard(storageMedium);
+			const DROPFILES* DropFiles = static_cast<DROPFILES*>(ResourceGuard.DataPointer);
+
+			// pFiles is the offset to the beginning of the file list, in bytes
+			LPVOID FileListStart = (BYTE*)ResourceGuard.DataPointer + DropFiles->pFiles;
+
+			if (DropFiles->fWide) // TODO: support ansi filenames
+			{
+				// Unicode filenames
+				// The file list is NULL delimited with an extra NULL character at the end.
+				TCHAR* Pos = static_cast<TCHAR*>(FileListStart);
+				while (Pos[0] != 0)
+				{
+					filePath.Append(Pos);
+					break; // only care about the first one
+				}
+			}
+		}
+
+		if (!filePath.IsEmpty())
+		{
+			m_fileName.Clear();
+			m_fileName.Append(filePath.data);
+			PostNewTrace(false, 0);
+		}
+
+		return S_OK;
+	}
+
 	void VisualizerWin::PostNewTrace(u32 replay, bool paused)
 	{
 		KillTimer(m_hwnd, 0);
diff --git a/Engine/Source/Programs/UnrealBuildAccelerator/Visualizer/Private/Windows/UbaVisualizerWin.h b/Engine/Source/Programs/UnrealBuildAccelerator/Visualizer/Private/Windows/UbaVisualizerWin.h
--- a/Engine/Source/Programs/UnrealBuildAccelerator/Visualizer/Private/Windows/UbaVisualizerWin.h
+++ b/Engine/Source/Programs/UnrealBuildAccelerator/Visualizer/Private/Windows/UbaVisualizerWin.h
@@ -5,12 +5,13 @@
 #include "UbaNetworkClient.h"
 #include "UbaTraceReader.h"
 #include "UbaVisualizer.h"
+#include <Ole2.h>
 
 namespace uba
 {
 	////////////////////////////////////////////////////////////////////////////////////////////////////
 
-	class VisualizerWin : public Visualizer
+	class VisualizerWin : public Visualizer, public IDropTarget
 	{
 	public:
 		VisualizerWin(VisualizerConfig& config, Logger& logger);
@@ -44,6 +45,18 @@ namespace uba
 		virtual void StopUpdate() override;
 		void UpdateTheme();
 
+		// IUnknown interface
+		virtual HRESULT STDMETHODCALLTYPE QueryInterface(REFIID riid, void** ppvObject) override;
+		virtual ULONG STDMETHODCALLTYPE AddRef(void) override;
+		virtual ULONG STDMETHODCALLTYPE Release(void) override;
+		// IDropTarget interface
+		virtual HRESULT STDMETHODCALLTYPE DragEnter(__RPC__in_opt IDataObject* pDataObj, ::DWORD grfKeyState, POINTL pt, __RPC__inout::DWORD* pdwEffect) override;
+		virtual HRESULT STDMETHODCALLTYPE DragOver(::DWORD grfKeyState, POINTL pt, __RPC__inout::DWORD* pdwEffect) override;
+		virtual HRESULT STDMETHODCALLTYPE DragLeave(void) override;
+		virtual HRESULT STDMETHODCALLTYPE Drop(__RPC__in_opt IDataObject* pDataObj, ::DWORD grfKeyState, POINTL pt, __RPC__inout::DWORD* pdwEffect) override;
+
+		ULONG m_OLEReferenceCount = 0;
+
 		u64 m_lastPaintTimeMs = 0;
 
 		Atomic<bool> m_looping = false;
