[UBA]
* Bumped OverlayTableMemSize from 256kb to 16mb because we hit max while linking ltcg/lto of our biggest binaries (which ended up reaching 573kb)

diff --git a/Engine/Source/Programs/UnrealBuildAccelerator/Core/Public/UbaProtocol.h b/Engine/Source/Programs/UnrealBuildAccelerator/Core/Public/UbaProtocol.h
--- a/Engine/Source/Programs/UnrealBuildAccelerator/Core/Public/UbaProtocol.h
+++ b/Engine/Source/Programs/UnrealBuildAccelerator/Core/Public/UbaProtocol.h
@@ -57,7 +57,7 @@ namespace uba
 
 	inline constexpr u32 FileMappingTableMemSize = 16 * 1024 * 1024;
 	inline constexpr u32 DirTableMemSize = 128 * 1024 * 1024;
-	inline constexpr u32 OverlayTableMemSize = 256 * 1024;
+	inline constexpr u32 OverlayTableMemSize = 16 * 1024 * 1024; // Seen ~600kb max size for ltcg/thinlto linking of our biggest binaries (one temp file per .o file)
 
 	inline constexpr u32 EmergencyShutdownExitCode = 9886;
 }
