[UBA]
* Fixed bugs in StringView GetFileName and GetPath when StringView is not terminated
* Changed const StringView& to StringView since it is small and more efficient to send by-value

diff --git a/Engine/Source/Programs/UnrealBuildAccelerator/Core/Private/UbaStringBuffer.cpp b/Engine/Source/Programs/UnrealBuildAccelerator/Core/Private/UbaStringBuffer.cpp
--- a/Engine/Source/Programs/UnrealBuildAccelerator/Core/Private/UbaStringBuffer.cpp
+++ b/Engine/Source/Programs/UnrealBuildAccelerator/Core/Private/UbaStringBuffer.cpp
@@ -165,7 +165,7 @@ namespace uba
 		return *this;
 	}
 
-	StringBufferBase& StringBufferBase::Append(const StringView& view)
+	StringBufferBase& StringBufferBase::Append(StringView view)
 	{
 		return Append(view.data, view.count);
 	}
@@ -284,7 +284,7 @@ namespace uba
 		return *this;
 	}
 
-	StringBufferBase& StringBufferBase::Prepend(const StringView& view, u32 overwriteCount)
+	StringBufferBase& StringBufferBase::Prepend(StringView view, u32 overwriteCount)
 	{
 		UBA_ASSERT(count + view.count - overwriteCount < capacity);
 		memmove(data+view.count - overwriteCount, data, (count+1)*sizeof(tchar));
@@ -335,7 +335,7 @@ namespace uba
 		return TStrchr(data, c) != nullptr;
 	}
 
-	bool StringBufferBase::StartsWith(const StringView& str, bool ignoreCase) const
+	bool StringBufferBase::StartsWith(StringView str, bool ignoreCase) const
 	{
 		return count >= str.count && uba::StartsWith(data, str.data, ignoreCase);
 	}
@@ -345,12 +345,12 @@ namespace uba
 		return uba::EndsWith(data, count, value, TStrlen(value), ignoreCase);
 	}
 
-	bool StringBufferBase::EndsWith(const StringView& value, bool ignoreCase) const
+	bool StringBufferBase::EndsWith(StringView value, bool ignoreCase) const
 	{
 		return uba::EndsWith(data, count, value.data, value.count, ignoreCase);
 	}
 
-	bool StringBufferBase::Equals(const StringView& str, bool ignoreCase) const
+	bool StringBufferBase::Equals(StringView str, bool ignoreCase) const
 	{
 		return count == str.count && uba::Equals(data, str.data, count, ignoreCase);
 	}
@@ -485,7 +485,7 @@ namespace uba
 		return false;
 	}
 
-	bool StringView::StartsWith(const StringView& str, bool ignoreCase) const
+	bool StringView::StartsWith(StringView str, bool ignoreCase) const
 	{
 		return count >= str.count && uba::StartsWith(data, str.data, ignoreCase);
 	}
@@ -495,10 +495,19 @@ namespace uba
 		return uba::EndsWith(data, count, value, TStrlen(value), ignoreCase);
 	}
 
+	StringView StringView::GetFileName() const
+	{
+		for (const tchar* i=data+count-1, *e=data; i!=e; --i)
+			if (*i == PathSeparator)
+				return StringView(i+1, u32(data+count-i-1));
+		return *this;
+	}
+
 	StringView StringView::GetPath() const
 	{
-		if (const tchar* lastSeparator = TStrrchr(data, PathSeparator))
-			return StringView(data, u32(lastSeparator - data));
+		for (const tchar* i=data+count-1, *e=data; i!=e; --i)
+			if (*i == PathSeparator)
+				return StringView(data, u32(i-data));
 		return {};
 	}
 
diff --git a/Engine/Source/Programs/UnrealBuildAccelerator/Core/Public/UbaStringBuffer.h b/Engine/Source/Programs/UnrealBuildAccelerator/Core/Public/UbaStringBuffer.h
--- a/Engine/Source/Programs/UnrealBuildAccelerator/Core/Public/UbaStringBuffer.h
+++ b/Engine/Source/Programs/UnrealBuildAccelerator/Core/Public/UbaStringBuffer.h
@@ -33,7 +33,7 @@ namespace uba
 		StringBufferBase& Append(const tchar* str);
 		StringBufferBase& Append(const tchar c) { return Append(&c, 1); }
 		StringBufferBase& Append(const tchar* str, u64 charCount);
-		StringBufferBase& Append(const StringView& view);
+		StringBufferBase& Append(StringView view);
 		StringBufferBase& Appendf(const tchar* format, ...);
 		StringBufferBase& AppendDir(const StringBufferBase& str);
 		StringBufferBase& AppendDir(const tchar* dir);
@@ -42,7 +42,7 @@ namespace uba
 		StringBufferBase& AppendBase62(u64 v);
 		StringBufferBase& AppendValue(u64 v);
 		StringBufferBase& Append(const tchar* format, va_list& args);
-		StringBufferBase& Prepend(const StringView& view, u32 overwriteCount = 0);
+		StringBufferBase& Prepend(StringView view, u32 overwriteCount = 0);
 		StringBufferBase& Resize(u64 newSize);
 		StringBufferBase& Clear();
 
@@ -59,13 +59,13 @@ namespace uba
 
 		bool IsEmpty() const { return count == 0; }
 		bool StartsWith(const tchar* str, bool ignoreCase = true) const { return uba::StartsWith(data, str, ignoreCase); }
-		bool StartsWith(const StringView& str, bool ignoreCase = true) const;
+		bool StartsWith(StringView str, bool ignoreCase = true) const;
 		bool EndsWith(const tchar* value, bool ignoreCase = true) const;
-		bool EndsWith(const StringView& value, bool ignoreCase = true) const;
+		bool EndsWith(StringView value, bool ignoreCase = true) const;
 		bool Contains(tchar c) const;
 		bool Contains(const tchar* str, bool ignoreCase = true, const tchar** pos = nullptr) const { return uba::Contains(data, str, ignoreCase, pos); }
 		bool Equals(const tchar* str, bool ignoreCase = true) const { return uba::Equals(data, str, ignoreCase); }
-		bool Equals(const StringView& str, bool ignoreCase = true) const;
+		bool Equals(StringView str, bool ignoreCase = true) const;
 		const tchar* First(tchar c, u64 offset = 0) const;
 		const tchar* Last(tchar c, u64 offset = 0) const;
 		const tchar* GetFileName() const { return uba::GetFileName(data); }
@@ -104,7 +104,7 @@ namespace uba
 		explicit StringBuffer(const TString& str) : StringBufferBase(Capacity) { *buf = 0; Append(str); }
 		explicit StringBuffer(const tchar* str) : StringBufferBase(Capacity) { *buf = 0; if (str) Append(str); }
 		StringBuffer(const StringBufferBase& str) : StringBufferBase(Capacity) { *buf = 0; Append(str); }
-		StringBuffer(const StringView& str) : StringBufferBase(Capacity) { *buf = 0; Append(str); }
+		StringBuffer(StringView str) : StringBufferBase(Capacity) { *buf = 0; Append(str); }
 	private:
 		tchar buf[Capacity];
 	};
@@ -125,16 +125,16 @@ namespace uba
 		tchar operator[](u64 i) const { return data[i]; }
 
 		bool IsEmpty() const { return count == 0; }
-		bool StartsWith(const StringView& str, bool ignoreCase = true) const;
-		bool StartsWith(const tchar* str, bool ignoreCase = true) const { return uba::StartsWith(data, str, ignoreCase); }
-		bool EndsWith(const StringView& value, bool ignoreCase = true) const { return uba::EndsWith(data, count, value.data, value.count, ignoreCase); }
+		bool StartsWith(StringView str, bool ignoreCase = true) const; // Requires terminated StringViews
+		bool StartsWith(const tchar* str, bool ignoreCase = true) const { return uba::StartsWith(data, str, ignoreCase); } // Requires terminated StringView
+		bool EndsWith(StringView value, bool ignoreCase = true) const { return uba::EndsWith(data, count, value.data, value.count, ignoreCase); }
 		bool EndsWith(const tchar* value, bool ignoreCase = true) const;
 		bool Contains(tchar c, u32* outPos = nullptr) const;
-		bool Contains(const tchar* str, bool ignoreCase = true) const { return uba::Contains(data, str, ignoreCase); }
-		bool Contains(const StringView& str, bool ignoreCase = true) const { return uba::Contains(data, str.data, ignoreCase); }
-		bool Equals(const tchar* str, bool ignoreCase = true) const { return uba::Equals(data, str, ignoreCase); }
-		bool Equals(const StringView& str, bool ignoreCase = true) const { return count == str.count && uba::Equals(data, str.data, count, ignoreCase); }
-		StringView GetFileName() const { const tchar* str = uba::GetFileName(data); return StringView(str, count - u32(str - data)); }
+		bool Contains(const tchar* str, bool ignoreCase = true) const { return uba::Contains(data, str, ignoreCase); } // Requires terminated StringViews
+		bool Contains(StringView str, bool ignoreCase = true) const { return uba::Contains(data, str.data, ignoreCase); } // Requires terminated StringViews
+		bool Equals(const tchar* str, bool ignoreCase = true) const { return uba::Equals(data, str, ignoreCase); } // Requires terminated StringViews
+		bool Equals(StringView str, bool ignoreCase = true) const { return count == str.count && uba::Equals(data, str.data, count, ignoreCase); }
+		StringView GetFileName() const;
 		StringView GetPath() const;
 		StringView Skip(u32 skipCount) const { return StringView(data + skipCount, count - skipCount); }
 		StringView SubStr(u32 startOffset, u32 length) const { return StringView(data + startOffset, length); }
