[UBA]
* Added FixPath function that uses StringViews

diff --git a/Engine/Source/Programs/UnrealBuildAccelerator/Core/Public/UbaPathUtils.h b/Engine/Source/Programs/UnrealBuildAccelerator/Core/Public/UbaPathUtils.h
--- a/Engine/Source/Programs/UnrealBuildAccelerator/Core/Public/UbaPathUtils.h
+++ b/Engine/Source/Programs/UnrealBuildAccelerator/Core/Public/UbaPathUtils.h
@@ -382,6 +382,14 @@ namespace uba
 		return res;
 	}
 
+	inline bool FixPath(StringBufferBase& out, StringView fileName, StringView workingDir)
+	{
+		u32 count = out.count;
+		bool res = FixPath2(fileName.data, fileName.count, workingDir.data, workingDir.count, out.data + count, out.capacity - count, &out.count);
+		out.count += count;
+		return res;
+	}
+
 	inline bool IsAbsolutePath(const tchar* path)
 	{
 		if (!path)
