[UBA]
* Minor restructuring in file to make easier to find functions

diff --git a/Engine/Source/Programs/UnrealBuildAccelerator/Core/Public/UbaBinaryReaderWriter.h b/Engine/Source/Programs/UnrealBuildAccelerator/Core/Public/UbaBinaryReaderWriter.h
--- a/Engine/Source/Programs/UnrealBuildAccelerator/Core/Public/UbaBinaryReaderWriter.h
+++ b/Engine/Source/Programs/UnrealBuildAccelerator/Core/Public/UbaBinaryReaderWriter.h
@@ -7,8 +7,17 @@
 
 namespace uba
 {
+	////////////////////////////////////////////////////////////////////////////////////////////////////
+
 	enum ProcessCommunication { ProcessCommunication_Write };
 
+	template<typename CharType>
+	inline u64 GetUtf8WrittenBytes(const CharType* str, u64 strlen);
+	inline u8 Get7BitEncodedCount(u64 value);
+	inline u64 GetStringWriteSize(const tchar* str, u64 len);
+
+	////////////////////////////////////////////////////////////////////////////////////////////////////
+
 	struct BinaryWriter
 	{
 		inline void WriteByte(u8 value);
@@ -64,6 +73,8 @@ namespace uba
 		u8 buffer[Capacity];
 	};
 
+	////////////////////////////////////////////////////////////////////////////////////////////////////
+
 	struct BinaryReader
 	{
 		inline void ReadBytes(void* data, u64 size);
@@ -121,31 +132,11 @@ namespace uba
 		u8 buffer[Capacity];
 	};
 
-	// Implementations
+	////////////////////////////////////////////////////////////////////////////////////////////////////
 
 	#define UBA_ASSERT_WRITE(size) UBA_ASSERTF(m_pos + size <= m_end, TC("BinaryWriter overflow. Written: %llu, Capacity: %llu, Trying to write: %llu"), u64(m_pos - m_begin), u64(m_end - m_begin), u64(size))
 	#define UBA_ASSERT_READ(size) UBA_ASSERTF(m_pos + size <= m_end, TC("BinaryReader overflow. Read: %llu, Size: %llu, Trying to read: %llu"), u64(m_pos - m_begin), u64(m_end - m_begin), u64(size))
 
-	void BinaryWriter::WriteByte(u8 value)
-	{
-		UBA_ASSERT_WRITE(1);
-		*m_pos = value;
-		++m_pos;
-	}
-
-	void BinaryWriter::WriteBytes(const void* data, u64 size)
-	{
-		UBA_ASSERT_WRITE(size);
-		memcpy(m_pos, data, size);
-		m_pos += size;
-	}
-
-	void BinaryWriter::WriteString(const tchar* str)
-	{
-		UBA_ASSERT(str);
-		WriteString(str, u64(TStrlen(str)));
-	}
-
 	template<typename CharType>
 	inline u64 GetUtf8WrittenBytes(const CharType* str, u64 strlen)
 	{
@@ -184,6 +175,28 @@ namespace uba
 		return Get7BitEncodedCount(len) + actualBytes;
 	}
 
+	////////////////////////////////////////////////////////////////////////////////////////////////////
+
+	void BinaryWriter::WriteByte(u8 value)
+	{
+		UBA_ASSERT_WRITE(1);
+		*m_pos = value;
+		++m_pos;
+	}
+
+	void BinaryWriter::WriteBytes(const void* data, u64 size)
+	{
+		UBA_ASSERT_WRITE(size);
+		memcpy(m_pos, data, size);
+		m_pos += size;
+	}
+
+	void BinaryWriter::WriteString(const tchar* str)
+	{
+		UBA_ASSERT(str);
+		WriteString(str, u64(TStrlen(str)));
+	}
+
 	template<typename CharType>
 	void BinaryWriter::WriteUtf8String(const CharType* str, u64 len)
 	{
@@ -311,6 +324,8 @@ namespace uba
 		return data;
 	}
 
+	////////////////////////////////////////////////////////////////////////////////////////////////////
+
 	void BinaryReader::ReadBytes(void* data, u64 size)
 	{
 		UBA_ASSERT_READ(size);
@@ -622,4 +637,5 @@ namespace uba
 		return *(u64*)m_pos;
 	}
 
+	////////////////////////////////////////////////////////////////////////////////////////////////////
 }
