[UBA]
* Added "final" to a bunch of classes/methods
* Various fixes to make pvs studio happy

diff --git a/Engine/Source/Programs/UnrealBuildAccelerator/Common/Private/UbaNetworkBackendTcp.cpp b/Engine/Source/Programs/UnrealBuildAccelerator/Common/Private/UbaNetworkBackendTcp.cpp
--- a/Engine/Source/Programs/UnrealBuildAccelerator/Common/Private/UbaNetworkBackendTcp.cpp
+++ b/Engine/Source/Programs/UnrealBuildAccelerator/Common/Private/UbaNetworkBackendTcp.cpp
@@ -1100,7 +1100,7 @@ namespace uba
 		if (!shouldPoll)
 		{
 			#if UBA_LOG_SOCKET_ERRORS
-			logger.Info(TC("SendSocket - send returned an error for socket %i: %s (%s)"), socket, LastErrorToText(WSAGetLastError()).data, hint);
+			logger.Info(TC("SendSocket - send returned an error for socket %llu: %s (%s)"), socket.internal, LastErrorToText(WSAGetLastError()).data, hint);
 			#endif
 			return false;
 		}
@@ -1120,23 +1120,23 @@ namespace uba
 						return false;
 					++timeoutCounter;
 					if (timeoutCounter == 4)
-						logger.Info(TC("SendSocket - WSAWaitForMultipleEvents returned timeout after 40 seconds for socket %i (%s)"), socket, hint);
+						logger.Info(TC("SendSocket - WSAWaitForMultipleEvents returned timeout after 40 seconds for socket %llu (%s)"), socket.internal, hint);
 					if (timeoutCounter < 6*10) // Loop for 10 minutes
 						continue;
-					logger.Info(TC("SendSocket - WSAWaitForMultipleEvents returned timeout after 10 minutes for socket %i (%s)"), socket, hint);
+					logger.Info(TC("SendSocket - WSAWaitForMultipleEvents returned timeout after 10 minutes for socket %llu (%s)"), socket.internal, hint);
 					return false;
 				}
 				if (wr == WSA_WAIT_FAILED)
 				{
 					#if UBA_LOG_SOCKET_ERRORS
-					logger.Info(TC("SendSocket - WSAWaitForMultipleEvents returned an error for socket %i: %s (%s)"), socket, LastErrorToText(WSAGetLastError()).data, hint);
+					logger.Info(TC("SendSocket - WSAWaitForMultipleEvents returned an error for socket %llu: %s (%s)"), socket.internal, LastErrorToText(WSAGetLastError()).data, hint);
 					#endif
 				}
 				WSANETWORKEVENTS ne{};
 				if (WSAEnumNetworkEvents(ToPlatformSocket(socket), connection.wsaEvent, &ne) == SOCKET_ERROR)
 				{
 					#if UBA_LOG_SOCKET_ERRORS
-					logger.Info(TC("SendSocket - WSAEnumNetworkEvents returned an error for socket %i: %s (%s)"), socket, LastErrorToText(WSAGetLastError()).data, hint);
+					logger.Info(TC("SendSocket - WSAEnumNetworkEvents returned an error for socket %llu: %s (%s)"), socket.internal, LastErrorToText(WSAGetLastError()).data, hint);
 					#endif
 					return false;
 				}
@@ -1159,23 +1159,23 @@ namespace uba
 						return false;
 					++timeoutCounter;
 					if (timeoutCounter == 4)
-						logger.Info(TC("SendSocket - WSAPoll returned timeout after 40 seconds for socket %i (%s)"), socket, hint);
+						logger.Info(TC("SendSocket - WSAPoll returned timeout after 40 seconds for socket %llu (%s)"), socket.internal, hint);
 					if (timeoutCounter < 6*10) // Loop for 10 minutes
 						continue;
-					logger.Info(TC("SendSocket - WSAPoll returned timeout after 10 minutes for socket %i (%s)"), socket, hint);
+					logger.Info(TC("SendSocket - WSAPoll returned timeout after 10 minutes for socket %llu (%s)"), socket.internal, hint);
 					return false;
 				}
 				if (res < 0)
 				{
 					#if UBA_LOG_SOCKET_ERRORS
-					logger.Info(TC("SendSocket - WSAPoll returned an error for socket %i: %s (%s)"), socket, LastErrorToText(WSAGetLastError()).data, hint);
+					logger.Info(TC("SendSocket - WSAPoll returned an error for socket %llu: %s (%s)"), socket.internal, LastErrorToText(WSAGetLastError()).data, hint);
 					#endif
 					return false;
 				}
 				if (revents & (PollERR | PollHUP | PollNVAL))
 				{
 					#if UBA_LOG_SOCKET_ERRORS
-					logger.Info(TC("SendSocket - WSAPoll revents contained error flags for socket %i: %s (%s)"), socket, LastErrorToText(WSAGetLastError()).data, hint);
+					logger.Info(TC("SendSocket - WSAPoll revents contained error flags for socket %llu: %s (%s)"), socket.internal, LastErrorToText(WSAGetLastError()).data, hint);
 					#endif
 					return false;
 				}
@@ -1258,7 +1258,7 @@ namespace uba
 					if (lastError != WSA_IO_PENDING)
 					{
 						#if UBA_LOG_SOCKET_ERRORS
-						logger.Info(TC("WSASend - error for socket %i: %s (%s)"), socket, LastErrorToText(lastError).data, hint);
+						logger.Info(TC("WSASend - error for socket %llu: %s (%s)"), socket.internal, LastErrorToText(lastError).data, hint);
 						#endif
 						return false;
 					}
@@ -1272,7 +1272,7 @@ namespace uba
 
 					if (ev.IsSet(2*1000))
 						break;
-					logger.Info(TC("SendSocket - WSASend returned timeout after 40 seconds for socket %i (%s)"), socket, hint);
+					logger.Info(TC("SendSocket - WSASend returned timeout after 40 seconds for socket %llu (%s)"), socket.internal, hint);
 					u64 startTime = GetTime();
 					while (true)
 					{
@@ -1285,7 +1285,7 @@ namespace uba
 						u64 sinceStartSeconds = TimeToMs(GetTime() - startTime)/1000;
 						if (sinceStartSeconds >= DefaultNetworkSendTimeoutSeconds)
 						{
-							logger.Info(TC("SendSocket - WSASend returned timeout after 10 minutes for socket %i (%s)"), socket, hint);
+							logger.Info(TC("SendSocket - WSASend returned timeout after 10 minutes for socket %llu (%s)"), socket.internal, hint);
 							return false;
 						}
 					}
@@ -1296,7 +1296,7 @@ namespace uba
 				if (!WSAGetOverlappedResult(ToPlatformSocket(socket), &overlapped, &bytesSent, FALSE, &flags))
 				{
 					#if UBA_LOG_SOCKET_ERRORS
-					logger.Info(TC("WSAGetOverlappedResult - error for socket %i: %s (%s)"), socket, LastErrorToText().data, hint);
+					logger.Info(TC("WSAGetOverlappedResult - error for socket %llu: %s (%s)"), socket.internal, LastErrorToText().data, hint);
 					#endif
 					return false;
 				}
@@ -1307,7 +1307,7 @@ namespace uba
 				{
 					#if !UBA_USE_OVERLAPPED_SEND_WITH_LOCK
 					//#if UBA_LOG_SOCKET_ERRORS
-					logger.Warning(TC("SendSocket - WSASend did not send all data in one call for socket %i. Note that this might be allowed behavior. Hard to read out of documentation. (%s)"), socket, hint);
+					logger.Warning(TC("SendSocket - WSASend did not send all data in one call for socket %llu. Note that this might be allowed behavior. Hard to read out of documentation. (%s)"), socket.internal, hint);
 					//#endif
 					return false;
 					#else
@@ -1403,7 +1403,7 @@ namespace uba
 			if (read == 0)
 			{
 				#if UBA_LOG_SOCKET_ERRORS
-				connection.logger.Info(TC("RecvSocket - recv gracefully closed by peer for socket %i and connection %s (%s%s)"), socket, GuidToString(connection.uid).str, connection.recvHint, hint);
+				connection.logger.Info(TC("RecvSocket - recv gracefully closed by peer for socket %llu and connection %s (%s%s)"), socket.internal, GuidToString(connection.uid).str, connection.recvHint, hint);
 				#endif
 				return false;
 			}
@@ -1454,20 +1454,20 @@ namespace uba
 								break;
 							return false;
 						}
-						connection.logger.Info(TC("RecvSocket - WSAPoll returned timeout for socket %i and connection %s after %s (%s%s)"), socket, GuidToString(connection.uid).str, TimeToText(MsToTime(timeoutMs)).str, connection.recvHint, hint);
+						connection.logger.Info(TC("RecvSocket - WSAPoll returned timeout for socket %llu and connection %s after %s (%s%s)"), socket.internal, GuidToString(connection.uid).str, TimeToText(MsToTime(timeoutMs)).str, connection.recvHint, hint);
 						return false;
 					}
 					if (res < 0)
 					{
 						#if UBA_LOG_SOCKET_ERRORS
-						connection.logger.Info(TC("RecvSocket - WSAPoll returned an err-or for socket %i and connection %s: %s (%s%s)"), socket, GuidToString(connection.uid).str, LastErrorToText(WSAGetLastError()).data, connection.recvHint, hint);
+						connection.logger.Info(TC("RecvSocket - WSAPoll returned an err-or for socket %llu and connection %s: %s (%s%s)"), socket.internal, GuidToString(connection.uid).str, LastErrorToText(WSAGetLastError()).data, connection.recvHint, hint);
 						#endif
 						return false;
 					}
 					if ((revents & PollIN) == 0)
 					{
 						#if UBA_LOG_SOCKET_ERRORS
-						connection.logger.Info(TC("RecvSocket - WSAPoll revents contained error for socket %i and connection %s: %s (%s%s)"), socket, GuidToString(connection.uid).str, LastErrorToText(WSAGetLastError()).data, connection.recvHint, hint);
+						connection.logger.Info(TC("RecvSocket - WSAPoll revents contained error for socket %llu and connection %s: %s (%s%s)"), socket.internal, GuidToString(connection.uid).str, LastErrorToText(WSAGetLastError()).data, connection.recvHint, hint);
 						#endif
 						return false;
 					}
@@ -1479,11 +1479,11 @@ namespace uba
 
 			#if !PLATFORM_WINDOWS
 			if (!isFirstCall && errno != ECONNRESET)
-				return connection.logger.Warning(TC("RecvSocket - recv err-or on socket %i and connection %s: %s (%s%s)"), socket, GuidToString(connection.uid).str, strerror(errno), connection.recvHint, hint);
+				return connection.logger.Warning(TC("RecvSocket - recv err-or on socket %llu and connection %s: %s (%s%s)"), socket.internal, GuidToString(connection.uid).str, strerror(errno), connection.recvHint, hint);
 			#endif
 
 			#if UBA_LOG_SOCKET_ERRORS
-			connection.logger.Info(TC("RecvSocket - read returned an err-or for socket %i and connection %s: %s (%s%s)"), socket, GuidToString(connection.uid).str, LastErrorToText(WSAGetLastError()).data, connection.recvHint, hint);
+			connection.logger.Info(TC("RecvSocket - read returned an err-or for socket %llu and connection %s: %s (%s%s)"), socket.internal, GuidToString(connection.uid).str, LastErrorToText(WSAGetLastError()).data, connection.recvHint, hint);
 			#endif
 			return false;
 		}
@@ -1633,7 +1633,7 @@ namespace uba
 		int ret = WSARecv(ToPlatformSocket(socket), &connection.wsaBuf, 1, NULL, &flags, &connection.overlapped, NULL);
 		if (ret == SOCKET_ERROR && WSAGetLastError() != WSA_IO_PENDING)
 		{
-			connection.logger.Info(TC("WSARecv failed for socket %i trying to receive %u bytes (%s)"), socket, dataSize, LastErrorToText(WSAGetLastError()).data);
+			connection.logger.Info(TC("WSARecv failed for socket %llu trying to receive %u bytes (%s)"), socket.internal, dataSize, LastErrorToText(WSAGetLastError()).data);
 			PostQueuedCompletionStatus(m_iocpHandle, 0, (ULONG_PTR)&connection, &connection.overlapped );
 			return false;
 		}
diff --git a/Engine/Source/Programs/UnrealBuildAccelerator/Common/Public/UbaNetworkBackendMemory.h b/Engine/Source/Programs/UnrealBuildAccelerator/Common/Public/UbaNetworkBackendMemory.h
--- a/Engine/Source/Programs/UnrealBuildAccelerator/Common/Public/UbaNetworkBackendMemory.h
+++ b/Engine/Source/Programs/UnrealBuildAccelerator/Common/Public/UbaNetworkBackendMemory.h
@@ -8,7 +8,7 @@
 
 namespace uba
 {
-	class NetworkBackendMemory : public NetworkBackend
+	class NetworkBackendMemory final : public NetworkBackend
 	{
 	public:
 		NetworkBackendMemory(LogWriter& writer, const tchar* prefix = TC("NetworkBackendMemory"));
diff --git a/Engine/Source/Programs/UnrealBuildAccelerator/Common/Public/UbaNetworkBackendTcp.h b/Engine/Source/Programs/UnrealBuildAccelerator/Common/Public/UbaNetworkBackendTcp.h
--- a/Engine/Source/Programs/UnrealBuildAccelerator/Common/Public/UbaNetworkBackendTcp.h
+++ b/Engine/Source/Programs/UnrealBuildAccelerator/Common/Public/UbaNetworkBackendTcp.h
@@ -29,7 +29,7 @@ namespace uba
 
 	////////////////////////////////////////////////////////////////////////////////////////////////////
 
-	class NetworkBackendTcp : public NetworkBackend
+	class NetworkBackendTcp final : public NetworkBackend
 	{
 	public:
 		NetworkBackendTcp(const NetworkBackendTcpCreateInfo& info, const tchar* prefix = TC("NetworkBackendTcp"));
diff --git a/Engine/Source/Programs/UnrealBuildAccelerator/Common/Public/UbaSession.h b/Engine/Source/Programs/UnrealBuildAccelerator/Common/Public/UbaSession.h
--- a/Engine/Source/Programs/UnrealBuildAccelerator/Common/Public/UbaSession.h
+++ b/Engine/Source/Programs/UnrealBuildAccelerator/Common/Public/UbaSession.h
@@ -261,7 +261,7 @@ namespace uba
 		Atomic<u64> m_reuseCount = 0;
 
 		ReaderWriterLock m_fileMappingTableMemLock;
-		u8* m_fileMappingTableMem;
+		u8* m_fileMappingTableMem = nullptr;
 		u32 m_fileMappingTableSize = 0;
 		Futex m_fileMappingTableLookupLock;
 		struct FileMappingEntry
diff --git a/Engine/Source/Programs/UnrealBuildAccelerator/Common/Public/UbaStorage.h b/Engine/Source/Programs/UnrealBuildAccelerator/Common/Public/UbaStorage.h
--- a/Engine/Source/Programs/UnrealBuildAccelerator/Common/Public/UbaStorage.h
+++ b/Engine/Source/Programs/UnrealBuildAccelerator/Common/Public/UbaStorage.h
@@ -157,45 +157,45 @@ namespace uba
 		bool CheckFileTable(const tchar* searchPath);
 		const tchar* GetTempPath();
 
-		virtual bool SaveCasTable(bool deleteIsRunningfile, bool deleteDropped = true) override;
-		virtual u64 GetStorageCapacity() override;
-		virtual u64 GetStorageUsed() override;
+		virtual bool SaveCasTable(bool deleteIsRunningfile, bool deleteDropped = true) override final;
+		virtual u64 GetStorageCapacity() override final;
+		virtual u64 GetStorageUsed() override final;
 		virtual bool GetZone(StringBufferBase& out) override;
-		virtual bool Reset() override;
+		virtual bool Reset() override final;
 		bool DeleteAllCas();
 		bool IsCasTableLoaded() const { return m_casTableLoaded; }
 
 		virtual bool StoreCompressed() const final { return m_storeCompressed; }
 		virtual void PrintSummary(Logger& logger) override;
 
-		virtual bool DecompressFileToMemory(const tchar* fileName, FileHandle fileHandle, u8* dest, u64 decompressedSize, const tchar* writeHint, u64 fileStartOffset) override;
-		virtual bool CreateDirectory(const tchar* dir, bool* outAlreadyExists = nullptr) override;
-		virtual bool DeleteCasForFile(const tchar* file) override;
-		virtual bool GetFileSize(u64& outSize, const CasKey& casKey, const tchar* hint) override;
-		virtual bool GetFileData(void* outData, const CasKey& casKey, const tchar* hint) override;
+		virtual bool DecompressFileToMemory(const tchar* fileName, FileHandle fileHandle, u8* dest, u64 decompressedSize, const tchar* writeHint, u64 fileStartOffset) override final;
+		virtual bool CreateDirectory(const tchar* dir, bool* outAlreadyExists = nullptr) override final;
+		virtual bool DeleteCasForFile(const tchar* file) override final;
+		virtual bool GetFileSize(u64& outSize, const CasKey& casKey, const tchar* hint) override final;
+		virtual bool GetFileData(void* outData, const CasKey& casKey, const tchar* hint) override final;
 		virtual bool RetrieveCasFile(RetrieveResult& out, const CasKey& casKey, StringKey keyHint, const tchar* hint, FileMappingBuffer* mappingBuffer = nullptr, bool createIndependentMapping = false, u64 memoryMapAlignment = 1, bool allowProxy = true, u32 clientId = 0) override;
-		virtual bool AddRef(const CasKey& casKey) override;
-		virtual bool ReleaseRef(const CasKey& casKey, bool dropIfZero) override;
-		virtual bool VerifyAndGetCachedFileInfo(CachedFileInfo& out, StringKey fileNameKey, u64 verifiedLastWriteTime, u64 verifiedSize) override;
-		virtual bool InvalidateCachedFileInfo(StringKey fileNameKey) override;
+		virtual bool AddRef(const CasKey& casKey) override final;
+		virtual bool ReleaseRef(const CasKey& casKey, bool dropIfZero) override final;
+		virtual bool VerifyAndGetCachedFileInfo(CachedFileInfo& out, StringKey fileNameKey, u64 verifiedLastWriteTime, u64 verifiedSize) override final;
+		virtual bool InvalidateCachedFileInfo(StringKey fileNameKey) override final;
 		virtual bool StoreCasFile(CasKey& out, const tchar* fileName, const CasKey& casKeyOverride, bool deferCreation) override;
 		virtual bool StoreCasFileClient(CasKey& out, StringKey fileNameKey, const tchar* fileName, SharedMemoryHandle memoryHandle, u64 memoryOffset, u64 fileSize, const tchar* hint, bool keepMappingInMemory = false, bool storeCompressed = true) override;
-		virtual bool DropCasFile(const CasKey& casKey, bool forceDelete, const tchar* hint, bool reduceRefCount = false) override;
-		virtual bool ReportBadCasFile(const CasKey& casKey) override;
-		virtual bool CalculateCasKey(CasKey& out, StringKey fileNameKey, const tchar* fileName) override;
+		virtual bool DropCasFile(const CasKey& casKey, bool forceDelete, const tchar* hint, bool reduceRefCount = false) override final;
+		virtual bool ReportBadCasFile(const CasKey& casKey) override final;
+		virtual bool CalculateCasKey(CasKey& out, StringKey fileNameKey, const tchar* fileName) override final;
 		virtual bool CopyOrLink(const CasKey& casKey, const tchar* destination, u32 fileAttributes, bool writeCompressed = false, u32 clientId = 0, const FormattingFunc& formattingFunc = {}, bool isTemp = false, bool allowHardLink = true) override;
 		virtual bool WritePlaceholder(const CasKey& casKey, const tchar* destination, u32 fileAttributes, u32 clientId) override;
-		virtual bool FakeCopy(const CasKey& casKey, const tchar* destination, u64 size, u64 lastWritten, u32 fileAttributes, bool deleteExisting = true) override;
+		virtual bool FakeCopy(const CasKey& casKey, const tchar* destination, u64 size, u64 lastWritten, u32 fileAttributes, bool deleteExisting = true) override final;
 		virtual bool GetCasFileName(StringBufferBase& out, const CasKey& casKey) override;
 		virtual MappedView2 MapView(const CasKey& casKey, const tchar* hint) override;
-		virtual void UnmapView(const MappedView& view, const tchar* hint) override;
-		virtual SharedMemoryAllocator& GetSharedMemory() override;
+		virtual void UnmapView(const MappedView& view, const tchar* hint) override final;
+		virtual SharedMemoryAllocator& GetSharedMemory() override final;
 
-		virtual void ReportFileWrite(StringKey fileNameKey, const tchar* fileName) override;
+		virtual void ReportFileWrite(StringKey fileNameKey, const tchar* fileName) override final;
 		
 		virtual StorageStats& Stats() final;
-		virtual void AddStats(StorageStats& stats) override;
-		virtual void TraceUpdate(Trace& trace, u32 startRow) override;
+		virtual void AddStats(StorageStats& stats) override final;
+		virtual void TraceUpdate(Trace& trace, u32 startRow) override final;
 
 		struct CasEntry;
 		struct CasEntryWriteLock;
@@ -220,8 +220,8 @@ namespace uba
 			bool dropAfterUse = false;
 		};
 
-		virtual void RegisterExternalFileMappingsProvider(ExternalFileMappingsProvider&& provider) override;
-		virtual void DeleteExternalFileMapping(StringKey fileNameKey, const tchar* fileName) override;
+		virtual void RegisterExternalFileMappingsProvider(ExternalFileMappingsProvider&& provider) override final;
+		virtual void DeleteExternalFileMapping(StringKey fileNameKey, const tchar* fileName) override final;
 		virtual bool WriteCompressedFile(WriteResult& out, const tchar* from, FileHandle readHandle, u8* readMem, u64 fileSize, const tchar* toFile, const void* header, u64 headerSize, u64 lastWriteTime = 0) final;
 		bool WriteCasFileNoCheck(WriteCasResult& out, const StringKey& fileNameKey, const tchar* fileName, const CasKey& casKey, const tchar* casFile, bool storeCompressed, bool useMapping, DeferredSource* deferredSource, bool allowFail);
 		bool WriteCasFile(WriteCasResult& out, const tchar* fileName, const CasKey& casKey, bool useMapping, DeferredSource* deferredSouirce, bool allowFail);
@@ -229,7 +229,7 @@ namespace uba
 		bool AddCasFile(StringKey fileNameKey, const tchar* fileName, const CasKey& casKey, bool deferCreation, u64 fileSize, u64 lastWriteTime, CasEntry** outEntry, bool dropCasAfterUse);
 		void CasEntryAccessed(const CasKey& casKey);
 		virtual bool IsDisallowedPath(const tchar* fileName);
-		virtual bool DecompressMemoryToMemory(const u8* compressedData, u64 compressedSize, u8* writeData, u64 decompressedSize, const tchar* readHint, const tchar* writeHint) override;
+		virtual bool DecompressMemoryToMemory(const u8* compressedData, u64 compressedSize, u8* writeData, u64 decompressedSize, const tchar* readHint, const tchar* writeHint) override final;
 
 		void CasEntryAccessed(CasEntry& entry);
 		void CasEntryWritten(CasEntry& entry, u64 size);
diff --git a/Engine/Source/Programs/UnrealBuildAccelerator/Host/Private/UbaExports.cpp b/Engine/Source/Programs/UnrealBuildAccelerator/Host/Private/UbaExports.cpp
--- a/Engine/Source/Programs/UnrealBuildAccelerator/Host/Private/UbaExports.cpp
+++ b/Engine/Source/Programs/UnrealBuildAccelerator/Host/Private/UbaExports.cpp
@@ -1219,43 +1219,47 @@ extern "C"
 		GetCurrentDirectoryW(temp);
 		ubaInstance->workDir = temp.data;
 
-		if (auto coordinatorTable = config.GetTable(TC("Coordinator")))
-		{
-			const tchar* coordinatorName;
-			if (coordinatorTable->GetValueAsString(coordinatorName, TC("Name")))
+		auto createCoordinator = [&]()
 			{
-				auto& logger = sessionServer->GetLogger();
-				const tchar* rootDir = nullptr;
-				coordinatorTable->GetValueAsString(rootDir, TC("RootDir"));
-				if (!rootDir)
-					rootDir = sessionServer->GetRootDir();
-				StringBuffer<512> coordinatorWorkDir(rootDir);
-				coordinatorWorkDir.EnsureEndsWithSlash().Append(coordinatorName);
-				StringBuffer<512> binariesDir;
-				if (!GetDirectoryOfCurrentModule(logger, binariesDir))
-					return nullptr;
-
-				CoordinatorCreateInfo cinfo;
-				cinfo.workDir = coordinatorWorkDir.data;
-				cinfo.binariesDir = binariesDir.data;
-
-				coordinatorTable->GetValueAsString(cinfo.pool, TC("Pool"));
-				UBA_ASSERT(cinfo.pool);
-
-				cinfo.maxCoreCount = 500;
-				coordinatorTable->GetValueAsU32(cinfo.maxCoreCount, TC("MaxCoreCount"));
-
-				cinfo.logging = false;
-				coordinatorTable->GetValueAsBool(cinfo.logging, TC("Log"));
-
-				const tchar* uri = nullptr;
-				if (coordinatorTable->GetValueAsString(uri, TC("Uri")))
-					uba::SetEnvironmentVariableW(TC("UE_HORDE_URL"), uri);
-
-				if (!ubaInstance->coordinator.Create(logger, coordinatorName, cinfo, *networkServer->backend, *networkServer, scheduler))
-					return nullptr;
-			}
-		}
+				if (auto coordinatorTable = config.GetTable(TC("Coordinator")))
+				{
+					const tchar* coordinatorName;
+					if (coordinatorTable->GetValueAsString(coordinatorName, TC("Name")))
+					{
+						auto& logger = sessionServer->GetLogger();
+						const tchar* rootDir = nullptr;
+						coordinatorTable->GetValueAsString(rootDir, TC("RootDir"));
+						if (!rootDir)
+							rootDir = sessionServer->GetRootDir();
+						StringBuffer<512> coordinatorWorkDir(rootDir);
+						coordinatorWorkDir.EnsureEndsWithSlash().Append(coordinatorName);
+						StringBuffer<512> binariesDir;
+						if (!GetDirectoryOfCurrentModule(logger, binariesDir))
+							return;
+
+						CoordinatorCreateInfo cinfo;
+						cinfo.workDir = coordinatorWorkDir.data;
+						cinfo.binariesDir = binariesDir.data;
+
+						coordinatorTable->GetValueAsString(cinfo.pool, TC("Pool"));
+						UBA_ASSERT(cinfo.pool);
+
+						cinfo.maxCoreCount = 500;
+						coordinatorTable->GetValueAsU32(cinfo.maxCoreCount, TC("MaxCoreCount"));
+
+						cinfo.logging = false;
+						coordinatorTable->GetValueAsBool(cinfo.logging, TC("Log"));
+
+						const tchar* uri = nullptr;
+						if (coordinatorTable->GetValueAsString(uri, TC("Uri")))
+							uba::SetEnvironmentVariableW(TC("UE_HORDE_URL"), uri);
+
+						if (!ubaInstance->coordinator.Create(logger, coordinatorName, cinfo, *networkServer->backend, *networkServer, scheduler))
+							return;
+					}
+				}
+			};
+		createCoordinator();
 
 		return ubaInstance;
 	}
