[UBA]
* Improved bloom filter by using all 32-bit values in StringKey

diff --git a/Engine/Source/Programs/UnrealBuildAccelerator/Core/Private/UbaHash.cpp b/Engine/Source/Programs/UnrealBuildAccelerator/Core/Private/UbaHash.cpp
--- a/Engine/Source/Programs/UnrealBuildAccelerator/Core/Private/UbaHash.cpp
+++ b/Engine/Source/Programs/UnrealBuildAccelerator/Core/Private/UbaHash.cpp
@@ -117,18 +117,29 @@ namespace uba
 
 	void BloomFilter::Add(const StringKey& key)
 	{
-		u32 bit = key.a % (sizeof(bytes)*8);
-		u32 byte = bit / 8;
-		u32 bitInByte = bit - byte*8;
-		bytes[byte] |= (1 << bitInByte);
+		auto setBit = [&](u32 hash)
+			{
+				u32 bit = hash % (sizeof(bytes) * 8);
+				bytes[bit / 8] |= (1 << (bit % 8));
+			};
+		setBit(u32(key.a));
+		setBit(u32(key.a >> 32));
+		setBit(u32(key.b));
+		setBit(u32(key.b >> 32));
 	}
 
 	bool BloomFilter::IsGuaranteedMiss(const StringKey& key) const
 	{
-		u32 bit = key.a % (sizeof(bytes)*8);
-		u32 byte = bit / 8;
-		u32 bitInByte = bit - byte*8;
-		return (bytes[byte] & (1 << bitInByte)) == 0;
+		auto notSet = [&](u32 hash) -> bool
+			{
+				u32 bit = hash % (sizeof(bytes) * 8);
+				return (bytes[bit / 8] & (1 << (bit % 8))) == 0;
+			};
+
+		return notSet(u32(key.a))
+			|| notSet(u32(key.a >> 32))
+			|| notSet(u32(key.b))
+			|| notSet(u32(key.b >> 32));
 	}
 
 	constexpr BloomFilter g_empty;
