[UBA]
* Fixed so overlay table is always rooted (going all the way to the drive root). This means that it can live independently of main table so no need to populate main table when adding overlay entries
* Moved offset to file in directory table holder to be before file name so files can be easily visited on session side
* Renamed DirTableMemSize and OverlayTableMemSize to DirTableMaxSize and OverlayTableMaxSize
* Fixed deinit of overlay to work with new design where overlay is separate from main table

diff --git a/Engine/Source/Programs/UnrealBuildAccelerator/Common/Public/UbaDirectoryTableHolder.h b/Engine/Source/Programs/UnrealBuildAccelerator/Common/Public/UbaDirectoryTableHolder.h
--- a/Engine/Source/Programs/UnrealBuildAccelerator/Common/Public/UbaDirectoryTableHolder.h
+++ b/Engine/Source/Programs/UnrealBuildAccelerator/Common/Public/UbaDirectoryTableHolder.h
@@ -19,7 +19,9 @@ namespace uba
 		u8* memory = nullptr;
 		u32 committed = 0;
 		u32 size = 0;
-		UnorderedMap<StringKey, u32> lookup;
+		struct Dir { u32 offset = 0; UnorderedMap<StringKey, u32> files; };
+		UnorderedMap<StringKey, Dir> lookup;
+		Vector<StringKey> roots;
 	};
 
 	////////////////////////////////////////////////////////////////////////////////////////////////////
@@ -31,7 +33,7 @@ namespace uba
 		virtual ~DirectoryTableHolder();
 		void SetTreatTempDirAsEmpty(); // We never want to populate files in Temp
 
-		bool RefreshDirectory(const tchar* dirName, bool forceRegister = false); // Tell uba a directory on disk has been changed by some other system while session is running
+		bool RefreshDirectory(const tchar* dirPath, bool forceRegister = false); // Tell uba a directory on disk has been changed by some other system while session is running
 		bool RegisterNewFile(const tchar* filePath); // Tell uba a new file on disk has been added by some other system while session is running
 		void RegisterDeleteFile(const tchar* filePath); // Tell uba a file on disk has been deleted by some other system while session is running
 		bool RegisterNewDirectory(const tchar* directoryPath); // Tell uba a directory on disk has been added by some other system while session is running
@@ -44,27 +46,32 @@ namespace uba
 		bool EnsureOverlayMemory(DirectoryTableOverlay& overlay, u64 neededSize);
 		void DestroyOverlay(DirectoryTableOverlay& overlay);
 
+		void PrintOverlay(DirectoryTableOverlay& overlay);
+		void PrintOverlay(DirectoryTableOverlay& overlay, DirectoryTableOverlay::Dir& dir, StringBufferBase& path);
+
 		MutableLogger m_logger;
 
 	//protected:
 
-		bool PopulateOverlayFile(DirectoryTableOverlay& overlay, StringView fileName, const Function<bool(FileInformation& out, StringKey)>& getFileInformation);
-		bool AddOverlayFile(DirectoryTableOverlay& overlay, StringView fileName);
+		bool PopulateOverlayFile(DirectoryTableOverlay& overlay, StringView filePath, const Function<bool(FileInformation& out, StringKey)>& getFileInformation);
+		bool AddOverlayFile(DirectoryTableOverlay& overlay, StringView filePath);
 		bool MoveOverlayFile(DirectoryTableOverlay& overlay, StringView from, StringView to);
-		bool WriteOverlayEntry(DirectoryTableOverlay& overlay, StringView fileName, const FileInformation& fileInfo);
-		bool WriteOverlayEntry(DirectoryTableOverlay& overlay, StringKey dirKey, StringView fileNameWithoutSlash, const FileInformation& fileInfo);
+		bool WriteOverlayEntry(DirectoryTableOverlay& overlay, StringView filePath, const FileInformation& fileInfo);
+		bool WriteOverlayEntry(DirectoryTableOverlay& overlay, StringKey dirKey, StringView filePath, StringView fileName, const FileInformation& fileInfo);
+		bool WriteOverlayEntry(DirectoryTableOverlay& overlay, StringKey dirKey, DirectoryTableOverlay::Dir& dir, StringView fileName, const FileInformation& fileInfo, u32& outFileOffset);
+		bool WriteOverlayDirectoryEntriesRecursive(DirectoryTableOverlay& overlay, StringKey filePathKey, StringView filePath);
 
-		bool RegisterCreateFileForWrite(StringKey fileNameKey, StringView fileName, DirectoryTableOverlay* overlay, u64 fileSize, u64 lastWriteTime, bool registerRealFile, bool invalidateStorage);
-		bool RegisterDeleteFile(StringKey fileNameKey, StringView fileName, DirectoryTableOverlay* overlay);
+		bool RegisterCreateFileForWrite(StringKey filePathKey, StringView filePath, DirectoryTableOverlay* overlay, u64 fileSize, u64 lastWriteTime, bool registerRealFile, bool invalidateStorage);
+		bool RegisterDeleteFile(StringKey filePathKey, StringView filePath, DirectoryTableOverlay* overlay);
 
 		u32 WriteDirectoryEntries(const StringKey& dirKey, StringView dirPath, u32* outTableOffset = nullptr);
 		void WriteDirectoryEntriesRecursive(const StringKey& dirKey, StringView dirPath, u32& outTableOffset);
 		bool WriteDirectoryEntriesInternal(DirectoryTable::Directory& dir, const StringKey& dirKey, StringView dirPath, bool isRefresh, u32& outTableOffset);
 		void EnsureDirectoryTableMemory(u64 neededSize);
 
-		virtual void FileEntryAdded(StringKey fileNameKey, u64 lastWritten, u64 size);
-		virtual bool ShouldWriteToDisk(StringView fileName);
-		virtual void InvalidateCachedFileInfo(const StringKey& fileNameKey);
+		virtual void FileEntryAdded(StringKey filePathKey, u64 lastWritten, u64 size);
+		virtual bool ShouldWriteToDisk(StringView filePath);
+		virtual void InvalidateCachedFileInfo(const StringKey& filePathKey);
 
 		bool m_runningRemote;
 
@@ -80,9 +87,8 @@ namespace uba
 
 	////////////////////////////////////////////////////////////////////////////////////////////////////
 
-	bool GetDirKey(StringKey& outDirKey, StringBufferBase& outDirName, const tchar*& outLastSlash, StringView fileName);
-	StringKey GetKeyAndFixedName(StringBuffer<>& outFixedFilePath, StringKeyHasher& outHasher, const tchar* filePath, StringKey* outDirKey = nullptr, const tchar** outBaseFileName = nullptr);
-	StringKey GetKeyAndFixedName(StringBuffer<>& outFixedFilePath, const tchar* filePath, StringKey* outDirKey = nullptr, const tchar** outBaseFileName = nullptr);
+	bool GetDirKey(StringKey& outDirKey, StringBufferBase& outDirPath, const tchar*& outLastSlash, StringView filePath);
+	StringKey GetKeyAndFixedName(StringBuffer<>& outFixedFilePath, const tchar* filePath);
 
 	////////////////////////////////////////////////////////////////////////////////////////////////////
 }
diff --git a/Engine/Source/Programs/UnrealBuildAccelerator/Core/Public/UbaDirectoryTable.h b/Engine/Source/Programs/UnrealBuildAccelerator/Core/Public/UbaDirectoryTable.h
--- a/Engine/Source/Programs/UnrealBuildAccelerator/Core/Public/UbaDirectoryTable.h
+++ b/Engine/Source/Programs/UnrealBuildAccelerator/Core/Public/UbaDirectoryTable.h
@@ -56,7 +56,8 @@ namespace uba
 	struct DirTableOffset { u32 internal = InvalidTableOffset; };
 	inline bool IsDirectoryEntry(DirTableOffset d) { return (d.internal & DirectoryEntryFlag) != 0; }
 	inline bool IsValidEntry(DirTableOffset d) { return d.internal != InvalidTableOffset; }
-	inline bool IsOverlayOffset(u32 offset) { return (offset & OverlayTableFlag) != 0; }
+	inline bool IsOverlayOffset(DirTableOffset d) { return (d.internal & OverlayTableFlag) != 0; }
+
 	////////////////////////////////////////////////////////////////////////////////////////////////////
 
 	class DirectoryTable
@@ -78,7 +79,7 @@ namespace uba
 
 		inline void Init(const u8* mem, u32 tableCount, u32 tableSize);
 		inline void InitOverlay(const u8* mem, u32 tableSize);
-		inline void DeinitOverlay(const u8* oldMem, u32 size);
+		inline void DeinitOverlay();
 		inline void ParseDirectoryTable(DirTableSize size);
 		inline void ParseDirectoryTableNoLock(DirTableSize size);
 		inline void ParseDirectoryTableNoLock(const u8* memory, u32 from, u32 to, u32 isOverlayFlag);
@@ -129,6 +130,7 @@ namespace uba
 		inline DirectoryTable(MemoryBlock& block);
 		DirectoryTable(const DirectoryTable&) = delete;
 		void operator=(const DirectoryTable&) = delete;
+		inline bool IsOverlayOffset(u32 offset) { return (offset & OverlayTableFlag) != 0; }
 
 		MemoryBlock& m_memoryBlock;
 		ReaderWriterLock m_lookupLock;
@@ -154,49 +156,41 @@ namespace uba
 
 	void DirectoryTable::InitOverlay(const u8* mem, u32 tableSize)
 	{
-		UBA_ASSERT(u64(mem) > OverlayTableFlag);
 		m_overlayMemory = mem;
 		ParseDirectoryTableNoLock(mem, 0, tableSize, OverlayTableFlag);
 		m_overlaySize = tableSize;
 	}
 
-	void DirectoryTable::DeinitOverlay(const u8* oldMem, u32 size)
+	void DirectoryTable::DeinitOverlay()
 	{
-		if (!size)
+		if (!m_overlaySize)
 			return;
 
-		m_overlayMemory = oldMem;
 		for (auto i=m_lookup.begin(), e=m_lookup.end(); i!=e;)
 		{
 			auto& dir = i->second;
-			if (!IsOverlayOffset(dir.latestOffset))
+
+			// Directory has no overlay
+			if (dir.latestOverlayOffset == InvalidTableOffset)
 			{
 				++i;
 				continue;
 			}
 
-			while (true)
-			{
-				BinaryReader reader(GetReader(dir.latestOffset));
-				dir.latestOffset = u32(reader.Read7BitEncoded());
-				if (!IsOverlayOffset(dir.latestOffset))
-					break;
-			}
+			// Directory was only overlay, just remove it
 			if (dir.latestOffset == InvalidTableOffset)
 			{
 				i = m_lookup.erase(i);
+				e = m_lookup.end();
 				continue;
 			}
 
-			dir.populatedOffset = dir.latestOffset;
+			dir.latestOverlayOffset = InvalidTableOffset;
+			dir.populatedOverlayOffset = InvalidTableOffset;
 
-			for (auto j=dir.files.begin(); j!=dir.files.end();)
-			{
-				if (!IsOverlayOffset(j->second.internal))
-					++j;
-				else
-					j = dir.files.erase(j);
-			}
+			// Let's force a repopulate 
+			dir.files.clear();
+			dir.populatedOffset = InvalidTableOffset;
 
 			++i;
 		}
diff --git a/Engine/Source/Programs/UnrealBuildAccelerator/Core/Public/UbaProtocol.h b/Engine/Source/Programs/UnrealBuildAccelerator/Core/Public/UbaProtocol.h
--- a/Engine/Source/Programs/UnrealBuildAccelerator/Core/Public/UbaProtocol.h
+++ b/Engine/Source/Programs/UnrealBuildAccelerator/Core/Public/UbaProtocol.h
@@ -56,8 +56,8 @@ namespace uba
 	inline constexpr u32 CommunicationMemSize = IsWindows ? 64*1024 : 64*1024*2; // Macos expands some commandlines to be crazy long
 
 	inline constexpr u32 FileMappingTableMemSize = 16 * 1024 * 1024;
-	inline constexpr u32 DirTableMemSize = 128 * 1024 * 1024;
-	inline constexpr u32 OverlayTableMemSize = 16 * 1024 * 1024; // Seen ~600kb max size for ltcg/thinlto linking of our biggest binaries (one temp file per .o file)
+	inline constexpr u32 DirTableMaxSize = 128 * 1024 * 1024;
+	inline constexpr u32 OverlayTableMaxSize = 16 * 1024 * 1024; // Seen ~600kb max size for ltcg/thinlto linking of our biggest binaries (one temp file per .o file)
 
 	inline constexpr u32 EmergencyShutdownExitCode = 9886;
 }
diff --git a/Engine/Source/Programs/UnrealBuildAccelerator/Detours/Private/PosixOS/UbaDetoursFunctionsPosix.cpp b/Engine/Source/Programs/UnrealBuildAccelerator/Detours/Private/PosixOS/UbaDetoursFunctionsPosix.cpp
--- a/Engine/Source/Programs/UnrealBuildAccelerator/Detours/Private/PosixOS/UbaDetoursFunctionsPosix.cpp
+++ b/Engine/Source/Programs/UnrealBuildAccelerator/Detours/Private/PosixOS/UbaDetoursFunctionsPosix.cpp
@@ -43,7 +43,6 @@ namespace uba
 
 	SharedMemoryAllocatorHandle g_sharedMemoryAllocator;
 	u8* g_sharedMemory;
-	u32 g_overlayFilesSize;
 }
 
 using namespace uba;
@@ -2956,7 +2955,7 @@ namespace uba
 		GetMappingHandleName(dirTableUidName, directoryTableHandle);
 		int dirTableFd = shm_open(dirTableUidName.data, O_RDONLY, S_IRUSR | S_IWUSR);
 		UBA_ASSERTF(dirTableFd != -1, "shm_open failed (%s)", strerror(errno));
-		u8* dirTableMem = (u8*)mmap(NULL, DirTableMemSize, PROT_READ, MAP_SHARED, dirTableFd, 0);
+		u8* dirTableMem = (u8*)mmap(NULL, DirTableMaxSize, PROT_READ, MAP_SHARED, dirTableFd, 0);
 		UBA_ASSERTF(dirTableMem != MAP_FAILED, "mmap for dirtable mem failed (%s)", strerror(errno));
 		g_directoryTable.Init(dirTableMem, directoryTableCount, directoryTableSize.main);
 
@@ -2964,10 +2963,9 @@ namespace uba
 		GetMappingHandleName(overlayTableUidName, overlayTableHandle);
 		int overlayTableFd = shm_open(overlayTableUidName.data, O_RDONLY, S_IRUSR | S_IWUSR);
 		UBA_ASSERTF(overlayTableFd != -1, "shm_open failed (%s)", strerror(errno));
-		u8* overlayTableMem = (u8*)mmap(NULL, OverlayTableMemSize, PROT_READ, MAP_SHARED, overlayTableFd, 0);
+		u8* overlayTableMem = (u8*)mmap(NULL, OverlayTableMaxSize, PROT_READ, MAP_SHARED, overlayTableFd, 0);
 		UBA_ASSERTF(overlayTableMem != MAP_FAILED, "mmap for dirtable mem failed (%s)", strerror(errno));
 		g_directoryTable.InitOverlay(overlayTableMem , directoryTableSize.overlay);
-		g_overlayFilesSize = directoryTableSize.overlay;
 
 		if (trackInputs)
 			InitTrackInput();
@@ -3121,14 +3119,6 @@ extern "C"
 		FlushDebugLog();
 		#endif
 
-		Vector<u8> oldOverlayFiles;
-		if (g_overlayFilesSize)
-		{
-			oldOverlayFiles.resize(g_overlayFilesSize);
-			memcpy(oldOverlayFiles.data(), g_directoryTable.m_overlayMemory, g_overlayFilesSize);
-
-		}
-
 		*outArguments = 0;
 		bool newProcess;
 		DirTableSize directoryTableSize;
@@ -3167,14 +3157,13 @@ extern "C"
 			LogHeader(ToView(outArguments));
 			#endif
 
+			g_directoryTable.DeinitOverlay();
+
 			if (u32 overlaySize = directoryTableSize.overlay)
 			{
-				const u8* overlayMem = g_directoryTable.m_overlayMemory;
-				g_directoryTable.DeinitOverlay(oldOverlayFiles.data(), g_overlayFilesSize);
 				directoryTableSize.overlay = 0; // Want to just parse normal directory first
 				g_directoryTable.ParseDirectoryTableNoLock(directoryTableSize);
-				g_directoryTable.InitOverlay(overlayMem, overlaySize);
-				g_overlayFilesSize = overlaySize;
+				g_directoryTable.InitOverlay(g_directoryTable.m_overlayMemory, overlaySize);
 			}
 		}
 
diff --git a/Engine/Source/Programs/UnrealBuildAccelerator/Detours/Private/Windows/UbaDetoursFunctionsWin.cpp b/Engine/Source/Programs/UnrealBuildAccelerator/Detours/Private/Windows/UbaDetoursFunctionsWin.cpp
--- a/Engine/Source/Programs/UnrealBuildAccelerator/Detours/Private/Windows/UbaDetoursFunctionsWin.cpp
+++ b/Engine/Source/Programs/UnrealBuildAccelerator/Detours/Private/Windows/UbaDetoursFunctionsWin.cpp
@@ -66,8 +66,6 @@ namespace uba
 	SharedMemoryHandle g_permanentFilesHandle;
 	u8* g_permanentFilesMemory;
 
-	u32 g_overlayFilesSize;
-
 
 bool g_useMiMalloc;
 constexpr u64  g_pageSize = 64*1024;
@@ -1190,7 +1188,6 @@ void Init(const DetoursPayload& payload, u64 startTime)
 		TimerScope ts2(g_stats.dirTable);
 		g_directoryTable.Init(directoryTableMem, directoryTableCount, directoryTableSize.main);
 		g_directoryTable.InitOverlay(overlayTableMem, directoryTableSize.overlay);
-		g_overlayFilesSize = directoryTableSize.overlay;
 	}
 
 	if (permanentFilesHandle.IsValid())
@@ -1299,14 +1296,6 @@ extern "C"
 		FlushDebugLog();
 		#endif
 
-		Vector<u8> oldOverlayFiles;
-		if (g_overlayFilesSize)
-		{
-			oldOverlayFiles.resize(g_overlayFilesSize);
-			memcpy(oldOverlayFiles.data(), g_directoryTable.m_overlayMemory, g_overlayFilesSize);
-
-		}
-
 		*outArguments = 0;
 		bool newProcess;
 		DirTableSize directoryTableSize;
@@ -1345,14 +1334,13 @@ extern "C"
 			LogHeader(ToView(outArguments));
 			#endif
 
+			g_directoryTable.DeinitOverlay();
+
 			if (u32 overlaySize = directoryTableSize.overlay)
 			{
-				const u8* overlayMem = g_directoryTable.m_overlayMemory;
-				g_directoryTable.DeinitOverlay(oldOverlayFiles.data(), g_overlayFilesSize);
 				directoryTableSize.overlay = 0; // Want to just parse normal directory first
 				g_directoryTable.ParseDirectoryTableNoLock(directoryTableSize);
-				g_directoryTable.InitOverlay(overlayMem, overlaySize);
-				g_overlayFilesSize = overlaySize;
+				g_directoryTable.InitOverlay(g_directoryTable.m_overlayMemory, overlaySize);
 			}
 		}
 
