[UBA]
* Added missing file to previous submit

diff --git a/Engine/Source/Programs/UnrealBuildAccelerator/Common/Private/UbaDirectoryTableHolder.cpp b/Engine/Source/Programs/UnrealBuildAccelerator/Common/Private/UbaDirectoryTableHolder.cpp
--- a/Engine/Source/Programs/UnrealBuildAccelerator/Common/Private/UbaDirectoryTableHolder.cpp
+++ b/Engine/Source/Programs/UnrealBuildAccelerator/Common/Private/UbaDirectoryTableHolder.cpp
@@ -11,20 +11,29 @@ namespace uba
 {
 	////////////////////////////////////////////////////////////////////////////////////////////////////
 
-	bool GetDirKey(StringKey& outDirKey, StringBufferBase& outDirName, const tchar*& outLastSlash, StringView fileName)
+	u32 DefaultDirAttributes()
 	{
-		outLastSlash = TStrrchr(fileName.data, PathSeparator);
-		UBA_ASSERTF(outLastSlash, TC("Can't get dir key for path %s"), fileName.data);
+		#if PLATFORM_WINDOWS
+		return FILE_ATTRIBUTE_DIRECTORY;
+		#else
+		return S_IFDIR;
+		#endif
+	}
+
+	bool GetDirKey(StringKey& outDirKey, StringBufferBase& outDirPath, const tchar*& outLastSlash, StringView filePath)
+	{
+		outLastSlash = TStrrchr(filePath.data, PathSeparator);
+		UBA_ASSERTF(outLastSlash, TC("Can't get dir key for path %s"), filePath.data);
 		if (!outLastSlash)
 			return false;
 
-		u64 dirLen = u64(outLastSlash - fileName.data);
-		outDirName.Append(fileName.data, dirLen);
-		outDirKey = CaseInsensitiveFs ? ToStringKeyLower(outDirName) : ToStringKey(outDirName);
+		u64 dirLen = u64(outLastSlash - filePath.data);
+		outDirPath.Append(filePath.data, dirLen);
+		outDirKey = CaseInsensitiveFs ? ToStringKeyLower(outDirPath) : ToStringKey(outDirPath);
 		return true;
 	}
 
-	StringKey GetKeyAndFixedName(StringBuffer<>& outFixedFilePath, StringKeyHasher& outHasher, const tchar* filePath, StringKey* outDirKey, const tchar** outBaseFileName)
+	StringKey GetKeyAndFixedName2(StringBuffer<>& outFixedFilePath, StringKeyHasher& outHasher, const tchar* filePath, StringKey* outDirKey, StringView* outFileName)
 	{
 		StringBuffer<> workingDir;
 		if (!IsAbsolutePath(filePath))
@@ -36,22 +45,22 @@ namespace uba
 
 		StringKey dirKey;
 		StringBuffer<> dirNameForHash;
-		const tchar* baseFileName = nullptr;
-		GetDirKey(dirKey, dirNameForHash, baseFileName, outFixedFilePath);
+		const tchar* lastSlash = nullptr;
+		GetDirKey(dirKey, dirNameForHash, lastSlash, outFixedFilePath);
 		if (outDirKey)
 			*outDirKey = dirKey;
-		if (outBaseFileName)
-			*outBaseFileName = baseFileName;
+		if (outFileName)
+			*outFileName = StringView(lastSlash + 1, outFixedFilePath.count - u32(lastSlash - outFixedFilePath.data) - 1);
 
 		if (CaseInsensitiveFs)
 			dirNameForHash.MakeLower();
 
 		outHasher.Update(dirNameForHash);
 
-		if (baseFileName)
+		if (lastSlash)
 		{
 			StringBuffer<256> baseFileNameForHash;
-			baseFileNameForHash.Append(baseFileName);
+			baseFileNameForHash.Append(lastSlash);
 			if (CaseInsensitiveFs)
 				baseFileNameForHash.MakeLower();
 
@@ -71,10 +80,10 @@ namespace uba
 		return result;
 	}
 
-	StringKey GetKeyAndFixedName(StringBuffer<>& outFixedFilePath, const tchar* filePath, StringKey* outDirKey, const tchar** outBaseFileName)
+	StringKey GetKeyAndFixedName(StringBuffer<>& outFixedFilePath, const tchar* filePath)
 	{
 		StringKeyHasher hasher;
-		return GetKeyAndFixedName(outFixedFilePath, hasher, filePath, outDirKey, outBaseFileName);
+		return GetKeyAndFixedName2(outFixedFilePath, hasher, filePath, nullptr, nullptr);
 	}
 
 	////////////////////////////////////////////////////////////////////////////////////////////////////
@@ -85,9 +94,9 @@ namespace uba
 	,	m_directoryTableMemory(TC("SessionDirectoryTable"))
 	,	m_directoryTable(m_directoryTableMemory)
 	{
-		m_directoryTableHandle = uba::CreateMemoryMappingW(m_logger, PAGE_READWRITE|SEC_RESERVE, DirTableMemSize, nullptr, TC("DirMappings"));
+		m_directoryTableHandle = uba::CreateMemoryMappingW(m_logger, PAGE_READWRITE|SEC_RESERVE, DirTableMaxSize, nullptr, TC("DirMappings"));
 		UBA_ASSERT(m_directoryTableHandle.IsValid());
-		m_directoryTableMem = MapViewOfFile(m_logger, m_directoryTableHandle, FILE_MAP_WRITE, 0, DirTableMemSize);
+		m_directoryTableMem = MapViewOfFile(m_logger, m_directoryTableHandle, FILE_MAP_WRITE, 0, DirTableMaxSize);
 		UBA_ASSERT(m_directoryTableMem);
 
 		m_directoryTable.m_memory = m_directoryTableMem;
@@ -96,7 +105,7 @@ namespace uba
 
 	DirectoryTableHolder::~DirectoryTableHolder()
 	{
-		UnmapViewOfFile(m_logger, m_directoryTableMem, DirTableMemSize, TC("DirectoryTable"), false); // Can't discard since only parts might be mapped
+		UnmapViewOfFile(m_logger, m_directoryTableMem, DirTableMaxSize, TC("DirectoryTable"), false); // Can't discard since only parts might be mapped
 		CloseFileMapping(m_logger, m_directoryTableHandle, TC("DirectoryTable"));
 	}
 
@@ -112,13 +121,13 @@ namespace uba
 #endif
 	}
 
-	bool DirectoryTableHolder::RefreshDirectory(const tchar* dirName, bool forceRegister)
+	bool DirectoryTableHolder::RefreshDirectory(const tchar* dirPath, bool forceRegister)
 	{
 		UBA_ASSERT(!m_runningRemote);
 
-		StringBuffer<> dirPath;
+		StringBuffer<> fixedDirPath;
 		StringKeyHasher hasher;
-		GetKeyAndFixedName(dirPath, hasher, dirName);
+		GetKeyAndFixedName2(fixedDirPath, hasher, dirPath, nullptr, nullptr);
 
 		StringKey dirKey = ToStringKey(hasher);
 
@@ -129,7 +138,7 @@ namespace uba
 		{
 			lookupLock.Leave();
 			if (forceRegister)
-				WriteDirectoryEntries(dirKey, dirPath);
+				WriteDirectoryEntries(dirKey, fixedDirPath);
 			return true;
 		}
 		DirectoryTable::Directory& dir = res->second;
@@ -145,7 +154,7 @@ namespace uba
 		UBA_ASSERT(dir.populatedOffset == 1);
 
 		u32 tableOffset;
-		return WriteDirectoryEntriesInternal(dir, dirKey, dirPath, true, tableOffset);
+		return WriteDirectoryEntriesInternal(dir, dirKey, fixedDirPath, true, tableOffset);
 	}
 
 	bool DirectoryTableHolder::RegisterNewFile(const tchar* filePath)
@@ -196,9 +205,9 @@ namespace uba
 	{
 		UBA_ASSERT(!overlay.handle.IsValid());
 		SCOPED_WRITE_LOCK(overlay.lock, lock);
-		overlay.handle = uba::CreateMemoryMappingW(m_logger, PAGE_READWRITE|SEC_RESERVE, OverlayTableMemSize, nullptr, TC("Overlay"));
+		overlay.handle = uba::CreateMemoryMappingW(m_logger, PAGE_READWRITE|SEC_RESERVE, OverlayTableMaxSize, nullptr, TC("Overlay"));
 		UBA_ASSERT(overlay.handle.IsValid());
-		overlay.memory = MapViewOfFile(m_logger, overlay.handle, FILE_MAP_WRITE, 0, OverlayTableMemSize);
+		overlay.memory = MapViewOfFile(m_logger, overlay.handle, FILE_MAP_WRITE, 0, OverlayTableMaxSize);
 		UBA_ASSERT(overlay.memory);
 	}
 
@@ -219,9 +228,9 @@ namespace uba
 		if (totalSize <= overlay.committed)
 			return true;
 		u64 newSize = AlignUp(totalSize, 32ull*1024);
-		if (newSize > OverlayTableMemSize)
+		if (newSize > OverlayTableMaxSize)
 		{
-			static bool called = m_logger.Error(TC("Overlay table overflow. OverlayTableMemSize need to be increased (Size: %llu)"), overlay.size);
+			static bool called = m_logger.Error(TC("Overlay table overflow. OverlayTableMaxSize need to be increased (Size: %llu)"), overlay.size);
 			return false;
 		}
 
@@ -238,7 +247,7 @@ namespace uba
 		if (!overlay.handle.IsValid())
 			return;
 		SCOPED_WRITE_LOCK(overlay.lock, lock);
-		UnmapViewOfFile(m_logger, overlay.memory, OverlayTableMemSize, TC("Overlay"), false); // Can't discard since only parts might be mapped
+		UnmapViewOfFile(m_logger, overlay.memory, OverlayTableMaxSize, TC("Overlay"), false); // Can't discard since only parts might be mapped
 		overlay.memory = nullptr;
 		CloseFileMapping(m_logger, overlay.handle, TC("Overlay"));
 		overlay.handle = {};
@@ -247,25 +256,61 @@ namespace uba
 		overlay.lookup.clear();
 	}
 
-	bool DirectoryTableHolder::PopulateOverlayFile(DirectoryTableOverlay& overlay, StringView fileName, const Function<bool(FileInformation&, StringKey)>& getFileInformation)
+	void DirectoryTableHolder::PrintOverlay(DirectoryTableOverlay& overlay)
+	{
+		StringBuffer<> path;
+		for (auto& root : overlay.roots)
+		{
+			auto findIt = overlay.lookup.find(root);
+			UBA_ASSERT(findIt != overlay.lookup.end());
+			PrintOverlay(overlay, findIt->second, path);
+		}
+		m_logger.Info(TC("OVERLAY MEMORY SIZE: %llu"), overlay.size);
+	}
+
+	void DirectoryTableHolder::PrintOverlay(DirectoryTableOverlay& overlay, DirectoryTableOverlay::Dir& dir, StringBufferBase& path)
+	{
+		for (auto& file : dir.files)
+		{
+			u32 pathLen = path.count;
+			BinaryReader reader(overlay.memory, file.second);
+			reader.ReadString(path);
+
+			auto findIt = overlay.lookup.find(file.first);
+			if (findIt != overlay.lookup.end())
+			{
+				path.Append(PathSeparator);
+				PrintOverlay(overlay, findIt->second, path);
+			}
+			else
+			{
+				//if (reader.ReadFileAttributes())
+					m_logger.Info(TC("OVERLAY FILE: %s"), path.data);
+			}
+			path.Resize(pathLen);
+		}
+	}
+
+	bool DirectoryTableHolder::PopulateOverlayFile(DirectoryTableOverlay& overlay, StringView filePath, const Function<bool(FileInformation&, StringKey)>& getFileInformation)
 	{
 		StringKey dirKey;
-		const tchar* lastSlash = nullptr;
-		StringBuffer<> fixedFileName;
-		StringKey fileNameKey = GetKeyAndFixedName(fixedFileName, fileName.data, &dirKey, &lastSlash);
+		StringView fileName;
+		StringBuffer<> fixedFilePath;
+		StringKeyHasher hasher;
+		StringKey filePathKey = GetKeyAndFixedName2(fixedFilePath, hasher, filePath.data, &dirKey, &fileName);
 		FileInformation fileInfo;
-		if (!getFileInformation(fileInfo, fileNameKey))
+		if (!getFileInformation(fileInfo, filePathKey))
 			return false;
-		return WriteOverlayEntry(overlay, dirKey, ToView(lastSlash + 1), fileInfo);
+		return WriteOverlayEntry(overlay, dirKey, fixedFilePath, fileName, fileInfo);
 	}
 
-	bool DirectoryTableHolder::AddOverlayFile(DirectoryTableOverlay& overlay, StringView fileName)
+	bool DirectoryTableHolder::AddOverlayFile(DirectoryTableOverlay& overlay, StringView filePath)
 	{
 		FileInformation fi;
 		fi.size = InvalidValue; // TODO: This is very wrong
 		fi.attributes = DefaultAttributes();
 		fi.index = ++m_fileIndexCounter;
-		return WriteOverlayEntry(overlay, fileName, fi);
+		return WriteOverlayEntry(overlay, filePath, fi);
 	}
 
 	bool DirectoryTableHolder::MoveOverlayFile(DirectoryTableOverlay& overlay, StringView from, StringView to)
@@ -275,75 +320,112 @@ namespace uba
 		return AddOverlayFile(overlay, to);
 	}
 
-	bool DirectoryTableHolder::WriteOverlayEntry(DirectoryTableOverlay& overlay, StringView fileName, const FileInformation& fileInfo)
+	bool DirectoryTableHolder::WriteOverlayEntry(DirectoryTableOverlay& overlay, StringView filePath, const FileInformation& fileInfo)
 	{
 		StringKey dirKey;
-		const tchar* lastSlash = nullptr;
-		StringBuffer<> fixedFileName;
-		StringKey fileNameKey = GetKeyAndFixedName(fixedFileName, fileName.data, &dirKey, &lastSlash);
+		StringView fileName;
+		StringBuffer<> fixedFilePath;
+		StringKeyHasher hasher;
+		StringKey filePathKey = GetKeyAndFixedName2(fixedFilePath, hasher, filePath.data, &dirKey, &fileName);
 
 		// We need to make sure as far down as possible that the directories are created.. overlay entries must always somehow be attached to the end of the normal table
 		// TODO: This must ALSO be made for remotes
-		if (!m_runningRemote)
-		{
-			StringView dirName(fixedFileName.data, u32(lastSlash - fixedFileName.data));
-			WriteDirectoryEntries(dirKey, dirName);
-		}
+		//if (!m_runningRemote)
+		//{
+		//	StringView dirName(fixedFilePath.data, u32(fileName.data - fixedFilePath.data) - 1);
+		//	WriteDirectoryEntries(dirKey, dirName);
+		//}
 
-		return WriteOverlayEntry(overlay, dirKey, ToView(lastSlash + 1), fileInfo);
+		return WriteOverlayEntry(overlay, dirKey, fixedFilePath, fileName, fileInfo);
 	}
 
-	bool DirectoryTableHolder::WriteOverlayEntry(DirectoryTableOverlay& overlay, StringKey dirKey, StringView fileNameWithoutSlash, const FileInformation& fileInfo)
+	bool DirectoryTableHolder::WriteOverlayEntry(DirectoryTableOverlay& overlay, StringKey dirKey, StringView filePath, StringView fileName, const FileInformation& fileInfo)
 	{
 		SCOPED_WRITE_LOCK(overlay.lock, lock);
 
 		auto dirOffsetIns = overlay.lookup.try_emplace(dirKey);
+		auto& dir = dirOffsetIns.first->second;
 
-		u8 temp[1024];
+		if (!dir.offset)
+		{
+			StringView dirPath(filePath.data, u32(fileName.data - filePath.data) - 1);
+			WriteOverlayDirectoryEntriesRecursive(overlay, dirKey, dirPath);
+		}
+
+		auto& fileOffset = dir.files[CaseInsensitiveFs ? ToStringKeyLower(fileName) : ToStringKey(fileName)];
+		return WriteOverlayEntry(overlay, dirKey, dir, fileName, fileInfo, fileOffset);
+	}
+
+	bool DirectoryTableHolder::WriteOverlayEntry(DirectoryTableOverlay& overlay, StringKey dirKey, DirectoryTableOverlay::Dir& dir, StringView fileName, const FileInformation& fileInfo, u32& outFileOffset)
+	{
+		u8 temp[256];
 		BinaryWriter writer(temp, 0, sizeof(temp));
 		writer.WriteStringKey(dirKey);
-		writer.Write7BitEncoded(dirOffsetIns.first->second); // Previous entry for same directory
-		if (!dirOffsetIns.first->second)
-		{
-			u32 dirAttributes = 0;
-			#if PLATFORM_WINDOWS
-			dirAttributes = FILE_ATTRIBUTE_DIRECTORY;
-			#else
-			dirAttributes = S_IFDIR;
-			#endif
-
-			writer.WriteFileAttributes(dirAttributes);
-			writer.WriteVolumeSerial(fileInfo.volumeSerialNumber); // Reuse file info serial
+		writer.Write7BitEncoded(dir.offset); // Previous entry for same directory
+		if (!dir.offset)
+		{
+			writer.WriteFileAttributes(DefaultDirAttributes());
+			writer.WriteVolumeSerial(0);
 			writer.WriteFileIndex(++m_fileIndexCounter);
 		}
 		writer.Write7BitEncoded(1); // Count
-		writer.WriteString(fileNameWithoutSlash);
+		u32 relativePosition = u32(writer.GetPosition());
+		writer.WriteString(fileName);
 		writer.WriteFileAttributes(fileInfo.attributes);
 		writer.WriteVolumeSerial(fileInfo.volumeSerialNumber);
 		writer.WriteFileIndex(fileInfo.index);
 		writer.WriteFileTime(fileInfo.lastWriteTime);
 		writer.WriteFileSize(fileInfo.size);
 		u64 written = writer.GetPosition();
-
 		if (!EnsureOverlayMemory(overlay, written + 8))
 			return false;
-
 		BinaryWriter memWriter(overlay.memory, overlay.size, overlay.committed);
 		memWriter.Write7BitEncoded(written); // Storage size
-		dirOffsetIns.first->second = u32(memWriter.GetPosition()) + sizeof(StringKey);
+		dir.offset = u32(memWriter.GetPosition()) + sizeof(StringKey);
+		outFileOffset = u32(memWriter.GetPosition()) + relativePosition;
+
 		memWriter.WriteBytes(writer.GetData(), written);
 		overlay.size = u32(memWriter.GetPosition());
 		return true;
 	}
 
-	bool DirectoryTableHolder::RegisterCreateFileForWrite(StringKey fileNameKey, StringView fileName, DirectoryTableOverlay* overlay, u64 fileSize, u64 lastWriteTime, bool registerRealFile, bool invalidateStorage)
+	bool DirectoryTableHolder::WriteOverlayDirectoryEntriesRecursive(DirectoryTableOverlay& overlay, StringKey filePathKey, StringView filePath)
+	{
+		StringView dirPath = filePath.GetPath();
+		if (!dirPath.count)
+		{
+			overlay.roots.push_back(filePathKey);
+			return true;
+		}
+		StringKey dirKey = CaseInsensitiveFs ? ToStringKeyLower(dirPath) : ToStringKey(dirPath);
+		auto& dir = overlay.lookup.try_emplace(dirKey).first->second;
+
+		u8 temp[256];
+		BinaryWriter writer(temp, 0, sizeof(temp));
+
+		if (!dir.offset)
+			if (!WriteOverlayDirectoryEntriesRecursive(overlay, dirKey, dirPath))
+				return false;
+
+		auto res = dir.files.try_emplace(filePathKey);
+		if (!res.second)
+			return true;
+
+		FileInformation fi;
+		fi.attributes = DefaultDirAttributes();
+		fi.index = ++m_fileIndexCounter;
+		fi.volumeSerialNumber = 1;
+		return WriteOverlayEntry(overlay, dirKey, dir, filePath.GetFileName(), fi, res.first->second);
+	}
+
+	bool DirectoryTableHolder::RegisterCreateFileForWrite(StringKey filePathKey, StringView filePath, DirectoryTableOverlay* overlay, u64 fileSize, u64 lastWriteTime, bool registerRealFile, bool invalidateStorage)
 	{
 		auto& dirTable = m_directoryTable;
 
 		StringKey dirKey;
 		const tchar* lastSlash;
 		StringBuffer<> dirName;
-		if (!GetDirKey(dirKey, dirName, lastSlash, fileName))
+		if (!GetDirKey(dirKey, dirName, lastSlash, filePath))
 			return true;
 
 		// Remote is not updating its own directory table
@@ -357,7 +439,7 @@ namespace uba
 				info.volumeSerialNumber = 1;
 				info.size = fileSize;
 				info.lastWriteTime = lastWriteTime;
-				WriteOverlayEntry(*overlay, fileName, info);
+				WriteOverlayEntry(*overlay, filePath, info);
 			}
 			return true;
 		}
@@ -365,7 +447,7 @@ namespace uba
 #if 0//_DEBUG  // Bring this back, turned off right now because a few lines above the call to this method we add a mapping
 		{
 			SCOPED_FUTEX(m_fileMappingTableLookupLock, lookupLock);
-			auto findIt = m_fileMappingTableLookup.find(fileNameKey);
+			auto findIt = m_fileMappingTableLookup.find(filePathKey);
 			if (findIt != m_fileMappingTableLookup.end())
 			{
 				FileMappingEntry& entry = findIt->second;
@@ -376,7 +458,7 @@ namespace uba
 		}
 #endif
 
-		bool shouldWriteToDisk = registerRealFile && ShouldWriteToDisk(fileName);
+		bool shouldWriteToDisk = registerRealFile && ShouldWriteToDisk(filePath);
 
 		// When not writing to disk we need to populate lookup before adding non-written files.. otherwise they will be lost once lookup is actually populated
 		if (!shouldWriteToDisk)
@@ -412,17 +494,17 @@ namespace uba
 			UBA_ASSERT(res);
 			dirLock.Enter();
 		}
-		UBA_ASSERTF(dir.populatedOffset == 1 || (!ShouldWriteToDisk(fileName) && dir.populatedOffset != 3), TC("Registering create file for write %s with unexpect dir.populatedOffset %u "), fileName.data, dir.populatedOffset);
+		UBA_ASSERTF(dir.populatedOffset == 1 || (!ShouldWriteToDisk(filePath) && dir.populatedOffset != 3), TC("Registering create file for write %s with unexpect dir.populatedOffset %u "), filePath.data, dir.populatedOffset);
 
-		if (fileNameKey == StringKeyZero)
+		if (filePathKey == StringKeyZero)
 		{
 			StringBuffer<> forKey;
-			forKey.Append(fileName);
+			forKey.Append(filePath);
 			if (CaseInsensitiveFs)
 				forKey.MakeLower();
-			fileNameKey = ToStringKey(forKey);
+			filePathKey = ToStringKey(forKey);
 		}
-		auto insres = dir.files.try_emplace(fileNameKey, ~u32(0));
+		auto insres = dir.files.try_emplace(filePathKey, ~u32(0));
 
 		u64 fileIndex = InvalidValue;
 		u32 attributes = 0;
@@ -432,8 +514,8 @@ namespace uba
 		if (shouldWriteToDisk)
 		{
 			FileInformation info;
-			if (!GetFileInformation(info, m_logger, fileName.data))
-				return m_logger.Error(TC("Failed to get file information for %s while checking file added for write. This should not happen! (%s)"), fileName.data, LastErrorToText().data);
+			if (!GetFileInformation(info, m_logger, filePath.data))
+				return m_logger.Error(TC("Failed to get file information for %s while checking file added for write. This should not happen! (%s)"), filePath.data, LastErrorToText().data);
 
 			attributes = info.attributes;
 			volumeSerial = info.volumeSerialNumber;
@@ -458,6 +540,8 @@ namespace uba
 		if (!insres.second && IsValidEntry(insres.first->second))
 		{
 			BinaryReader reader(m_directoryTableMem + insres.first->second.internal);
+			reader.SkipString();
+
 			u32 oldAttr = reader.ReadFileAttributes();
 
 			if (IsDirectory(oldAttr))  // Ignore updating directories.. they should always be the same regardless
@@ -479,16 +563,16 @@ namespace uba
 					return true;
 				// TODO: Somehow this can happen and I have no idea how. last written time should be set on close file so it shouldnt be possible.
 				//else
-				//	m_logger.Warning(TC("Somehow file %s has same last written time at two points in time but different size (old %llu new %llu)"), fileName.data, oldSize, fileSize);
+				//	m_logger.Warning(TC("Somehow file %s has same last written time at two points in time but different size (old %llu new %llu)"), filePath.data, oldSize, fileSize);
 			}
 		}
 
 		// There are directory crawlers happening in parallel so we need to really make sure to invalidate this one since a crawler can actually
 		// hit this file with information from a query before it was written.. and then it will turn it back to "verified" using old info
 		if (registerRealFile && invalidateStorage)
-			InvalidateCachedFileInfo(fileNameKey);
+			InvalidateCachedFileInfo(filePathKey);
 
-		FileEntryAdded(fileNameKey, lastWriteTime, fileSize);
+		FileEntryAdded(filePathKey, lastWriteTime, fileSize);
 
 		u32 volumeIndex = m_volumeCache.GetSerialIndex(volumeSerial);
 		u8 temp[1024];
@@ -499,8 +583,8 @@ namespace uba
 			writer.WriteStringKey(dirKey);
 			writer.Write7BitEncoded(dir.latestOffset); // Previous entry for same directory
 			writer.Write7BitEncoded(1); // Count
-			writer.WriteString(lastSlash + 1);
 			entryPos = writer.GetPosition();
+			writer.WriteString(lastSlash + 1);
 			writer.WriteFileAttributes(attributes);
 			writer.WriteVolumeSerial(volumeIndex);
 			writer.WriteFileIndex(fileIndex);
@@ -513,7 +597,7 @@ namespace uba
 		}
 
 #if UBA_DEBUG_TRACK_DIR
-		m_debugLogger->Info(TC("TRACKADD    %s (Size: %llu, Attr: %u, Key: %s, Id: %llu)"), fileName.data, fileSize, attributes, KeyToString(fileNameKey).data, fileIndex);
+		m_debugLogger->Info(TC("TRACKADD    %s (Size: %llu, Attr: %u, Key: %s, Id: %llu)"), filePath.data, fileSize, attributes, KeyToString(filePathKey).data, fileIndex);
 #endif
 
 		if (overlay && !isDirectory)
@@ -524,7 +608,7 @@ namespace uba
 			info.lastWriteTime = lastWriteTime;
 			info.size = fileSize;
 			info.volumeSerialNumber = volumeIndex;
-			if (!WriteOverlayEntry(*overlay, dirKey, ToView(lastSlash + 1), info))
+			if (!WriteOverlayEntry(*overlay, dirKey, filePath, ToView(lastSlash + 1), info))
 				return false;
 		}
 
@@ -542,19 +626,19 @@ namespace uba
 		return true;
 	}
 
-	bool DirectoryTableHolder::RegisterDeleteFile(StringKey fileNameKey, StringView fileName, DirectoryTableOverlay* overlay)
+	bool DirectoryTableHolder::RegisterDeleteFile(StringKey filePathKey, StringView filePath, DirectoryTableOverlay* overlay)
 	{
 		StringKey dirKey;
 		const tchar* lastSlash;
 		StringBuffer<> dirName;
-		if (!GetDirKey(dirKey, dirName, lastSlash, fileName))
+		if (!GetDirKey(dirKey, dirName, lastSlash, filePath))
 			return false;
 
 		if (overlay)
 		{
 			// TODO: If it does not exist in overlay nor main it should not create any entry
 			FileInformation fileInfo;
-			if (!WriteOverlayEntry(*overlay, dirKey, ToView(lastSlash+1), fileInfo))
+			if (!WriteOverlayEntry(*overlay, dirKey, filePath, ToView(lastSlash+1), fileInfo))
 				return false;
 		}
 
@@ -581,19 +665,19 @@ namespace uba
 		if (dir.populatedOffset == 2) // Invalid directory, don't need to register this
 			return true;
 
-		UBA_ASSERTF(dir.populatedOffset == 1, TC("Registering deleted file %s with unexpect dir.populatedOffset %u "), fileName.data, dir.populatedOffset);
+		UBA_ASSERTF(dir.populatedOffset == 1, TC("Registering deleted file %s with unexpect dir.populatedOffset %u "), filePath.data, dir.populatedOffset);
 
-		if (fileNameKey == StringKeyZero)
+		if (filePathKey == StringKeyZero)
 		{
 			StringBuffer<> forKey;
-			forKey.Append(fileName);
+			forKey.Append(filePath);
 			if (CaseInsensitiveFs)
 				forKey.MakeLower();
-			fileNameKey = ToStringKey(forKey);
+			filePathKey = ToStringKey(forKey);
 		}
 
 		// Does not exist, no need adding to file table
-		if (dir.files.erase(fileNameKey) == 0)
+		if (dir.files.erase(filePathKey) == 0)
 			return true;
 
 		u8 temp[1024];
@@ -616,7 +700,7 @@ namespace uba
 		}
 
 #if UBA_DEBUG_TRACK_DIR
-		m_debugLogger->Info(TC("TRACKDEL    %s (Key: %s)"), fileName.data, KeyToString(fileNameKey).data);
+		m_debugLogger->Info(TC("TRACKDEL    %s (Key: %s)"), filePath.data, KeyToString(filePathKey).data);
 #endif
 
 		{
@@ -780,6 +864,9 @@ namespace uba
 						return;
 					UBA_ASSERT(e.attributes);
 					UBA_ASSERT(e.nameLen < 512);
+
+					res.first->second = {u32(memoryWriter.GetPosition())}; // Temporary offset that will be used further down to calculate the real offset
+
 					memoryWriter.WriteString(e.name, e.nameLen);
 
 #if UBA_DEBUG_TRACK_DIR
@@ -790,7 +877,6 @@ namespace uba
 					if (id == 0xffffffffffffffffllu) // When using projfs we might not have the file yet and in that case we need to make this up.
 						id = ++m_fileIndexCounter;
 
-					res.first->second = {u32(memoryWriter.GetPosition())}; // Temporary offset that will be used further down to calculate the real offset
 					memoryWriter.WriteFileAttributes(e.attributes);
 					memoryWriter.WriteVolumeSerial(e.volumeSerial == volumeSerial ? volumeSerialIndex : m_volumeCache.GetSerialIndex(e.volumeSerial));
 					memoryWriter.WriteFileIndex(id);
@@ -833,11 +919,7 @@ namespace uba
 			m_debugLogger->Info(TC("    FORCED EMPTY"));
 #endif
 
-#if PLATFORM_WINDOWS
-			dirAttributes = FILE_ATTRIBUTE_DIRECTORY;
-#else
-			UBA_ASSERTF(false, TC("Not implemented"));
-#endif
+			dirAttributes = DefaultDirAttributes();
 		}
 
 		written = memoryWriter.GetPosition();
@@ -879,7 +961,7 @@ namespace uba
 
 		// Update offsets to be relative to full memory
 		for (auto& kv : dir.files)
-			kv.second = {filesOffset + kv.second.internal};
+			kv.second.internal += filesOffset;
 
 
 		outTableOffset = tableOffset;
@@ -894,9 +976,9 @@ namespace uba
 			return;
 
 		u64 newSize = AlignUp(neededSize, u64(1024*1024));
-		if (newSize > DirTableMemSize)
+		if (newSize > DirTableMaxSize)
 		{
-			static bool called = m_logger.Error(TC("Directory table overflow. DirTableMemSize need to be increased (Size: %llu)"), dirTable.m_memorySize);
+			static bool called = m_logger.Error(TC("Directory table overflow. DirTableMaxSize need to be increased (Size: %llu)"), dirTable.m_memorySize);
 		}
 
 		u64 toCommit = newSize - m_directoryTableMemCommitted;
@@ -907,7 +989,7 @@ namespace uba
 		LockMemory(m_logger, address, toCommit, TC("DirectoryTable"));
 	}
 
-	void DirectoryTableHolder::FileEntryAdded(StringKey fileNameKey, u64 lastWritten, u64 size)
+	void DirectoryTableHolder::FileEntryAdded(StringKey filePathKey, u64 lastWritten, u64 size)
 	{
 	}
 
@@ -916,7 +998,7 @@ namespace uba
 		return true;
 	}
 
-	void DirectoryTableHolder::InvalidateCachedFileInfo(const StringKey& fileNameKey)
+	void DirectoryTableHolder::InvalidateCachedFileInfo(const StringKey& filePathKey)
 	{
 	}
 }
\ No newline at end of file
