[UBA]
* Fixed bug in StringView GetFileName/GetPath if string was empty

diff --git a/Engine/Source/Programs/UnrealBuildAccelerator/Core/Private/UbaStringBuffer.cpp b/Engine/Source/Programs/UnrealBuildAccelerator/Core/Private/UbaStringBuffer.cpp
--- a/Engine/Source/Programs/UnrealBuildAccelerator/Core/Private/UbaStringBuffer.cpp
+++ b/Engine/Source/Programs/UnrealBuildAccelerator/Core/Private/UbaStringBuffer.cpp
@@ -497,17 +497,19 @@ namespace uba
 
 	StringView StringView::GetFileName() const
 	{
-		for (const tchar* i=data+count-1, *e=data; i!=e; --i)
-			if (*i == PathSeparator)
-				return StringView(i+1, u32(data+count-i-1));
+		if (count)
+			for (const tchar* i=data+count-1, *e=data; i!=e; --i)
+				if (*i == PathSeparator)
+					return StringView(i+1, u32(data+count-i-1));
 		return *this;
 	}
 
 	StringView StringView::GetPath() const
 	{
-		for (const tchar* i=data+count-1, *e=data; i!=e; --i)
-			if (*i == PathSeparator)
-				return StringView(data, u32(i-data));
+		if (count)
+			for (const tchar* i=data+count-1, *e=data; i!=e; --i)
+				if (*i == PathSeparator)
+					return StringView(data, u32(i-data));
 		return {};
 	}
 
