[UBA]
* Fixed so directories under forced empty directory (temp) still exist in forced empty dir in order to be able to use sub directories

diff --git a/Engine/Source/Programs/UnrealBuildAccelerator/Common/Private/UbaDirectoryTableHolder.cpp b/Engine/Source/Programs/UnrealBuildAccelerator/Common/Private/UbaDirectoryTableHolder.cpp
--- a/Engine/Source/Programs/UnrealBuildAccelerator/Common/Private/UbaDirectoryTableHolder.cpp
+++ b/Engine/Source/Programs/UnrealBuildAccelerator/Common/Private/UbaDirectoryTableHolder.cpp
@@ -778,17 +778,22 @@ namespace uba
 			return;
 		}
 
-		dirPath.count = u32(rit - dirPath.data);
+		StringView parentDirPath(dirPath.data, u32(rit - dirPath.data));
 
 		StringBuffer<> parentDirForHash;
-		parentDirForHash.Append(dirPath);
+		parentDirForHash.Append(parentDirPath);
 		if (CaseInsensitiveFs)
 			parentDirForHash.MakeLower();
 		StringKey parentKey = ToStringKey(parentDirForHash);
 
 		// Traverse through ancestors and populate them, this is an optimization
 		u32 parentOffset;
-		WriteDirectoryEntriesRecursive(parentKey, dirPath, parentOffset);
+		WriteDirectoryEntriesRecursive(parentKey, parentDirPath, parentOffset);
+
+		dirLock.Leave();
+
+		if (parentKey == m_directoryForcedEmpty) // We need to manually add this dir to the forced empty dir
+			RegisterCreateFileForWrite(parentKey, dirPath, nullptr, 0, 0, true, false);
 	}
 
 	bool DirectoryTableHolder::WriteDirectoryEntriesInternal(DirectoryTable::Directory& dir, const StringKey& dirKey, StringView dirPath, bool isRefresh, u32& outTableOffset)
diff --git a/Engine/Source/Programs/UnrealBuildAccelerator/Test/Private/UbaTestDirectoryTable.cpp b/Engine/Source/Programs/UnrealBuildAccelerator/Test/Private/UbaTestDirectoryTable.cpp
--- a/Engine/Source/Programs/UnrealBuildAccelerator/Test/Private/UbaTestDirectoryTable.cpp
+++ b/Engine/Source/Programs/UnrealBuildAccelerator/Test/Private/UbaTestDirectoryTable.cpp
@@ -8,7 +8,7 @@ namespace uba
 	bool TestDirectoryTable(LoggerWithWriter& logger, const StringBufferBase& testRootDir)
 	{
 		DirectoryTableHolder holder(logger.m_writer, TC(""), false);
-		//holder.SetTreatTempDirAsEmpty();
+		holder.SetTreatTempDirAsEmpty();
 
 		CHECK_TRUE(holder.RefreshDirectory(testRootDir.data, true));
 
