[UBA]
* Fixed so detoured DeleteFileW use directory table to early out if file does not exist. Should remove all the network messages that cl.exe cause for non-existing files

diff --git a/Engine/Source/Programs/UnrealBuildAccelerator/Detours/Private/Windows/UbaDetoursFunctionsKernelBase.inl b/Engine/Source/Programs/UnrealBuildAccelerator/Detours/Private/Windows/UbaDetoursFunctionsKernelBase.inl
--- a/Engine/Source/Programs/UnrealBuildAccelerator/Detours/Private/Windows/UbaDetoursFunctionsKernelBase.inl
+++ b/Engine/Source/Programs/UnrealBuildAccelerator/Detours/Private/Windows/UbaDetoursFunctionsKernelBase.inl
@@ -1276,6 +1276,16 @@ BOOL Detoured_DeleteFileW(LPCWSTR lpFileName)
 
 	StringKey fileNameKey = ToStringKey(fixedNameLower);
 
+	if (!fixedName.StartsWith(g_systemTemp))
+	{
+		DirTableOffset dirTableOffset;
+		if (!Rpc_GetEntryOffset(dirTableOffset, fileNameKey, fixedName, false))
+		{
+			SetLastError(ERROR_FILE_NOT_FOUND);
+			return false;
+		}
+	}
+
 	DirTableSize directoryTableSize;
 	bool result;
 	u32 errorCode;
