[UBA]
* Fixed so cache allows entries without inputs (command line could contain everything needed to produce output)
* Added error handling on cache server when exact same cache entry is created at the same time from the same client
* Added some stubbed out code in cache client related to checking file before downloading and writing it

diff --git a/Engine/Source/Programs/UnrealBuildAccelerator/Common/Private/UbaCacheClient.cpp b/Engine/Source/Programs/UnrealBuildAccelerator/Common/Private/UbaCacheClient.cpp
--- a/Engine/Source/Programs/UnrealBuildAccelerator/Common/Private/UbaCacheClient.cpp
+++ b/Engine/Source/Programs/UnrealBuildAccelerator/Common/Private/UbaCacheClient.cpp
@@ -247,9 +247,6 @@ namespace uba
 		if (!m_connected)
 			return false;
 
-		if (!inputsSize)
-			return false;
-
 		#if UBA_TRACE_WRITE_CACHE
 		TrackWorkScope tws(m_client, AsView(TC("WriteToCache")));
 		#else
@@ -1123,6 +1120,16 @@ namespace uba
 			}
 			#endif
 
+			// TODO: Implement this.
+			if (false)
+			{
+				StringKey fileNameKey = ToStringKeyLower(path);
+				CasKey existingCasKey;
+				if (m_storage.CalculateCasKey(existingCasKey, fileNameKey, path.data))
+					if (existingCasKey == casKey)
+						return true;
+			}
+
 			DowngradedLogger logger(m_connected, m_logger.m_writer, TC("UbaCacheClientDownload"));
 			bool writeCompressed = m_session.ShouldStoreIntermediateFilesCompressed() && g_globalRules.FileCanBeCompressed(path);
 			if (!fetcher.RetrieveFile(logger, m_client, casKey, path.data, writeCompressed, nullptr, attributes))
@@ -1309,7 +1316,7 @@ namespace uba
 
 	bool CacheClient::SendCacheEntry(TrackWorkScope& tws, Bucket& bucket, CacheSendStats& stats, const RootPaths& rootPaths, const CasKey& cmdKey, const Map<u32, u32>& inputsStringToCasKey, const Map<u32, u32>& outputsStringToCasKey, const u8* logLines, u64 logLinesSize, const ProcessStartInfo& info)
 	{
-		StackBinaryReader<1024> reader;
+		StackBinaryReader<24*1024> reader;
 
 		UBA_TRACE_WRITE_HINT("SendCacheEntryMessage");
 
@@ -1702,13 +1709,19 @@ namespace uba
 		}
 
 		// Add arguments list to key
-		auto hashString = [&](const tchar* str, u64 strLenIncTerm, u32 rootPos) { hasher.Update(str, strLenIncTerm*sizeof(tchar)); };
-		if (!rootPaths.NormalizeString(m_logger, info.arguments, TStrlen(info.arguments), hashString, false, info.GetDescription(), TC(" calculating command line hash")))
+		if (m_useRoots)
 		{
-			if (report)
-				m_logger.Info(TC("   Failed to normalize commandline %s"), info.arguments);
-			return CasKeyZero;
+			auto hashString = [&](const tchar* str, u64 strLenIncTerm, u32 rootPos) { hasher.Update(str, strLenIncTerm*sizeof(tchar)); };
+			if (!rootPaths.NormalizeString(m_logger, info.arguments, TStrlen(info.arguments), hashString, false, info.GetDescription(), TC(" calculating command line hash")))
+			{
+				if (report)
+					m_logger.Info(TC("   Failed to normalize commandline %s"), info.arguments);
+				return CasKeyZero;
+			}
 		}
+		else
+			hasher.Update(info.arguments, TStrlen(info.arguments)*sizeof(tchar));
+
 		if (report)
 			m_logger.Info(TC("   %s %s"), CasKeyString(ToCasKey(hasher, false)).str, info.arguments);
 
diff --git a/Engine/Source/Programs/UnrealBuildAccelerator/Common/Private/UbaCacheServer.cpp b/Engine/Source/Programs/UnrealBuildAccelerator/Common/Private/UbaCacheServer.cpp
--- a/Engine/Source/Programs/UnrealBuildAccelerator/Common/Private/UbaCacheServer.cpp
+++ b/Engine/Source/Programs/UnrealBuildAccelerator/Common/Private/UbaCacheServer.cpp
@@ -1498,7 +1498,7 @@ namespace uba
 	UBA_NOINLINE bool CacheServer::HandleMessage(const ConnectionInfo& connectionInfo, u8 messageType, BinaryReader& reader, BinaryWriter& writer)
 	{
 		if (messageType != CacheMessageType_Connect && m_isRunningMaintenance)
-			return m_logger.Debug(TC("Can't handle network message %s while running maintenance mode"), ToString(CacheMessageType(messageType))).ToFalse();
+			return m_logger.Debug(TC("Can't handle network message %s while running maintenance mode"), ToString(CacheMessageType(messageType)).data).ToFalse();
 
 		switch (messageType)
 		{
@@ -1879,8 +1879,10 @@ namespace uba
 		else
 		{
 			SCOPED_FUTEX(connectionBucket.deferredCacheEntryLookupLock, lock3);
-			UBA_FOR_ASSERT(bool res =) connectionBucket.deferredCacheEntryLookup.try_emplace(cmdKey, std::move(newEntry)).second;
-			UBA_ASSERT(res);
+			auto& deferredEntry = connectionBucket.deferredCacheEntryLookup[cmdKey];
+			if (deferredEntry.id != 0)
+				return m_logger.Warning(TC("Client %u %s adding deferred entry that already exists"), clientId);
+			deferredEntry = std::move(newEntry);
 		}
 
 		//m_logger.Info(TC("Added new cache entry (%u inputs and %u outputs)"), u32(inputs.size()), outputCount);
diff --git a/Engine/Source/Programs/UnrealBuildAccelerator/Common/Private/UbaScheduler.cpp b/Engine/Source/Programs/UnrealBuildAccelerator/Common/Private/UbaScheduler.cpp
--- a/Engine/Source/Programs/UnrealBuildAccelerator/Common/Private/UbaScheduler.cpp
+++ b/Engine/Source/Programs/UnrealBuildAccelerator/Common/Private/UbaScheduler.cpp
@@ -1572,13 +1572,22 @@ namespace uba
 
 			FinishProcess(ph);
 			delete[] dependencies;
-			delete si;
-		}
 
-		if (m_writeToCache && exitCode == 0)
-		{
-			// TODO: Read dep.json file
-			UBA_ASSERTF(false, TC("Not implemented"));
+			// This does only work if build step actually always write. Many tools read and check if file looks ok and in that case don't write
+			// So this is not a reliable approach..
+			if (m_writeToCache && exitCode == 0 && process.GetExecutionType() != ProcessExecutionType_FromCache)
+			{
+				UBA_ASSERTF(false, TC("Not implemented"));
+				#if 0
+				auto inputs = process.GetTrackedInputs();
+				auto outputs = process.GetTrackedOutputs();
+
+				CacheClient& cacheClient = *m_cacheClients[0];
+				cacheClient.WriteToCache(si->cacheBucketId, *si, inputs.data(), inputs.size(), outputs.data(), outputs.size(), nullptr, 0, process.GetId());
+				#endif
+			}
+
+			delete si;
 		}
 
 		ph.m_process = nullptr;
@@ -2468,16 +2477,16 @@ namespace uba
 								++readIt;
 								last = c;
 							};
-							argLen = u32(writeIt - arg);
+							u32 newLen = u32(writeIt - arg);
 
 							if (command.application.empty())
 							{
-								command.application.assign(arg, arg + argLen);
+								command.application.assign(arg, arg + newLen);
 								return;
 							}
 							if (args.count)
 								args.Append(' ');
-							args.Append(arg, argLen);
+							args.Append(arg, newLen);
 						});
 						command.arguments = args.ToString();
 					}
diff --git a/Engine/Source/Programs/UnrealBuildAccelerator/Common/Public/UbaScheduler.h b/Engine/Source/Programs/UnrealBuildAccelerator/Common/Public/UbaScheduler.h
--- a/Engine/Source/Programs/UnrealBuildAccelerator/Common/Public/UbaScheduler.h
+++ b/Engine/Source/Programs/UnrealBuildAccelerator/Common/Public/UbaScheduler.h
@@ -106,6 +106,8 @@ namespace uba
 		void SetProcessFinishedCallback(const Function<void(const ProcessHandle&)>& processFinished); // Set callback 
 
 		SessionServer& GetSession() { return m_session; }
+		u32 GetCacheClientCount() { return u32(m_cacheClients.size()); }
+		CacheClient& GetCacheClient(u32 index) { return *m_cacheClients[index]; }
 		
 		u32 GetProcessCountThatCanRunRemotelyNow();
 		void GetProcessWeightThatCanRunRemotelyNow(float& outTotal, float& outCrossArchitecture, float& outCrossPlatform);
diff --git a/Engine/Source/Programs/UnrealBuildAccelerator/Detours/Private/Windows/UbaDetoursFunctionsNtDll.inl b/Engine/Source/Programs/UnrealBuildAccelerator/Detours/Private/Windows/UbaDetoursFunctionsNtDll.inl
--- a/Engine/Source/Programs/UnrealBuildAccelerator/Detours/Private/Windows/UbaDetoursFunctionsNtDll.inl
+++ b/Engine/Source/Programs/UnrealBuildAccelerator/Detours/Private/Windows/UbaDetoursFunctionsNtDll.inl
@@ -1700,7 +1700,7 @@ NTSTATUS NTAPI Shared_NtCreateFile(bool IsCreateFunc, PHANDLE hFileHandle, ACCES
 		{
 			if (!keepInMemory && useContent && !isWrite)
 			{
-				if (!info.tracked)
+				if (!info.tracked && !IsOverlayOffset(dirTableOffset))
 				{
 					info.tracked = true;
 					TrackInput(fileName);
