[UBA]
* Improved error message when not finding binary
* Added commented out code for printing directory overlay

diff --git a/Engine/Source/Programs/UnrealBuildAccelerator/Common/Private/UbaProcess.cpp b/Engine/Source/Programs/UnrealBuildAccelerator/Common/Private/UbaProcess.cpp
--- a/Engine/Source/Programs/UnrealBuildAccelerator/Common/Private/UbaProcess.cpp
+++ b/Engine/Source/Programs/UnrealBuildAccelerator/Common/Private/UbaProcess.cpp
@@ -829,6 +829,8 @@ namespace uba
 		{
 			CleanupWrittenFiles();
 
+			//m_session.PrintOverlay(m_shared.overlay);
+
 			for (auto& memoryHandle : m_shared.sharedMemoryHandles)
 				m_session.m_sharedMemory.CloseHandle(m_session.m_logger, memoryHandle, TC(""));
 			m_shared.sharedMemoryHandles.clear();
diff --git a/Engine/Source/Programs/UnrealBuildAccelerator/Common/Private/UbaSession.cpp b/Engine/Source/Programs/UnrealBuildAccelerator/Common/Private/UbaSession.cpp
--- a/Engine/Source/Programs/UnrealBuildAccelerator/Common/Private/UbaSession.cpp
+++ b/Engine/Source/Programs/UnrealBuildAccelerator/Common/Private/UbaSession.cpp
@@ -2275,7 +2275,7 @@ namespace uba
 		if (!IsAbsolutePath(startInfo.application))
 		{
 			if (!SearchPathForFile(m_logger, outRealApplication.Clear(), startInfo.application, ToView(startInfo.workingDir), {}))
-				return false;
+				return m_logger.Error(TC("Failed to find path to application %s"), startInfo.application);
 			startInfo.applicationStr = outRealApplication.data;
 			startInfo.application = startInfo.applicationStr.c_str();
 		}
