[UBA]
* Moved memset of new memory to detoured process to minimize various lock times

diff --git a/Engine/Source/Programs/UnrealBuildAccelerator/Common/Private/UbaProcess.cpp b/Engine/Source/Programs/UnrealBuildAccelerator/Common/Private/UbaProcess.cpp
--- a/Engine/Source/Programs/UnrealBuildAccelerator/Common/Private/UbaProcess.cpp
+++ b/Engine/Source/Programs/UnrealBuildAccelerator/Common/Private/UbaProcess.cpp
@@ -1738,7 +1738,7 @@ namespace uba
 		SharedMemoryHandle h = m_session.m_sharedMemory.CreateHandle(m_session.m_logger, hint.data);
 		writer.Write7BitEncoded(h.ToU64());
 		if (initialSize)
-			m_session.m_sharedMemory.ExtendMemory(writer, h, initialSize);
+			m_session.m_sharedMemory.ExtendMemory(writer, h, initialSize, TC("HandleCreate"), false); // memset happens in detoured process
 
 		SCOPED_FUTEX(m_shared.writtenFilesLock, lock);
 		m_shared.sharedMemoryHandles.insert(h);
@@ -1749,7 +1749,7 @@ namespace uba
 	{
 		SharedMemoryHandle h = SharedMemoryHandle::FromU64(reader.ReadU64());
 		u64 size = reader.ReadU64();
-		m_session.m_sharedMemory.ExtendMemory(writer, h, size);
+		m_session.m_sharedMemory.ExtendMemory(writer, h, size, TC("HandleCommit"), false); // memset happens in detoured process
 		return true;
 	}
 
diff --git a/Engine/Source/Programs/UnrealBuildAccelerator/Common/Private/UbaSharedMemoryAllocator.cpp b/Engine/Source/Programs/UnrealBuildAccelerator/Common/Private/UbaSharedMemoryAllocator.cpp
--- a/Engine/Source/Programs/UnrealBuildAccelerator/Common/Private/UbaSharedMemoryAllocator.cpp
+++ b/Engine/Source/Programs/UnrealBuildAccelerator/Common/Private/UbaSharedMemoryAllocator.cpp
@@ -664,7 +664,7 @@ namespace uba
 		SCOPED_FUTEX(m_lock, lock);
 		while (left)
 		{
-			Slice slice = AllocateNoLock(left, slices.size());
+			Slice slice = AllocateNoLock(left);
 			slices.push_back(slice);
 			left -= slice.size;
 		}
@@ -691,7 +691,7 @@ namespace uba
 		handle.slices.clear();
 	}
 
-	SharedMemoryAllocator::Slice SharedMemoryAllocator::AllocateNoLock(u64 desiredSize, u64 existingSliceCount)
+	SharedMemoryAllocator::Slice SharedMemoryAllocator::AllocateNoLock(u64 desiredSize)
 	{
 		// Find best fit.
 		// Will return first iterator not less than desiredSize
diff --git a/Engine/Source/Programs/UnrealBuildAccelerator/Common/Public/UbaSharedMemoryAllocator.h b/Engine/Source/Programs/UnrealBuildAccelerator/Common/Public/UbaSharedMemoryAllocator.h
--- a/Engine/Source/Programs/UnrealBuildAccelerator/Common/Public/UbaSharedMemoryAllocator.h
+++ b/Engine/Source/Programs/UnrealBuildAccelerator/Common/Public/UbaSharedMemoryAllocator.h
@@ -71,7 +71,7 @@ namespace uba
 		void ReleaseHandle(SharedMemoryHandle h);
 		void ExtendHandle(InternalHandle& handle, u64 size, const tchar* hint, bool zeroMemory);
 		void ClearHandle(InternalHandle& handle);
-		Slice AllocateNoLock(u64 desiredSize, u64 existingSliceCount);
+		Slice AllocateNoLock(u64 desiredSize);
 		void FreeNoLock(Slice slice);
 
 		Logger& m_logger;
diff --git a/Engine/Source/Programs/UnrealBuildAccelerator/Detours/Private/UbaDetoursFileMappingTable.cpp b/Engine/Source/Programs/UnrealBuildAccelerator/Detours/Private/UbaDetoursFileMappingTable.cpp
--- a/Engine/Source/Programs/UnrealBuildAccelerator/Detours/Private/UbaDetoursFileMappingTable.cpp
+++ b/Engine/Source/Programs/UnrealBuildAccelerator/Detours/Private/UbaDetoursFileMappingTable.cpp
@@ -296,8 +296,10 @@ namespace uba
 		SharedMemoryHandle handle{ reader.Read7BitEncoded() };
 
 		out.Init(allocatorHandle, capacity);
-		if (initialSize)
-			out.AddRequestedMemory(allocatorHandle, reader);
+		if (!initialSize)
+			return handle;
+		out.AddRequestedMemory(allocatorHandle, reader);
+		memset(out.GetMemory(), 0, out.GetMappedSize());
 		return handle;
 	}
 
@@ -309,7 +311,9 @@ namespace uba
 		writer.WriteU64(handle.ToU64());
 		writer.WriteU64(size);
 		BinaryReader reader = writer.Flush();
+		u64 start = view.GetMappedSize();
 		view.AddRequestedMemory(allocatorHandle, reader);
+		memset(view.GetMemory() + start, 0, view.GetMappedSize() - start);
 	}
 
 	void Rpc_GetSharedMemory(SharedMemoryView& view, SharedMemoryAllocatorHandle allocatorHandle, u8* allocatorMem, SharedMemoryHandle handle, u64 offset, u64 size, const tchar* hint, SharedMemoryMapType type)
diff --git a/Engine/Source/Programs/UnrealBuildAccelerator/Detours/Private/UbaDetoursShared.cpp b/Engine/Source/Programs/UnrealBuildAccelerator/Detours/Private/UbaDetoursShared.cpp
--- a/Engine/Source/Programs/UnrealBuildAccelerator/Detours/Private/UbaDetoursShared.cpp
+++ b/Engine/Source/Programs/UnrealBuildAccelerator/Detours/Private/UbaDetoursShared.cpp
@@ -553,7 +553,7 @@ namespace uba
 			return fileName.data;
 		}
 		#endif
-
+/*
 		{
 			SCOPED_READ_LOCK(g_mappedFileTable.m_lookupLock, lock);
 			auto it = g_mappedFileTable.m_lookup.find(fileNameKey);
@@ -606,7 +606,7 @@ namespace uba
 #endif
 			}
 		}
-
+*/
 		if (!g_allowDirectoryCache || fileName.StartsWith(g_systemTemp.data)) // We need to skip SystemTemp.. lots of stuff going on there.
 		{
 			if (foundMapping)
