[UBA]
* Fixed potential (extremely unlikely) write-after-free bug in storage server.
* Fixed potential double disconnect call in network server
* Added sync locks around remote process returned callback (to make claude happy)
* Fixed so speedtest runs again without logging errors
* Fixed race condition in assert (only debug)
* Added so .cc files are "rarely read" so they don't go through proxy
* Added StringView(const tchar*) = delete to make sure const tchar* is not implicitly cast to TString

diff --git a/Engine/Source/Programs/UnrealBuildAccelerator/Common/Private/UbaSessionServer.cpp b/Engine/Source/Programs/UnrealBuildAccelerator/Common/Private/UbaSessionServer.cpp
--- a/Engine/Source/Programs/UnrealBuildAccelerator/Common/Private/UbaSessionServer.cpp
+++ b/Engine/Source/Programs/UnrealBuildAccelerator/Common/Private/UbaSessionServer.cpp
@@ -48,7 +48,8 @@ namespace uba
 		~RemoteProcess()
 		{
 			if (m_knownInputsDone.IsCreated())
-				m_knownInputsDone.IsSet(50*1000);
+				if (!m_knownInputsDone.IsSet(50*1000))
+					LoggerWithWriter(g_consoleLogWriter).Error(TC("RemoteProcess dtor took more than 50 seconds to exit. Something is very wrong"));
 			delete[] m_knownInputs;
 		}
 
@@ -122,7 +123,7 @@ namespace uba
 		Vector<ReceivedFile> m_receivedFiles;
 		Vector<u8> m_trackedInputs;
 		Vector<u8> m_trackedOutputs;
-		bool m_cancelled = false;
+		Atomic<bool> m_cancelled = false;
 		bool m_canCrossArchitecture = false;
 		bool m_canCrossPlatform = false;
 		bool m_isIdle = false;
@@ -900,6 +901,7 @@ namespace uba
 			remoteProcess->m_clientId = ~0u;
 			remoteProcess->m_sessionId = 0;
 
+			SCOPED_READ_LOCK(m_remoteProcessReturnedEventLock, lock2);
 			if (m_remoteProcessReturnedEvent)
 				m_remoteProcessReturnedEvent(*remoteProcess);
 		}
@@ -944,6 +946,7 @@ namespace uba
 
 		if (!m_queuedRemoteProcesses.empty())
 		{
+			SCOPED_READ_LOCK(m_remoteProcessReturnedEventLock, lock2);
 			if (m_remoteProcessReturnedEvent)
 			{
 				m_logger.Info(TC("No client sessions connected and there are %llu processes left in the remote queue. Will return all queued remote processes"), m_queuedRemoteProcesses.size());
@@ -1732,13 +1735,14 @@ namespace uba
 			auto returnProcess = [&]()
 				{
 					m_queuedRemoteProcesses.push_front({process});
+					SCOPED_READ_LOCK(m_remoteProcessReturnedEventLock, lock2);
 					if (m_remoteProcessReturnedEvent)
 						m_remoteProcessReturnedEvent(*process);
 				};
 
 			auto& startInfo = process->m_startInfo;
 
-			if (!IsAbsolutePath(startInfo.application))
+			if (!IsAbsolutePath(startInfo.application) && !Equals(startInfo.application, TC("speedtest")))
 			{
 				StringBuffer<> temp;
 				if (!SearchPathForFile(m_logger, temp, startInfo.application, startInfo.workingDirStr, {}))
@@ -2063,6 +2067,7 @@ namespace uba
 		m_trace.ProcessReturned(process->m_processId, reason);
 		m_queuedRemoteProcesses.push_front(h);
 
+		SCOPED_READ_LOCK(m_remoteProcessReturnedEventLock, lock2);
 		if (m_remoteProcessReturnedEvent)
 			m_remoteProcessReturnedEvent(*process);
 		return true;
diff --git a/Engine/Source/Programs/UnrealBuildAccelerator/Common/Private/UbaStorageServer.cpp b/Engine/Source/Programs/UnrealBuildAccelerator/Common/Private/UbaStorageServer.cpp
--- a/Engine/Source/Programs/UnrealBuildAccelerator/Common/Private/UbaStorageServer.cpp
+++ b/Engine/Source/Programs/UnrealBuildAccelerator/Common/Private/UbaStorageServer.cpp
@@ -903,7 +903,7 @@ namespace uba
 				auto findIt = m_activeFetches.find(fetchId);
 				if (findIt == m_activeFetches.end())
 					return m_logger.Detail(TC("Can't find active fetch %u, disconnected client? (fetch index %u, client id %u uid %s)"), fetchId, fetchIndex, connectionInfo.GetId(), GuidToString(connectionInfo.GetUid()).str).ToFalse();
-				ActiveFetch& fetch = findIt->second;
+				ActiveFetch& fetch = findIt->second; // This is safe because OnDisconnected can't happen while work is active
 				UBA_ASSERT(fetch.clientId == connectionInfo.GetId());
 				lock.Leave();
 
@@ -920,12 +920,11 @@ namespace uba
 				}
 				memcpy(writer.AllocWrite(toWrite), pos, toWrite);
 
+				fetch.sendCasTime += GetTime() - start;
+
 				bool isDone = fetch.left.fetch_sub(toWrite) == toWrite;
 				if (!isDone)
-				{
-					fetch.sendCasTime += GetTime() - start;
 					return true;
-				}
 
 				fetch.Release(*this, TC("FetchDone"));
 
@@ -935,7 +934,6 @@ namespace uba
 				activeLock.Leave();
 				PushId(fetchId);
 
-				sendCasTime += GetTime() - start;
 				Stats().sendCas += Timer{sendCasTime, 1};
 				return true;
 			}
diff --git a/Engine/Source/Programs/UnrealBuildAccelerator/Core/Private/UbaApplicationRules.cpp b/Engine/Source/Programs/UnrealBuildAccelerator/Core/Private/UbaApplicationRules.cpp
--- a/Engine/Source/Programs/UnrealBuildAccelerator/Core/Private/UbaApplicationRules.cpp
+++ b/Engine/Source/Programs/UnrealBuildAccelerator/Core/Private/UbaApplicationRules.cpp
@@ -115,6 +115,7 @@ namespace uba
 		virtual bool IsRarelyRead(const StringView& file) const override
 		{
 			return file.EndsWith(TCV(".cpp"))
+				|| file.EndsWith(TCV(".cc"))
 				|| (file.EndsWith(TCV(".rsp")) && !file.EndsWith(TCV("Shared.rsp")))
 				|| file.EndsWith(TCV(".i")) // .i is preprocessed source code (for static analysis)
 				;
@@ -368,6 +369,7 @@ namespace uba
 		virtual bool IsRarelyRead(const StringView& file) const override
 		{
 			return file.EndsWith(TCV(".cpp"))
+				|| file.EndsWith(TCV(".cc"))
 				|| (file.EndsWith(TCV(".rsp")) && !file.EndsWith(TCV("Shared.rsp")));
 		}
 
diff --git a/Engine/Source/Programs/UnrealBuildAccelerator/Core/Public/UbaStringBuffer.h b/Engine/Source/Programs/UnrealBuildAccelerator/Core/Public/UbaStringBuffer.h
--- a/Engine/Source/Programs/UnrealBuildAccelerator/Core/Public/UbaStringBuffer.h
+++ b/Engine/Source/Programs/UnrealBuildAccelerator/Core/Public/UbaStringBuffer.h
@@ -95,20 +95,6 @@ namespace uba
 		StringBufferBase(u32 c) : capacity(c) { count = 0; *data = 0; }
 	};
 
-
-	template<u32 Capacity = 512>
-	class StringBuffer : public StringBufferBase
-	{
-	public:
-		StringBuffer() : StringBufferBase(Capacity) { *buf = 0; }
-		explicit StringBuffer(const TString& str) : StringBufferBase(Capacity) { *buf = 0; Append(str); }
-		explicit StringBuffer(const tchar* str) : StringBufferBase(Capacity) { *buf = 0; if (str) Append(str); }
-		StringBuffer(const StringBufferBase& str) : StringBufferBase(Capacity) { *buf = 0; Append(str); }
-		StringBuffer(StringView str) : StringBufferBase(Capacity) { *buf = 0; Append(str); }
-	private:
-		tchar buf[Capacity];
-	};
-
 	////////////////////////////////////////////////////////////////////////////////////////////////////
 
 	enum NoInitEnum { NoInit };
@@ -120,6 +106,7 @@ namespace uba
 		constexpr StringView(const StringView&) = default;
 		constexpr StringView(const tchar* d, u32 c) : data(d), count(c) {}
 		constexpr StringView(const StringBufferBase& sb) : data(sb.data), count(sb.count) {}
+		constexpr StringView(const tchar* d) = delete;
 		StringView(const TString& str) : data(str.data()), count(u32(str.size())) {}
 
 		tchar operator[](u64 i) const { return data[i]; }
@@ -154,6 +141,21 @@ namespace uba
 
 	////////////////////////////////////////////////////////////////////////////////////////////////////
 
+	template<u32 Capacity = 512>
+	class StringBuffer : public StringBufferBase
+	{
+	public:
+		StringBuffer() : StringBufferBase(Capacity) { *buf = 0; }
+		explicit StringBuffer(const TString& str) : StringBufferBase(Capacity) { *buf = 0; Append(str); }
+		explicit StringBuffer(const tchar* str) : StringBufferBase(Capacity) { *buf = 0; if (str) Append(str); }
+		StringBuffer(const StringBufferBase& str) : StringBufferBase(Capacity) { *buf = 0; Append(str); }
+		StringBuffer(StringView str) : StringBufferBase(Capacity) { *buf = 0; Append(str); }
+	private:
+		tchar buf[Capacity];
+	};
+
+	////////////////////////////////////////////////////////////////////////////////////////////////////
+
 	struct LastErrorToText : StringBuffer<256>
 	{
 		LastErrorToText();
