[UBA]
* Fixed so scheduler is using queues for finding processes to run instead of scanning through the process list.

diff --git a/Engine/Source/Programs/UnrealBuildAccelerator/Common/Private/UbaScheduler.cpp b/Engine/Source/Programs/UnrealBuildAccelerator/Common/Private/UbaScheduler.cpp
--- a/Engine/Source/Programs/UnrealBuildAccelerator/Common/Private/UbaScheduler.cpp
+++ b/Engine/Source/Programs/UnrealBuildAccelerator/Common/Private/UbaScheduler.cpp
@@ -181,6 +181,32 @@ namespace uba
 
 	////////////////////////////////////////////////////////////////////////////////////////////////////
 
+	struct Scheduler::ProcessLink
+	{
+		ProcessLink* next;
+		u32 entryIndex;
+
+		static void DeleteAll(ProcessLink* first)
+		{
+			ProcessLink* it = first;
+			while (it)
+			{
+				ProcessLink* next = it->next;
+				delete it;
+				it = next;
+			}
+		}
+
+		template<typename Func>
+		static void Traverse(ProcessLink* first, const Func& func)
+		{
+			for (ProcessLink* it = first; it; it = it->next)
+				func(it->entryIndex);
+		}
+	};
+
+	////////////////////////////////////////////////////////////////////////////////////////////////////
+
 	Scheduler::Scheduler(const SchedulerCreateInfo& info)
 	:	m_session(info.session)
 	,	m_maxLocalProcessors(info.maxLocalProcessors != ~0u ? info.maxLocalProcessors : GetLogicalProcessorCount())
@@ -310,11 +336,10 @@ namespace uba
 			for (auto& entry : m_processEntries)
 			{
 				UBA_ASSERTF(entry.runStatus > RunStatus_Running, TC("Found processes in queue/running state when stopping scheduler."));
-				delete[] entry.dependencies;
+				ProcessLink::DeleteAll(entry.firstDependent);
 				delete entry.info;
 			}
 			m_processEntries.clear();
-			m_processEntriesRunStart = 0;
 			m_session.SetOuterScheduler(nullptr);
 		}
 	}
@@ -336,13 +361,6 @@ namespace uba
 			memcpy(ki, info.knownInputs, info.knownInputsBytes);
 		}
 
-		u32* dep = nullptr;
-		if (info.dependencyCount)
-		{
-			dep = new u32[info.dependencyCount];
-			memcpy(dep, info.dependencies, info.dependencyCount*sizeof(u32));
-		}
-
 		auto info2 = new ProcessStartInfo2(info.info, ki, info.knownInputsCount);
 		info2->Expand();
 		info2->weight = info.weight;
@@ -377,8 +395,8 @@ namespace uba
 		}
 		ProcessEntry entry;
 		entry.info = info2;
-		entry.dependencies = dep;
-		entry.dependencyCount = info.dependencyCount;
+		entry.firstDependent = nullptr;
+		entry.dependencyCount = 0;
 		entry.runStatus = RunStatus_QueuedForRun;
 		entry.cacheStatus = useCache ? CacheStatus_QueuedForTest : CacheStatus_Failed;
 		entry.canDetour = canDetour;
@@ -388,6 +406,113 @@ namespace uba
 		return entry;
 	}
 
+	Scheduler::ProcessQueue& Scheduler::GetReadyToRunQueue(ProcessEntry& entry)
+	{
+		if (!entry.canDetour || !entry.canExecuteRemotely)
+			return m_readyToRunLocal;
+		if (entry.canCrossArchitecture)
+			return m_readyToRunCrossArchitecture;
+		if (entry.canCrossPlatform)
+			return m_readyToRunCrossPlatform;
+		return m_readyToRunRemote;
+	}
+
+	void Scheduler::AddToQueueBackNoLock(ProcessQueue& queue, u32 entryIndex)
+	{
+		queue.insert(entryIndex);
+	}
+
+	void Scheduler::AddToQueueFrontNoLock(ProcessQueue& queue, u32 entryIndex)
+	{
+		queue.insert(entryIndex);
+	}
+	u32 Scheduler::PopFromQueue(ProcessQueue& queue)
+	{
+		if (queue.empty())
+			return ~0u;
+		auto it = queue.begin();
+		u32 entryIndex = *it;
+		queue.erase(it);
+		return entryIndex;
+	}
+
+	Scheduler::ProcessEntry* Scheduler::PeekFromQueue(ProcessQueue& queue)
+	{
+		if (queue.empty())
+			return nullptr;
+		return &m_processEntries[*queue.begin()];
+	}
+
+	void Scheduler::ClearQueue(ProcessQueue& queue)
+	{
+		queue.clear();
+	}
+
+	u32 Scheduler::GetQueueCount(ProcessQueue& queue)
+	{
+		return u32(queue.size());
+	}
+
+	void Scheduler::SetupDependents(ProcessEntry& entry, u32 entryIndex, const EnqueueProcessInfo& info)
+	{
+		// Go through all the processes and see if we have dependencies and if we should just skip this
+		u32 dependencyCount = 0;
+		for (const u32* i=info.dependencies, *e=i+info.dependencyCount; i!=e; ++i)
+		{
+			auto& dependency = m_processEntries[*i];
+			if (dependency.runStatus == RunStatus_Success)
+				continue;
+			if (dependency.runStatus == RunStatus_Failed || dependency.runStatus == RunStatus_Skipped)
+			{
+				AddToQueueBackNoLock(m_readyToSkip, entryIndex);
+				return;
+			}
+			else
+				++dependencyCount;
+		}
+
+		if (!dependencyCount)
+		{
+			if (entry.cacheStatus == CacheStatus_QueuedForTest)
+				AddToQueueBackNoLock(m_readyToTestCache, entryIndex);
+			else
+				AddToQueueBackNoLock(GetReadyToRunQueue(entry), entryIndex);
+			return;
+		}
+
+		entry.dependencyCount = dependencyCount;
+
+		for (const u32* i=info.dependencies, *e=i+info.dependencyCount; i!=e; ++i)
+		{
+			auto& dependency = m_processEntries[*i];
+			if (dependency.runStatus > RunStatus_Running)
+				continue;
+			dependency.firstDependent = new ProcessLink { dependency.firstDependent, entryIndex };
+		}
+	}
+
+	void Scheduler::UpdateDependentsNoLock(RunStatus newRunStatus, ProcessLink* firstDependent)
+	{
+		for (ProcessLink* it = firstDependent; it; it = it->next)
+		{
+			u32 dependentIndex = it->entryIndex;
+			auto& dependent = m_processEntries[dependentIndex];
+			if (--dependent.dependencyCount)
+				continue;
+			if (newRunStatus == RunStatus_Success)
+			{
+				if (dependent.cacheStatus == CacheStatus_QueuedForTest)
+					AddToQueueBackNoLock(m_readyToTestCache, dependentIndex);
+				else
+					AddToQueueBackNoLock(GetReadyToRunQueue(dependent), dependentIndex);
+				continue;
+			}
+			UBA_ASSERT(newRunStatus == RunStatus_Failed || newRunStatus == RunStatus_Skipped);
+			if (dependent.runStatus == RunStatus_QueuedForRun)
+				AddToQueueBackNoLock(m_readyToSkip, dependentIndex);
+
+		}
+	}
 
 	void Scheduler::SkipAllQueued()
 	{
@@ -399,9 +524,18 @@ namespace uba
 			{
 				entry.cacheStatus = CacheStatus_Failed;
 				entry.runStatus = RunStatus_Skipped;
+				// UpdateDependentsNoLock(RunStatus_Skipped, firstDependent); // No need to do this, all following entries are cancelled
 				skipped.push_back(entry.info);
 			}
 		}
+		ClearQueue(m_readyToTestCache);
+		ClearQueue(m_readyToDownloadCache);
+		ClearQueue(m_readyToRunLocal);
+		ClearQueue(m_readyToRunRemote);
+		ClearQueue(m_readyToRunCrossArchitecture);
+		ClearQueue(m_readyToRunCrossPlatform);
+		ClearQueue(m_readyToSkip);
+
 		lock.Leave();
 		for (auto pi : skipped)
 			SkipProcess(*pi);
@@ -424,13 +558,10 @@ namespace uba
 
 		SCOPED_FUTEX(m_processEntriesLock, lock);
 		u32 index = u32(m_processEntries.size());
-		m_processEntries.emplace_back(entry);
+		SetupDependents(m_processEntries.emplace_back(entry), index, info);
 		lock.Leave();
-
 		++m_totalProcesses;
-
-		UpdateQueueCounter(1);
-
+		UpdateQueueCounter(1, entry.canExecuteRemotely);
 		m_updateThreadWakeup.Set();
 		return index;
 	}
@@ -441,17 +572,25 @@ namespace uba
 		UBA_ASSERT(!m_cancelled);
 		u32 index = u32(m_processEntries.size());
 		auto& entry = m_processEntries.emplace_back();
+		entry.dependencyCount = 0;
+		entry.firstDependent = nullptr;
 		entry.runStatus = RunStatus_Suspended;
 		entry.cacheStatus = CacheStatus_Suspended;
 		++m_totalProcesses;
-		UpdateQueueCounter(1);
 		return index;
 	}
+
 	void Scheduler::ResumeQueuedProcess(u32 index, const EnqueueProcessInfo& info)
 	{
-		ProcessEntry entry = CreateProcessEntry(info);
+		auto newEntry = CreateProcessEntry(info);
 		SCOPED_FUTEX(m_processEntriesLock, lock);
-		m_processEntries[index] = entry;
+		auto& entry = m_processEntries[index];
+		ProcessLink* firstDependent = entry.firstDependent;
+		UBA_ASSERT(entry.dependencyCount == 0);
+		entry = newEntry;
+		entry.firstDependent = firstDependent;
+		SetupDependents(entry, index, info);
+		UpdateQueueCounter(1, entry.canExecuteRemotely);
 		m_updateThreadWakeup.Set();
 	}
 
@@ -517,33 +656,8 @@ namespace uba
 	{
 		if (m_session.IsRemoteExecutionDisabled())
 			return 0;
-
-		u32 count = 0;
-		auto processEntries = m_processEntries.data();
 		SCOPED_FUTEX_READ(m_processEntriesLock, lock);
-		for (u32 i=m_processEntriesRunStart, e=u32(m_processEntries.size()); i!=e; ++i)
-		{
-			auto& entry = processEntries[i];
-
-			if (!entry.canExecuteRemotely)
-				continue;
-			if (entry.runStatus != RunStatus_QueuedForRun || entry.cacheStatus != CacheStatus_Failed)
-				continue;
-
-			while (entry.dependencyCount)
-			{
-				auto depIndex = *entry.dependencies;
-				if (m_processEntries[depIndex].runStatus != RunStatus_Success)
-					break;
-				*entry.dependencies = entry.dependencies[--entry.dependencyCount];
-			}
-			if (!entry.dependencyCount)
-				++count;
-		}
-
-		count += m_activeRemoteProcesses;
-
-		return count;
+		return GetQueueCount(m_readyToRunRemote) + GetQueueCount(m_readyToRunCrossArchitecture) + GetQueueCount(m_readyToRunCrossPlatform);
 	}
 
 	void Scheduler::GetProcessWeightThatCanRunRemotelyNow(float& outTotal, float& outCrossArchitecture, float& outCrossPlatform)
@@ -551,45 +665,12 @@ namespace uba
 		outTotal = 0;
 		outCrossArchitecture = 0;
 		outCrossPlatform = 0;
-
 		if (m_session.IsRemoteExecutionDisabled())
 			return;
-
-		float totalWeight = 0;
-		float crossArchitecture = 0;
-		float crossPlatform = 0;
-
-		auto processEntries = m_processEntries.data();
 		SCOPED_FUTEX_READ(m_processEntriesLock, lock);
-		for (u32 i=m_processEntriesRunStart, e=u32(m_processEntries.size()); i!=e; ++i)
-		{
-			auto& entry = processEntries[i];
-
-			if (!entry.canExecuteRemotely)
-				continue;
-			if (entry.runStatus != RunStatus_QueuedForRun || entry.cacheStatus != CacheStatus_Failed)
-				continue;
-
-			while (entry.dependencyCount)
-			{
-				auto depIndex = *entry.dependencies;
-				if (m_processEntries[depIndex].runStatus != RunStatus_Success)
-					break;
-				*entry.dependencies = entry.dependencies[--entry.dependencyCount];
-			}
-			if (entry.dependencyCount)
-				continue;
-			float weight = entry.info->weight;
-			totalWeight += weight;
-			if (entry.canCrossArchitecture)
-				crossArchitecture += weight;
-			if (entry.canCrossPlatform)
-				crossPlatform += weight;
-		}
-
-		outTotal = totalWeight;
-		outCrossArchitecture = crossArchitecture;
-		outCrossPlatform = crossPlatform;
+		outCrossArchitecture = float(GetQueueCount(m_readyToRunCrossArchitecture));
+		outCrossPlatform = float(GetQueueCount(m_readyToRunCrossPlatform));
+		outTotal = float(GetQueueCount(m_readyToRunRemote)) + outCrossArchitecture + outCrossPlatform;
 	}
 
 	void Scheduler::ThreadLoop()
@@ -690,12 +771,13 @@ namespace uba
 		if (entry.runStatus != RunStatus_Running)
 			return m_session.GetLogger().Error(TC("This should not happen (%u)"), entry.runStatus);
 
-		bool success = !m_cancelled;
-
-		if (success)
+		ProcessLink* firstDependent = entry.firstDependent;
+		if (!m_cancelled)
 		{
 			entry.runStatus = RunStatus_QueuedForRun;
-			m_processEntriesRunStart = Min(m_processEntriesRunStart, processIndex);
+			AddToQueueFrontNoLock(GetReadyToRunQueue(entry), processIndex);
+			UpdateQueueCounter(1, entry.canExecuteRemotely);
+			lock.Leave();
 		}
 		else
 		{
@@ -703,10 +785,13 @@ namespace uba
 			entry.runStatus = RunStatus_Skipped;
 			ProcessHandle ph(&process);
 			FinishProcess(ph);
+
+			entry.firstDependent = nullptr;
+			UpdateDependentsNoLock(RunStatus_Skipped, firstDependent);
+			lock.Leave();
+			ProcessLink::DeleteAll(firstDependent);
 		}
-		lock.Leave();
 
-		UpdateQueueCounter(1);
 		UpdateActiveProcessCounter(isLocal, -1);
 		m_updateThreadWakeup.Set();
 		return true;
@@ -726,15 +811,19 @@ namespace uba
 		{
 			entry.runStatus = RunStatus_Skipped;
 			auto pi = entry.info;
+			ProcessLink* firstDependent = entry.firstDependent;
+			entry.firstDependent = nullptr;
+			UpdateDependentsNoLock(RunStatus_Skipped, firstDependent);
 			lock.Leave();
+			ProcessLink::DeleteAll(firstDependent);
 			SkipProcess(*pi);
 			return;
 		}
 
-		m_processEntriesRunStart = Min(m_processEntriesRunStart, processIndex);
+		UBA_ASSERT(entry.runStatus == RunStatus_QueuedForRun);
+		AddToQueueBackNoLock(GetReadyToRunQueue(entry), processIndex);
 		lock.Leave();
 
-		UpdateQueueCounter(1);
 		m_updateThreadWakeup.Set();
 	}
 
@@ -746,11 +835,7 @@ namespace uba
 			return;
 		if (m_session.IsRemoteExecutionDisabled())
 			return;
-		u32 count = 0;
-		SCOPED_FUTEX_READ(m_processEntriesLock, lock);
-		for (auto& entry : m_processEntries)
-			count += u32(entry.canExecuteRemotely && entry.runStatus <= RunStatus_QueuedForRun);
-		lock.Leave();
+		u32 count = m_queuedProcessesCanExecuteRemotely;
 		if (count < m_maxLocalProcessors)
 			m_session.DisableRemoteExecution();
 		else
@@ -829,7 +914,7 @@ namespace uba
 
 		if (!m_useThreadToExitRemoteProcess || info->isLocal)
 		{
-			/*
+			#if 0
 			static int counter = 0;
 			StringBuffer<> buf;
 			buf.Appendf("[REDACTED_USER_HOME]/access/action%u.txt", counter++);
@@ -849,7 +934,7 @@ namespace uba
 				}
 				f.Close();
 			}
-			*/
+			#endif
 
 			ExitProcess(*info, *handle.m_process, handle.m_process->GetExitCode(), false, false);
 			return;
@@ -869,250 +954,90 @@ namespace uba
 			}, TC("UbaScheRemFin"));
 	}
 
-	u32 Scheduler::PopCache()
-	{
-		bool allFinished = true;
-		u32 entryToProcess = ~0u;
-		auto processEntries = m_processEntries.data();
-
-		for (u32 i=m_processEntriesCacheStart, e=u32(m_processEntries.size()); i!=e; ++i)
-		{
-			auto& entry = processEntries[i];
-			auto cacheStatus = entry.cacheStatus;
-
-			if (cacheStatus >= CacheStatus_Downloading)
-			{
-				if (allFinished)
-					m_processEntriesCacheStart = i;
-				continue;
-			}
-
-			allFinished = false;
-
-			if (cacheStatus == CacheStatus_QueuedForTest)
-			{
-				bool canRun = true;
-				while (entry.dependencyCount)
-				{
-					auto depIndex = *entry.dependencies;
-					if (depIndex >= m_processEntries.size())
-					{
-						m_session.GetLogger().Error(TC("Found dependency on index %u but there are only %u processes registered"), depIndex, u32(m_processEntries.size()));
-						return ~0u;
-					}
-					auto depStatus = processEntries[depIndex].runStatus;
-					if (depStatus == RunStatus_Failed || depStatus == RunStatus_Skipped)
-					{
-						entry.cacheStatus = CacheStatus_Failed;
-						entry.runStatus = RunStatus_Skipped;
-						return i;
-					}
-					if (depStatus != RunStatus_Success)
-					{
-						canRun = false;
-						break;
-					}
-				
-					// Swap with last and reduce dependency count
-					*entry.dependencies = entry.dependencies[entry.dependencyCount - 1];
-					--entry.dependencyCount;
-				}
-				if (!canRun)
-					continue;
-
-				return i;
-			}
-
-			if (cacheStatus == CacheStatus_Suspended || cacheStatus == CacheStatus_Testing)
-				continue;
-
-			UBA_ASSERT(cacheStatus == CacheStatus_QueuedForDownload);
-			if (entryToProcess == ~0u)
-				entryToProcess = i;
-		}
-		return entryToProcess;
-	}
-
-	u32 Scheduler::PopProcess(bool isLocal, bool isCrossArchitecture, bool isCrossPlatform, RunStatus& outPrevStatus, bool forReuse, bool& outShouldSleep)
+	u32 Scheduler::PopProcessNoLock(bool isLocal, bool isCrossArchitecture, bool isCrossPlatform, RunStatus& outPrevStatus, bool forReuse, bool& outShouldSleep)
 	{
 		outShouldSleep = false;
 
-		float activeWeight = m_activeLocalProcessWeight;
-		float maxWeight = float(m_maxLocalProcessors);
-		bool atMaxLocalWeight = activeWeight >= maxWeight;
-		auto processEntries = m_processEntries.data();
-		bool allFinished = true;
-
-		u32 firstValidProcess = ~0u;
-		u32 firstValidRemoteCapableProcess = ~0u;
-		u32 firstValidCrossSystemCapableProcess = ~0u;
-
-		for (u32 i=m_processEntriesRunStart, e=u32(m_processEntries.size()); i!=e; ++i)
+		u32 skipEntry = PopFromQueue(m_readyToSkip);
+		if (skipEntry != ~0u)
 		{
-			auto& entry = processEntries[i];
-			auto runStatus = entry.runStatus;
+			auto& entry = m_processEntries[skipEntry];
+			outPrevStatus = entry.runStatus;
+			entry.runStatus = RunStatus_Skipped;
+			UpdateDependentsNoLock(RunStatus_Skipped, entry.firstDependent);
+			ProcessLink::DeleteAll(entry.firstDependent);
+			entry.firstDependent = nullptr;
+			return skipEntry;
+		}
 
-			// Check if we can move search starting point forward
-			if (runStatus > RunStatus_QueuedForRun)
+		auto returnEntry = [&](u32 index)
 			{
-				if (allFinished)
-				{
-					if (runStatus > RunStatus_Running)
-						m_processEntriesRunStart = i;
-					else
-						allFinished = false;
-				}
-				continue;
-			}
-			allFinished = false;
-
-			auto cacheStatus = entry.cacheStatus;
-			
-			if (cacheStatus != CacheStatus_Failed)
-				continue;
+				if (index == ~0u)
+					return index;
+				auto& entry = m_processEntries[index];
+				UBA_ASSERT(entry.cacheStatus == CacheStatus_Failed);
+				outPrevStatus = entry.runStatus;
+				entry.runStatus = RunStatus_Running;
+				return index;
+			};
 
-			UBA_ASSERT(cacheStatus != CacheStatus_Success);
+		if (!isLocal)
+		{
+			if (isCrossArchitecture)
+				return returnEntry(PopFromQueue(m_readyToRunCrossArchitecture));
+			if (isCrossPlatform)
+				return returnEntry(PopFromQueue(m_readyToRunCrossPlatform));
 
-			if (isLocal)
-			{
-				if (m_forceRemote)
-				{
-					if (entry.canExecuteRemotely)
-					{
-						if (!m_hadRemoteRequest)
-							continue;
-						if (!m_hadCrossSystemRequest)
-							continue;
-						if (entry.canCrossArchitecture || entry.canCrossPlatform)
-							continue;
-					}
-				}
-				if (!forReuse)
-				{
-					if (runStatus == RunStatus_QueuedForRun && atMaxLocalWeight)
-						continue;
-					if (activeWeight > 1.0f && activeWeight + entry.info->weight > maxWeight + 1.0f)
-						continue;
-				}
-			}
-			else
-			{
-				if (!entry.canExecuteRemotely)
-					continue;
-				if ((isCrossArchitecture && !entry.canCrossArchitecture) || (isCrossPlatform && !entry.canCrossPlatform))
-					continue;
-			}
+			u32 index = returnEntry(PopFromQueue(m_readyToRunRemote));
+			if (index != ~0u)
+				return index;
+			index = returnEntry(PopFromQueue(m_readyToRunCrossArchitecture));
+			if (index != ~0u)
+				return index;
+			return returnEntry(PopFromQueue(m_readyToRunCrossPlatform));
+		}
 
-			bool canRun = true;
-			while (entry.dependencyCount)
+		auto getEntry = [&](float weightLeft)
 			{
-				auto depIndex = *entry.dependencies;
-				if (depIndex >= m_processEntries.size())
-				{
-					m_session.GetLogger().Error(TC("Found dependency on index %u but there are only %u processes registered"), depIndex, u32(m_processEntries.size()));
+				if (ProcessEntry* entry = PeekFromQueue(m_readyToRunLocal)) // Check local only first
+					if (entry->info->weight <= weightLeft)
+						return PopFromQueue(m_readyToRunLocal);
+				if (!m_forceRemote || !m_hadRemoteRequest)
+					if (ProcessEntry* entry = PeekFromQueue(m_readyToRunRemote))
+						if (entry->info->weight <= weightLeft)
+							return PopFromQueue(m_readyToRunRemote);
+				if (m_forceRemote && m_hadCrossSystemRequest)
 					return ~0u;
-				}
-				auto depStatus = processEntries[depIndex].runStatus;
-				if (depStatus == RunStatus_Failed || depStatus == RunStatus_Skipped)
-				{
-					entry.runStatus = RunStatus_Skipped;
-					return i;
-				}
-				if (depStatus != RunStatus_Success)
-				{
-					canRun = false;
-					break;
-				}
-
-				// Swap with last and reduce dependency count
-				*entry.dependencies = entry.dependencies[entry.dependencyCount - 1];
-				--entry.dependencyCount;
-			}
-
-			if (!canRun)
-				continue;
-
-			// Process can run, now we just need to check if we _should_ run this process (we might want to prioritize later processes)
-
-			if (!isLocal)
-			{
-				if (m_deferRemoteCapableProcesses && m_hadCrossSystemRequest && (entry.canCrossPlatform || entry.canCrossArchitecture))
-				{
-					if (firstValidCrossSystemCapableProcess == ~0u)
-						firstValidCrossSystemCapableProcess = i;
-					continue;
-				}
-
-				firstValidRemoteCapableProcess = i;
-				break;
-			}
-
-			// Only local processes below here
-
-			// Process can only run locally, use it
-			if (!entry.canExecuteRemotely)
-			{
-				firstValidProcess = i;
-				break;
-			}
-
-			// If we don't want to defer remote-capable actions or haven't had any remote helpers, use this process
-			if (!m_deferRemoteCapableProcesses || !m_hadRemoteRequest)
-			{
-				firstValidProcess = i;
-				break;
-			}
-
-			// Track cross system and normal remote processes separately
-			if (entry.canCrossPlatform || entry.canCrossArchitecture)
-			{
-				if (firstValidCrossSystemCapableProcess == ~0u)
-					firstValidCrossSystemCapableProcess = i;
-			}
-			else
-			{
-				if (firstValidRemoteCapableProcess == ~0u)
-					firstValidRemoteCapableProcess = i;
-			}
-
-			continue;
-		}
-
-		if (firstValidProcess == ~0u)
-		{
-			if (firstValidRemoteCapableProcess != ~0u)
-				firstValidProcess = firstValidRemoteCapableProcess;
-			else if (firstValidCrossSystemCapableProcess != ~0u)
-				firstValidProcess = firstValidCrossSystemCapableProcess;
-			else
+				if (ProcessEntry* entry = PeekFromQueue(m_readyToRunCrossArchitecture))
+					if (entry->info->weight <= weightLeft)
+						return PopFromQueue(m_readyToRunCrossArchitecture);
+				if (ProcessEntry* entry = PeekFromQueue(m_readyToRunCrossPlatform))
+					if (entry->info->weight <= weightLeft)
+						return PopFromQueue(m_readyToRunCrossPlatform);
 				return ~0u;
-		}
-
-		auto& entry = processEntries[firstValidProcess];
+			};
 
-		if (isLocal)
-		{
-			if (!LocalProcessOkToSpawn())
-				return ~0u;
+		if (forReuse)
+			return returnEntry(getEntry(100000.0f)); // Don't care about max weight
 
-			if (!forReuse)
-			{
-				// If there are queued getnext, then we try to hand over running queued processes
-				// Setting the event will wake up the getnext thread.. and If the getnext thread is timing out we will retry
-				if (m_getNextQueueCount)
-				{
-					m_updateThreadWakeup.Set();
-					outShouldSleep = true;
-					return ~0u;
-				}
+		float activeWeight = m_activeLocalProcessWeight;
+		float maxWeight = float(m_maxLocalProcessors) + 0.9f;
+		if (activeWeight > 1.0f && (activeWeight >= maxWeight || !LocalProcessOkToSpawn()))
+			return ~0u;
 
-				m_activeLocalProcessWeight += entry.info->weight;
-			}
+		// If there are queued getnext, then we try to hand over running queued processes
+		// Setting the event will wake up the getnext thread.. and If the getnext thread is timing out we will retry
+		if (m_getNextQueueCount)
+		{
+			m_updateThreadWakeup.Set();
+			outShouldSleep = true;
+			return ~0u;
 		}
 
-		outPrevStatus = entry.runStatus;
-		entry.runStatus = RunStatus_Running;
-		return firstValidProcess;
+		u32 entry = getEntry(maxWeight - activeWeight);
+		if (entry != ~0u)
+			m_activeLocalProcessWeight += m_processEntries[entry].info->weight;
+		return returnEntry(entry);
 	}
 
 	bool Scheduler::RunCacheQuery()
@@ -1122,7 +1047,9 @@ namespace uba
 		if (m_activeCacheQueries >= m_maxParallelCacheQueries)
 			return false;
 
-		u32 entryToProcess = PopCache();
+		u32 entryToProcess = PopFromQueue(m_readyToTestCache);
+		if (entryToProcess == ~0u)
+			entryToProcess = PopFromQueue(m_readyToDownloadCache);
 		if (entryToProcess == ~0u)
 			return false;
 
@@ -1158,6 +1085,7 @@ namespace uba
 
 						SCOPED_FUTEX(m_processEntriesLock, lock);
 						m_processEntries[entryToProcess].cacheStatus = CacheStatus_QueuedForDownload;
+						AddToQueueBackNoLock(m_readyToDownloadCache, entryToProcess);
 						--m_activeCacheQueries;
 						break;
 					}
@@ -1179,11 +1107,12 @@ namespace uba
 		{
 			entry.cacheStatus = CacheStatus_Downloading;
 			++m_activeCacheQueries;
+			bool canExecuteRemotely = entry.canExecuteRemotely;
 			lock.Leave();
 
 			++m_activeCacheDownloads;
 
-			m_cacheWorkManager->AddWork([this, info, entryToProcess](const WorkContext&)
+			m_cacheWorkManager->AddWork([this, info, entryToProcess, canExecuteRemotely](const WorkContext&)
 				{
 					auto cacheInfo = info->cacheInfo;
 
@@ -1194,6 +1123,7 @@ namespace uba
 
 					if (success)
 					{
+						UpdateQueueCounter(-1, canExecuteRemotely); // Queue counter is only reduced when process is picked up for Running.. which means need to reduce it on cache success too
 						auto process = new CachedProcess(si);
 						process->logLines.swap(cacheInfo->result.logLines);
 						ProcessHandle ph(process);
@@ -1219,7 +1149,7 @@ namespace uba
 		else
 		{
 			UBA_ASSERT(entry.cacheStatus == CacheStatus_Failed);
-			UBA_ASSERT(entry.runStatus == RunStatus_Skipped);
+			UBA_ASSERTF(entry.runStatus == RunStatus_Skipped, TC("Expected process to be skipped: %u"), u32(entry.runStatus));
 			UBA_ASSERT(!info->cacheInfo);
 			lock.Leave();
 			SkipProcess(*info);
@@ -1241,7 +1171,7 @@ namespace uba
 				m_hadCrossSystemRequest = true;
 
 			bool shouldSleep;
-			u32 indexToRun = PopProcess(isLocal, isCrossArchitecture, isCrossPlatform, prevStatus, false, shouldSleep);
+			u32 indexToRun = PopProcessNoLock(isLocal, isCrossArchitecture, isCrossPlatform, prevStatus, false, shouldSleep);
 			if (indexToRun == ~0u)
 			{
 				lock.Leave();
@@ -1257,10 +1187,11 @@ namespace uba
 			auto info = processEntry.info;
 			bool canDetour = processEntry.canDetour && !m_forceNative;
 			bool wasSkipped = processEntry.runStatus == RunStatus_Skipped;
+			bool canExecuteRemotely = processEntry.canExecuteRemotely;
 			UBA_ASSERT(processEntry.cacheStatus == CacheStatus_Failed);
 			lock.Leave();
 
-			UpdateQueueCounter(-1);
+			UpdateQueueCounter(-1, canExecuteRemotely);
 
 			if (wasSkipped)
 			{
@@ -1394,7 +1325,7 @@ namespace uba
 			bool shouldSleep;
 			RunStatus prevStatus;
 			SCOPED_FUTEX(m_processEntriesLock, lock);
-			u32 indexToRun = PopProcess(isLocal, false, false, prevStatus, forReuse, shouldSleep);
+			u32 indexToRun = PopProcessNoLock(isLocal, false, false, prevStatus, forReuse, shouldSleep);
 			if (indexToRun == ~0u)
 			{
 				u64 usedTimeMs = TimeToMs(GetTime() - startTime);
@@ -1422,9 +1353,10 @@ namespace uba
 			auto& processEntry = m_processEntries[indexToRun];
 			auto newInfo = processEntry.info;
 			bool wasSkipped = processEntry.runStatus == RunStatus_Skipped;
+			bool canExecuteRemotely = processEntry.canExecuteRemotely;
 			lock.Leave();
 
-			UpdateQueueCounter(-1);
+			UpdateQueueCounter(-1, canExecuteRemotely);
 
 			if (wasSkipped)
 			{
@@ -1495,16 +1427,15 @@ namespace uba
 		bool isDone = exitedResponse == ProcessExitedResponse_None;
 
 		SCOPED_FUTEX(m_processEntriesLock, lock);
-		auto& entry = m_processEntries[info.processIndex];
-
-		u32* dependencies = entry.dependencies;
-
+		u32 entryIndex = info.processIndex;
+		auto& entry = m_processEntries[entryIndex];
+		ProcessLink* firstDependent = entry.firstDependent;
 		if (isDone)
 		{
-			//UBA_ASSERT(exitCode == 0);
 			entry.runStatus = exitCode == 0 ? RunStatus_Success : RunStatus_Failed;
+			entry.firstDependent = nullptr;
 			entry.info = nullptr;
-			entry.dependencies = nullptr;
+			UpdateDependentsNoLock(entry.runStatus, firstDependent);
 		}
 		else
 		{
@@ -1514,7 +1445,7 @@ namespace uba
 			entry.canExecuteRemotely = false;
 			entry.canDetour = exitedResponse != ProcessExitedResponse_RerunNative;
 			entry.runStatus = RunStatus_QueuedForRun;
-			m_processEntriesRunStart = Min(m_processEntriesRunStart, info.processIndex);
+			AddToQueueFrontNoLock(GetReadyToRunQueue(entry), entryIndex);
 		}
 
 		if (info.isLocal)
@@ -1532,7 +1463,7 @@ namespace uba
 			}
 		}
 
-		lock.Leave();
+		lock.Leave(); // entry reference can be bad after this since it can realloc
 
 		UpdateActiveProcessCounter(info.isLocal, -1);
 		m_updateThreadWakeup.Set();
@@ -1571,7 +1502,7 @@ namespace uba
 			}
 
 			FinishProcess(ph);
-			delete[] dependencies;
+			ProcessLink::DeleteAll(firstDependent);
 
 			// This does only work if build step actually always write. Many tools read and check if file looks ok and in that case don't write
 			// So this is not a reliable approach..
@@ -1611,9 +1542,11 @@ namespace uba
 		m_session.GetTrace().StatusUpdate(1, 6, str, LogEntryType_Info);
 	}
 
-	void Scheduler::UpdateQueueCounter(int offset)
+	void Scheduler::UpdateQueueCounter(int offset, bool canExecuteRemotely)
 	{
 		m_queuedProcesses += u32(offset);
+		if (canExecuteRemotely)
+			m_queuedProcessesCanExecuteRemotely += u32(offset);
 		m_session.UpdateProgress(m_totalProcesses, m_finishedProcesses, m_errorCount);
 	}
 
@@ -2000,7 +1933,8 @@ namespace uba
 					u64 startTime = GetTime();
 					if (!ProcessReturned(*ph.m_process, true))
 						continue;
-					m_session.GetLogger().Info(TC("Killed process %s (%s) - Too many racing processes"), ph.GetStartInfo().GetDescription(), TimeToText(GetTime() - startTime).str);
+					m_session.GetLogger().Debug(TC("Killed process %s (%s) - Too many racing processes"), ph.GetStartInfo().GetDescription(), TimeToText(GetTime() - startTime).str);
+					timeoutMs = 5;
 					break;
 				}
 			}
diff --git a/Engine/Source/Programs/UnrealBuildAccelerator/Common/Public/UbaScheduler.h b/Engine/Source/Programs/UnrealBuildAccelerator/Common/Public/UbaScheduler.h
--- a/Engine/Source/Programs/UnrealBuildAccelerator/Common/Public/UbaScheduler.h
+++ b/Engine/Source/Programs/UnrealBuildAccelerator/Common/Public/UbaScheduler.h
@@ -141,15 +141,26 @@ namespace uba
 		UBA_NOINLINE void ThreadLoop();
 
 		struct ProcessEntry;
+		struct ProcessLink;
+		struct ProcessQueue;
+
 		ProcessEntry CreateProcessEntry(const EnqueueProcessInfo& info);
+		ProcessQueue& GetReadyToRunQueue(ProcessEntry& entry);
+		void AddToQueueBackNoLock(ProcessQueue& queue, u32 entryIndex);
+		void AddToQueueFrontNoLock(ProcessQueue& queue, u32 entryIndex);
+		u32 PopFromQueue(ProcessQueue& queue);
+		ProcessEntry* PeekFromQueue(ProcessQueue& queue);
+		u32 GetQueueCount(ProcessQueue& queue);
+		void ClearQueue(ProcessQueue& queue);
+		void SetupDependents(ProcessEntry& entry, u32 entryIndex, const EnqueueProcessInfo& info);
+		void UpdateDependentsNoLock(RunStatus newRunStatus, ProcessLink* firstDependent);
 		void SkipAllQueued();
 		bool ProcessReturned(Process& process, bool isLocal);
 		void HandleCacheMissed(u32 processIndex);
 		void RemoteSlotAvailable(bool isCrossArchitecture, bool isCrossPlatform);
 		void ProcessExited(ExitProcessInfo* info, const ProcessHandle& handle);
 		bool CanRun();
-		u32 PopCache();
-		u32 PopProcess(bool isLocal, bool isCrossArchitecture, bool isCrossPlatform, RunStatus& outPrevStatus, bool forReuse, bool& outShouldSleep);
+		u32 PopProcessNoLock(bool isLocal, bool isCrossArchitecture, bool isCrossPlatform, RunStatus& outPrevStatus, bool forReuse, bool& outShouldSleep);
 		bool RunCacheQuery();
 		bool RunQueuedProcess(bool isLocal, bool isCrossArchitecture, bool isCrossPlatform);
 		bool RaceRemoteProcess();
@@ -157,9 +168,10 @@ namespace uba
 		void ExitProcess(ExitProcessInfo& info, Process& process, u32 exitCode, bool fromCache, bool forReuse);
 		void SkipProcess(ProcessStartInfo2& info);
 		void UpdateStatusHitMissCount();
-		void UpdateQueueCounter(int offset);
+		void UpdateQueueCounter(int offset, bool canRunRemotely);
 		void UpdateActiveProcessCounter(bool isLocal, int offset);
 		void FinishProcess(const ProcessHandle& handle);
+		void UpdateRunStatus(ProcessEntry& entry, u32 entryIndex, RunStatus newStatus);
 
 		void KillNewestLocalProcess(StringView reason);
 		void NativeProcessCreated(ProcessImpl& process);
@@ -170,16 +182,10 @@ namespace uba
 		void SaveMemTrackTable();
 		UBA_NOINLINE void ThreadMemoryCheckLoop();
 
-
-		SessionServer& m_session;
-		u32 m_maxLocalProcessors;
-		float m_maxRacingWeight;
-		u32 m_maxParallelCacheQueries;
-		
 		struct ProcessEntry
 		{
 			ProcessStartInfo2* info;
-			u32* dependencies;
+			ProcessLink* firstDependent;
 			u32 dependencyCount;
 			CacheStatus cacheStatus;
 			RunStatus runStatus;
@@ -189,16 +195,28 @@ namespace uba
 			bool canCrossPlatform;
 		};
 
+		struct ProcessQueue : Set<u32> {};
+
+		SessionServer& m_session;
+		u32 m_maxLocalProcessors;
+		float m_maxRacingWeight;
+		u32 m_maxParallelCacheQueries;
+
 		Futex m_processEntriesLock;
 		Vector<ProcessEntry> m_processEntries;
 
+		ProcessQueue m_readyToTestCache;
+		ProcessQueue m_readyToDownloadCache;
+		ProcessQueue m_readyToRunLocal; // Can only run locally
+		ProcessQueue m_readyToRunRemote; // Can run remotely on same architecture
+		ProcessQueue m_readyToRunCrossArchitecture; // Can run remotely and also cross architecture
+		ProcessQueue m_readyToRunCrossPlatform; // Can run remotely and also cross platform
+		ProcessQueue m_readyToSkip;
+
 		Function<void(const ProcessHandle&)> m_processFinished;
 
 		Event m_updateThreadWakeup;
 		Thread m_updateThread;
-		u32 m_processEntriesRunStart = 0;
-
-		u32 m_processEntriesCacheStart = 0;
 
 		Atomic<bool> m_loop = false;
 		bool m_enableProcessReuse;
@@ -219,11 +237,12 @@ namespace uba
 		Atomic<u32> m_activeCacheDownloads = 0;
 		Atomic<u32> m_getNextQueueCount = 0;
 
-		Atomic<u32> m_totalProcesses = 0;
-		Atomic<u32> m_queuedProcesses = 0;
-		Atomic<u32> m_activeLocalProcesses = 0;
-		Atomic<u32> m_activeRemoteProcesses = 0;
-		Atomic<u32> m_finishedProcesses = 0;
+		Atomic<u32> m_totalProcesses = 0; // Process count registed in to scheduler
+		Atomic<u32> m_queuedProcesses = 0; // Number of processes that has not started yet
+		Atomic<u32> m_activeLocalProcesses = 0; // Number of processes active on local machine
+		Atomic<u32> m_activeRemoteProcesses = 0; // Number of processes active on remote helpers
+		Atomic<u32> m_finishedProcesses = 0; // Processes finished (includes skipped)
+		Atomic<u32> m_queuedProcessesCanExecuteRemotely = 0; 
 		Atomic<u32> m_errorCount = 0;
 		Atomic<u32> m_cacheHitCount = 0;
 		Atomic<u32> m_cacheMissCount = 0;
