[UBA]
* Fixed bug in process reuse where overlay memory pointer was reset even though it shouldn't

diff --git a/Engine/Source/Programs/UnrealBuildAccelerator/Detours/Private/PosixOS/UbaDetoursFunctionsPosix.cpp b/Engine/Source/Programs/UnrealBuildAccelerator/Detours/Private/PosixOS/UbaDetoursFunctionsPosix.cpp
--- a/Engine/Source/Programs/UnrealBuildAccelerator/Detours/Private/PosixOS/UbaDetoursFunctionsPosix.cpp
+++ b/Engine/Source/Programs/UnrealBuildAccelerator/Detours/Private/PosixOS/UbaDetoursFunctionsPosix.cpp
@@ -3157,13 +3157,15 @@ extern "C"
 			LogHeader(ToView(outArguments));
 			#endif
 
+			const u8* overlayMemory = g_directoryTable.m_overlayMemory;
 			g_directoryTable.DeinitOverlay();
+			g_directoryTable.m_overlayMemory = overlayMemory;
 
 			if (u32 overlaySize = directoryTableSize.overlay)
 			{
 				directoryTableSize.overlay = 0; // Want to just parse normal directory first
 				g_directoryTable.ParseDirectoryTableNoLock(directoryTableSize);
-				g_directoryTable.InitOverlay(g_directoryTable.m_overlayMemory, overlaySize);
+				g_directoryTable.InitOverlay(overlayMemory, overlaySize);
 			}
 		}
 
diff --git a/Engine/Source/Programs/UnrealBuildAccelerator/Detours/Private/Windows/UbaDetoursFunctionsWin.cpp b/Engine/Source/Programs/UnrealBuildAccelerator/Detours/Private/Windows/UbaDetoursFunctionsWin.cpp
--- a/Engine/Source/Programs/UnrealBuildAccelerator/Detours/Private/Windows/UbaDetoursFunctionsWin.cpp
+++ b/Engine/Source/Programs/UnrealBuildAccelerator/Detours/Private/Windows/UbaDetoursFunctionsWin.cpp
@@ -1334,13 +1334,15 @@ extern "C"
 			LogHeader(ToView(outArguments));
 			#endif
 
+			const u8* overlayMemory = g_directoryTable.m_overlayMemory;
 			g_directoryTable.DeinitOverlay();
+			g_directoryTable.m_overlayMemory = overlayMemory;
 
 			if (u32 overlaySize = directoryTableSize.overlay)
 			{
 				directoryTableSize.overlay = 0; // Want to just parse normal directory first
 				g_directoryTable.ParseDirectoryTableNoLock(directoryTableSize);
-				g_directoryTable.InitOverlay(g_directoryTable.m_overlayMemory, overlaySize);
+				g_directoryTable.InitOverlay(overlayMemory, overlaySize);
 			}
 		}
 
