[UBA
* Fixed bug in GetDirKey not handling non-terminated strings properly
* Fixed bug where directory table could be end up not parsing main table if overlay table was populated first
* Added a bunch of unit tests
* Cleaned up some unit test code

diff --git a/Engine/Source/Programs/UnrealBuildAccelerator/Common/Private/UbaDirectoryTableHolder.cpp b/Engine/Source/Programs/UnrealBuildAccelerator/Common/Private/UbaDirectoryTableHolder.cpp
--- a/Engine/Source/Programs/UnrealBuildAccelerator/Common/Private/UbaDirectoryTableHolder.cpp
+++ b/Engine/Source/Programs/UnrealBuildAccelerator/Common/Private/UbaDirectoryTableHolder.cpp
@@ -22,13 +22,12 @@ namespace uba
 
 	bool GetDirKey(StringKey& outDirKey, StringBufferBase& outDirPath, const tchar*& outLastSlash, StringView filePath)
 	{
-		outLastSlash = TStrrchr(filePath.data, PathSeparator);
-		UBA_ASSERTF(outLastSlash, TC("Can't get dir key for path %s"), filePath.data);
-		if (!outLastSlash)
+		StringView dirPath = filePath.GetPath();
+		outLastSlash = dirPath.data + dirPath.count;
+		UBA_ASSERTF(dirPath.count, TC("Can't get dir key for path %s"), filePath.data);
+		if (!dirPath.count)
 			return false;
-
-		u64 dirLen = u64(outLastSlash - filePath.data);
-		outDirPath.Append(filePath.data, dirLen);
+		outDirPath.Append(dirPath);
 		outDirKey = CaseInsensitiveFs ? ToStringKeyLower(outDirPath) : ToStringKey(outDirPath);
 		return true;
 	}
@@ -301,7 +300,7 @@ namespace uba
 		FileInformation fileInfo;
 		if (!getFileInformation(fileInfo, filePathKey))
 			return false;
-		return WriteOverlayEntry(overlay, dirKey, fixedFilePath, fileName, fileInfo);
+		return WriteOverlayEntry(overlay, dirKey, fixedFilePath, fileName, fileInfo, !m_runningRemote);
 	}
 
 	bool DirectoryTableHolder::AddOverlayFile(DirectoryTableOverlay& overlay, StringView filePath)
@@ -327,19 +326,10 @@ namespace uba
 		StringBuffer<> fixedFilePath;
 		StringKeyHasher hasher;
 		StringKey filePathKey = GetKeyAndFixedName2(fixedFilePath, hasher, filePath.data, &dirKey, &fileName);
-
-		// We need to make sure as far down as possible that the directories are created.. overlay entries must always somehow be attached to the end of the normal table
-		// TODO: This must ALSO be made for remotes
-		//if (!m_runningRemote)
-		//{
-		//	StringView dirName(fixedFilePath.data, u32(fileName.data - fixedFilePath.data) - 1);
-		//	WriteDirectoryEntries(dirKey, dirName);
-		//}
-
-		return WriteOverlayEntry(overlay, dirKey, fixedFilePath, fileName, fileInfo);
+		return WriteOverlayEntry(overlay, dirKey, fixedFilePath, fileName, fileInfo, !m_runningRemote);
 	}
 
-	bool DirectoryTableHolder::WriteOverlayEntry(DirectoryTableOverlay& overlay, StringKey dirKey, StringView filePath, StringView fileName, const FileInformation& fileInfo)
+	bool DirectoryTableHolder::WriteOverlayEntry(DirectoryTableOverlay& overlay, StringKey dirKey, StringView filePath, StringView fileName, const FileInformation& fileInfo, bool populateMainTableParentDir)
 	{
 		SCOPED_WRITE_LOCK(overlay.lock, lock);
 
@@ -349,6 +339,18 @@ namespace uba
 		if (!dir.offset)
 		{
 			StringView dirPath(filePath.data, u32(fileName.data - filePath.data) - 1);
+
+			// I have a feeling there is a bug here somehow but can't create it in unit tests.
+			// Essentially the theory is that if a remote helper runs in a directory tree and gets the attribute of a overlay file (prepopulated)
+			// That will create the entire hierarchy to the overlay file.
+			// It then queries a sibling directory for a file. Now the dir hierarchy exists but sibling directory will not exist since the directory table
+			// does not know that it has not queried the directory for the main table.
+			if (populateMainTableParentDir)
+			{
+				UBA_ASSERT(!m_runningRemote);
+				WriteDirectoryEntries(dirKey, dirPath);
+			}
+
 			WriteOverlayDirectoryEntriesRecursive(overlay, dirKey, dirPath);
 		}
 
@@ -464,7 +466,7 @@ namespace uba
 		if (!shouldWriteToDisk)
 		{
 			UBA_FOR_ASSERT(u32 res =) WriteDirectoryEntries(dirKey, dirName);
-			UBA_ASSERTF(res, TC("Failed to write directory entries for %s"), dirName.data);
+			UBA_ASSERTF(res, TC("Failed to write directory entries for %.*s"), dirName.count, dirName.data);
 		}
 
 		SCOPED_READ_LOCK(dirTable.m_lookupLock, lookupCs);
@@ -514,8 +516,8 @@ namespace uba
 		if (shouldWriteToDisk)
 		{
 			FileInformation info;
-			if (!GetFileInformation(info, m_logger, filePath.data))
-				return m_logger.Error(TC("Failed to get file information for %s while checking file added for write. This should not happen! (%s)"), filePath.data, LastErrorToText().data);
+			if (!GetFileInformation(info, m_logger, filePath))
+				return m_logger.Error(TC("Failed to get file information for %.*s while checking file added for write. This should not happen! (%s)"), filePath.count, filePath.data, LastErrorToText().data);
 
 			attributes = info.attributes;
 			volumeSerial = info.volumeSerialNumber;
@@ -608,7 +610,7 @@ namespace uba
 			info.lastWriteTime = lastWriteTime;
 			info.size = fileSize;
 			info.volumeSerialNumber = volumeIndex;
-			if (!WriteOverlayEntry(*overlay, dirKey, filePath, ToView(lastSlash + 1), info))
+			if (!WriteOverlayEntry(*overlay, dirKey, filePath, ToView(lastSlash + 1), info, false))
 				return false;
 		}
 
@@ -638,7 +640,7 @@ namespace uba
 		{
 			// TODO: If it does not exist in overlay nor main it should not create any entry
 			FileInformation fileInfo;
-			if (!WriteOverlayEntry(*overlay, dirKey, filePath, ToView(lastSlash+1), fileInfo))
+			if (!WriteOverlayEntry(*overlay, dirKey, filePath, ToView(lastSlash+1), fileInfo, !m_runningRemote))
 				return false;
 		}
 
diff --git a/Engine/Source/Programs/UnrealBuildAccelerator/Common/Private/UbaFile.cpp b/Engine/Source/Programs/UnrealBuildAccelerator/Common/Private/UbaFile.cpp
--- a/Engine/Source/Programs/UnrealBuildAccelerator/Common/Private/UbaFile.cpp
+++ b/Engine/Source/Programs/UnrealBuildAccelerator/Common/Private/UbaFile.cpp
@@ -77,19 +77,36 @@ namespace uba
 		return (HANDLE)(fh == InvalidFileHandle ? InvalidFileHandle : (fh & FileHandleFlagMask));
 	}
 
-	#define PREFIX_FILENAME(fileName, prefix) \
-		UBA_ASSERT(TStrlen(fileName) < MaxPath); \
-		StringBuffer<MaxPath> STRING_JOIN(longName, __LINE__); \
-		if (IsAbsolutePath(fileName) && !IsUncPath(fileName)) \
-		{ \
-			auto& lsb = STRING_JOIN(longName, __LINE__); \
-			lsb.Append(TC(prefix)); \
-			FixPath(fileName, nullptr, 0, lsb); \
-			fileName = lsb.data; \
+	template<typename T>
+	void PrefixFileName(T& fileName, StringView prefix, StringBufferBase& temp)
+	{
+		StringView sv(ToView(fileName));
+		if (IsAbsolutePath(sv) && !IsUncPath(sv))
+		{
+			UBA_ASSERT(sv.count < MaxPath);
+			temp.Append(prefix);
+			FixPath(temp, sv, {});
+			if constexpr (std::is_same_v<T, StringView>)
+				fileName = temp;
+			else
+				fileName = temp.data;
+		}
+		else if constexpr (std::is_same_v<T, StringView>)
+		{
+			if (fileName.data[fileName.count] != 0) // must be null terminated
+			{
+				temp.Append(fileName);
+				fileName = temp;
+			}
 		}
-	#define MAKE_LONG_FILENAME(fileName) PREFIX_FILENAME(fileName, "\\\\?\\")
-	#define MAKE_NT_FILENAME(fileName)  PREFIX_FILENAME(fileName, "\\??\\")
+		
+	}
 
+	#define PREFIX_FILENAME(fileName, prefix) \
+		StringBuffer<MaxPath> STRING_JOIN(temp, __LINE__); \
+		PrefixFileName(fileName, prefix, STRING_JOIN(temp, __LINE__));
+	#define MAKE_LONG_FILENAME(fileName) PREFIX_FILENAME(fileName, TCV("\\\\?\\"))
+	#define MAKE_NT_FILENAME(fileName)  PREFIX_FILENAME(fileName, TCV("\\??\\"))
 #else
 	int asFileDescriptor(FileHandle fh)
 	{
@@ -97,6 +114,13 @@ namespace uba
 			return (int)fh;
 		return (int)(fh & FileHandleFlagMask);
 	}
+	#define ENSURE_TERMINATED(fileName) \
+		StringBuffer<MaxPath> STRING_JOIN(temp, __LINE__); \
+		if (fileName.data[fileName.count] != 0) \
+		{ \
+			STRING_JOIN(temp, __LINE__).Append(fileName); \
+			fileName = STRING_JOIN(temp, __LINE__); \
+		}
 #endif
 
 
@@ -249,7 +273,7 @@ namespace uba
 		return true;
 #else
 		FileInformation info;
-		if (!GetFileInformation(info, logger, fileName))
+		if (!GetFileInformation(info, logger, ToView(fileName)))
 			return false;
 		out.attributes = info.attributes;
 		out.lastWriteTime = info.lastWriteTime;
@@ -285,24 +309,25 @@ namespace uba
 #endif
 	}
 
-	bool GetFileInformation(FileInformation& out, Logger& logger, const tchar* fileName)
+	bool GetFileInformation(FileInformation& out, Logger& logger, StringView fileName)
 	{
 #if PLATFORM_WINDOWS
 		MAKE_LONG_FILENAME(fileName);
-		FileHandle h = uba::CreateFileW(fileName, 0, 0, OPEN_EXISTING, FILE_FLAG_BACKUP_SEMANTICS);
+		FileHandle h = uba::CreateFileW(fileName.data, 0, 0, OPEN_EXISTING, FILE_FLAG_BACKUP_SEMANTICS);
 		if (h == InvalidFileHandle)
 			return false;// logger.Error(TC("GetFileInformation: CreateFile failed for file %s (%s)"), fileName, LastErrorToText().data);
-		auto handleGuard = MakeGuard([&]() { uba::CloseFile(fileName, h); });
-		if (!GetFileInformationByHandle(out, logger, fileName, h))
+		auto handleGuard = MakeGuard([&]() { uba::CloseFile(fileName.data, h); });
+		if (!GetFileInformationByHandle(out, logger, fileName.data, h))
 			return false;// logger.Error(TC("Failed to get file information for %s while checking file added for write. This should not happen! (%s)"), fileName, LastErrorToText().data);
 		return true;
 #else
+		ENSURE_TERMINATED(fileName);
 		struct stat attr;
-		int res = stat(fileName, &attr);
+		int res = stat(fileName.data, &attr);
 		if (res != 0)
 		{
 			if (errno != ENOENT)
-				logger.Warning(TC("GetFileInformation: stat failed for file %s and this error handling not implemented (%s)"), fileName, strerror(errno));
+				logger.Warning(TC("GetFileInformation: stat failed for file %s and this error handling not implemented (%s)"), fileName.data, strerror(errno));
 			SetLastError(ERROR_FILE_NOT_FOUND);
 			return false;
 		}
@@ -926,7 +951,7 @@ namespace uba
 #if PLATFORM_WINDOWS
 		return FILE_ATTRIBUTE_NORMAL;
 #else
-		return S_IRUSR | S_IWUSR | (execute ? S_IXUSR : 0) | S_IRGRP | S_IROTH;
+		return S_IFREG | S_IRUSR | S_IWUSR | (execute ? S_IXUSR : 0) | S_IRGRP | S_IROTH;
 #endif
 	}
 
diff --git a/Engine/Source/Programs/UnrealBuildAccelerator/Common/Private/UbaLogger.cpp b/Engine/Source/Programs/UnrealBuildAccelerator/Common/Private/UbaLogger.cpp
--- a/Engine/Source/Programs/UnrealBuildAccelerator/Common/Private/UbaLogger.cpp
+++ b/Engine/Source/Programs/UnrealBuildAccelerator/Common/Private/UbaLogger.cpp
@@ -110,19 +110,17 @@ namespace uba
 	{
 		va_list arg;
 		va_start(arg, format);
-		tchar buffer[1024];
-		int count = Tvsprintf_s(buffer, 1024, format, arg);
-		if (count <= 0)
-			TStrcpy_s(buffer, 1024, format);
+		StringBuffer<2048> buffer;
+		buffer.Append(format, arg);
 		va_end(arg);
 #if PLATFORM_WINDOWS
-		wprintf(TC("UBA FATAL ERROR %u: %s\n"), code, buffer);
+		wprintf(TC("UBA FATAL ERROR %u: %s\n"), code, buffer.data);
 		fflush(stdout);
 		if (IsDebuggerPresent())
 			DebugBreak();
 		ExitProcess(code);
 #else
-		printf(TC("UBA FATAL ERROR %u: %s\n"), code, buffer);
+		printf(TC("UBA FATAL ERROR %u: %s\n"), code, buffer.data);
 		fflush(stdout);
 		kill(getpid(), SIGKILL);
 #endif
diff --git a/Engine/Source/Programs/UnrealBuildAccelerator/Common/Public/UbaDirectoryIterator.h b/Engine/Source/Programs/UnrealBuildAccelerator/Common/Public/UbaDirectoryIterator.h
--- a/Engine/Source/Programs/UnrealBuildAccelerator/Common/Public/UbaDirectoryIterator.h
+++ b/Engine/Source/Programs/UnrealBuildAccelerator/Common/Public/UbaDirectoryIterator.h
@@ -4,6 +4,7 @@
 
 #include "UbaFile.h"
 #include "UbaLogger.h"
+#include "UbaPathUtils.h"
 #include "UbaProcessStats.h"
 #include "UbaStringBuffer.h"
 
@@ -45,7 +46,7 @@ namespace uba
 		StringBuffer<> str;
 		str.Append(TCV("\\??\\"));
 
-		if (IsUncPath(dirPath.data))
+		if (IsUncPath(dirPath))
 		{
 			if (!dirPath.Skip(2).Contains('\\'))
 				return false;
diff --git a/Engine/Source/Programs/UnrealBuildAccelerator/Common/Public/UbaDirectoryTableHolder.h b/Engine/Source/Programs/UnrealBuildAccelerator/Common/Public/UbaDirectoryTableHolder.h
--- a/Engine/Source/Programs/UnrealBuildAccelerator/Common/Public/UbaDirectoryTableHolder.h
+++ b/Engine/Source/Programs/UnrealBuildAccelerator/Common/Public/UbaDirectoryTableHolder.h
@@ -57,7 +57,7 @@ namespace uba
 		bool AddOverlayFile(DirectoryTableOverlay& overlay, StringView filePath);
 		bool MoveOverlayFile(DirectoryTableOverlay& overlay, StringView from, StringView to);
 		bool WriteOverlayEntry(DirectoryTableOverlay& overlay, StringView filePath, const FileInformation& fileInfo);
-		bool WriteOverlayEntry(DirectoryTableOverlay& overlay, StringKey dirKey, StringView filePath, StringView fileName, const FileInformation& fileInfo);
+		bool WriteOverlayEntry(DirectoryTableOverlay& overlay, StringKey dirKey, StringView filePath, StringView fileName, const FileInformation& fileInfo, bool populateMainTableParentDir);
 		bool WriteOverlayEntry(DirectoryTableOverlay& overlay, StringKey dirKey, DirectoryTableOverlay::Dir& dir, StringView fileName, const FileInformation& fileInfo, u32& outFileOffset);
 		bool WriteOverlayDirectoryEntriesRecursive(DirectoryTableOverlay& overlay, StringKey filePathKey, StringView filePath);
 
diff --git a/Engine/Source/Programs/UnrealBuildAccelerator/Common/Public/UbaFile.h b/Engine/Source/Programs/UnrealBuildAccelerator/Common/Public/UbaFile.h
--- a/Engine/Source/Programs/UnrealBuildAccelerator/Common/Public/UbaFile.h
+++ b/Engine/Source/Programs/UnrealBuildAccelerator/Common/Public/UbaFile.h
@@ -62,7 +62,7 @@ namespace uba
 		u64 index = 0;
 	};
 	bool GetFileInformationByHandle(FileInformation& out, Logger& logger, const tchar* fileName, FileHandle hFile);
-	bool GetFileInformation(FileInformation& out, Logger& logger, const tchar* fileName);
+	bool GetFileInformation(FileInformation& out, Logger& logger, StringView fileName);
 
 	bool ReadFile(Logger& logger, const tchar* fileName, FileHandle fileHandle, void* b, u64 bufferLen);
 	bool WriteFile(Logger& logger, const tchar* fileName, FileHandle fileHandle, const void* b, u64 bufferLen);
diff --git a/Engine/Source/Programs/UnrealBuildAccelerator/Core/Private/UbaStringBuffer.cpp b/Engine/Source/Programs/UnrealBuildAccelerator/Core/Private/UbaStringBuffer.cpp
--- a/Engine/Source/Programs/UnrealBuildAccelerator/Core/Private/UbaStringBuffer.cpp
+++ b/Engine/Source/Programs/UnrealBuildAccelerator/Core/Private/UbaStringBuffer.cpp
@@ -518,6 +518,11 @@ namespace uba
 		return s ? StringView(s, TStrlen(s)) : StringView();;
 	}
 
+	StringView ToView(StringView sv)
+	{
+		return sv;
+	}
+
 	////////////////////////////////////////////////////////////////////////////////////////////////////
 
 	LastErrorToText::LastErrorToText() : LastErrorToText(GetLastError())
diff --git a/Engine/Source/Programs/UnrealBuildAccelerator/Core/Public/UbaPathUtils.h b/Engine/Source/Programs/UnrealBuildAccelerator/Core/Public/UbaPathUtils.h
--- a/Engine/Source/Programs/UnrealBuildAccelerator/Core/Public/UbaPathUtils.h
+++ b/Engine/Source/Programs/UnrealBuildAccelerator/Core/Public/UbaPathUtils.h
@@ -12,6 +12,8 @@ namespace uba
 	#define Local_GetLongPathNameW GetLongPathNameW
 	#endif
 
+	inline bool IsUncPath(const tchar* path) { return IsWindows ? (path[0] == '\\' && path[1] == '\\') : false; }
+	inline bool IsUncPath(StringView path) { return IsUncPath(path.data); }
 
 	inline bool FixPath2(const tchar* fileName, u32 fileNameLen, const tchar* workingDir, u64 workingDirCharLen, tchar* buffer, u64 bufferCharCapacity, u32* outBufferCharLen)
 	{
@@ -399,4 +401,14 @@ namespace uba
 		else
 			return path[0] == '/';
 	}
+
+	inline bool IsAbsolutePath(StringView path)
+	{
+		if (!path.count)
+			return false;
+		if constexpr (IsWindows)
+			return (path[0] && path[1] == ':') || IsUncPath(path);
+		else
+			return path[0] == '/';
+	}
 }
diff --git a/Engine/Source/Programs/UnrealBuildAccelerator/Core/Public/UbaPlatform.h b/Engine/Source/Programs/UnrealBuildAccelerator/Core/Public/UbaPlatform.h
--- a/Engine/Source/Programs/UnrealBuildAccelerator/Core/Public/UbaPlatform.h
+++ b/Engine/Source/Programs/UnrealBuildAccelerator/Core/Public/UbaPlatform.h
@@ -152,7 +152,6 @@ namespace uba
 	inline LARGE_INTEGER ToLargeInteger(u32 high, u32 low) { LARGE_INTEGER li; li.HighPart = (LONG)high; li.LowPart = low; return li; }
 	inline bool IsDirectory(u32 attributes) { return (attributes & FILE_ATTRIBUTE_DIRECTORY) != 0; }
 	inline bool IsExecutable(u32 attributes) { return false; }
-	inline bool IsUncPath(const tchar* path) { return path[0] == '\\' && path[1] == '\\'; }
 	inline tchar* TStrdup(const tchar* s) { return _wcsdup(s); }
 	#define TStrlen(s) u32(wcslen(s))
 	#define TStrchr(a, b) wcschr(a, b)
@@ -197,7 +196,6 @@ namespace uba
 	inline timespec ToTimeSpec(u64 time) { timespec ts; ts.tv_sec = time / 10'000'000ull; ts.tv_nsec = (time - (u64(ts.tv_sec) * 10'000'000ull)) * 100; return ts; }
 	inline bool IsDirectory(u32 attributes) { return S_ISDIR(attributes); }
 	inline bool IsExecutable(u32 attributes) { return (attributes & (S_IXUSR | S_IXGRP | S_IXOTH)) != 0; }
-	inline bool IsUncPath(const tchar* path) { return false; }
 	inline tchar* TStrdup(const tchar* s) { return strdup(s); }
 	#if PLATFORM_LINUX
 	inline constexpr bool CaseInsensitiveFs = false;
diff --git a/Engine/Source/Programs/UnrealBuildAccelerator/Core/Public/UbaStringBuffer.h b/Engine/Source/Programs/UnrealBuildAccelerator/Core/Public/UbaStringBuffer.h
--- a/Engine/Source/Programs/UnrealBuildAccelerator/Core/Public/UbaStringBuffer.h
+++ b/Engine/Source/Programs/UnrealBuildAccelerator/Core/Public/UbaStringBuffer.h
@@ -133,6 +133,7 @@ namespace uba
 	};
 
 	StringView ToView(const tchar* s);
+	StringView ToView(StringView sv);
 
 	template<typename T, size_t size>
 	constexpr StringView AsView(T (&buffer)[size]) { return StringView(buffer, size-1); }
diff --git a/Engine/Source/Programs/UnrealBuildAccelerator/Test/Private/UbaTest.h b/Engine/Source/Programs/UnrealBuildAccelerator/Test/Private/UbaTest.h
--- a/Engine/Source/Programs/UnrealBuildAccelerator/Test/Private/UbaTest.h
+++ b/Engine/Source/Programs/UnrealBuildAccelerator/Test/Private/UbaTest.h
@@ -6,9 +6,16 @@
 #include "UbaStringBuffer.h"
 
 #define CHECK_TRUE(x) \
+	do { \
 	if (!(x)) \
-		return logger.Error(TC("Failed %s (%s:%u)"), TC(#x), TC("") __FILE__, __LINE__);
+		return logger.Error(TC("Failed %s (%s:%u)"), TC(#x), TC("") __FILE__, __LINE__); \
+	} while (false)
 
+#define CHECK_TRUEF(x, ...) \
+	do { \
+	if (!(x)) \
+		return logger.Logf(LogEntryType_Error, ##__VA_ARGS__).ToFalse(); \
+	} while (false)
 
 namespace uba
 {
@@ -25,8 +32,8 @@ namespace uba
 	bool FilesEqual(Logger& logger, const tchar* file1, const tchar* file2);
 
 	using RunProcessFunction = Function<ProcessHandle(ProcessStartInfo&)>;
-	using TestSessionFunction = Function<bool(LoggerWithWriter& logger, SessionServer& session, const tchar* workingDir, const RunProcessFunction& runProcess, void* extraData)>;
-	bool RunLocal(LoggerWithWriter& logger, const StringBufferBase& testRootDir, const TestSessionFunction& testFunc, void* extraData = nullptr, bool enableDetour = true);
+	using TestSessionFunction = Function<bool(LoggerWithWriter& logger, SessionServer& session, const tchar* workingDir, const RunProcessFunction& runProcess, uintptr_t extraData)>;
+	bool RunLocal(LoggerWithWriter& logger, const StringBufferBase& testRootDir, const TestSessionFunction& testFunc, uintptr_t extraData = 0, bool enableDetour = true);
 	bool RunRemote(LoggerWithWriter& logger, const StringBufferBase& testRootDir, const TestSessionFunction& testFunc, bool deleteAll = true, bool serverShouldListen = true);
 	void GetTestAppPath(LoggerWithWriter& logger, StringBufferBase& out);
 	const tchar* GetSystemApplication();
diff --git a/Engine/Source/Programs/UnrealBuildAccelerator/Test/Private/UbaTestAll.h b/Engine/Source/Programs/UnrealBuildAccelerator/Test/Private/UbaTestAll.h
--- a/Engine/Source/Programs/UnrealBuildAccelerator/Test/Private/UbaTestAll.h
+++ b/Engine/Source/Programs/UnrealBuildAccelerator/Test/Private/UbaTestAll.h
@@ -88,6 +88,8 @@ namespace uba
 		UBA_TEST(TestRegisterVirtualFile) \
 		UBA_TEST(TestRemoteRegisterVirtualFile) \
 		UBA_TEST(TestRemoteProcessSpecialCase1) \
+		UBA_TEST(TestOverlayAccessOrdering) \
+		UBA_TEST(TestOverlayAccessOrderingRemote) \
 		UBA_TEST(TestSessionSpecialCopy) \
 		UBA_TEST(TestFileMapWrite) \
 		UBA_TEST(TestFileMapWriteRemote) \
diff --git a/Engine/Source/Programs/UnrealBuildAccelerator/Test/Private/UbaTestScheduler.cpp b/Engine/Source/Programs/UnrealBuildAccelerator/Test/Private/UbaTestScheduler.cpp
--- a/Engine/Source/Programs/UnrealBuildAccelerator/Test/Private/UbaTestScheduler.cpp
+++ b/Engine/Source/Programs/UnrealBuildAccelerator/Test/Private/UbaTestScheduler.cpp
@@ -10,7 +10,7 @@ namespace uba
 {
 	bool TestLocalSchedule(LoggerWithWriter& logger, const StringBufferBase& testRootDir)
 	{
-		return RunLocal(logger, testRootDir, [](LoggerWithWriter& logger, SessionServer& session, const tchar* workingDir, const RunProcessFunction& runProcess, void* extraData)
+		return RunLocal(logger, testRootDir, [](LoggerWithWriter& logger, SessionServer& session, const tchar* workingDir, const RunProcessFunction& runProcess, uintptr_t extraData)
 			{
 				SchedulerCreateInfo info(session);
 				Scheduler scheduler(info);
@@ -74,7 +74,7 @@ namespace uba
 
 	bool TestLocalScheduleReuse(LoggerWithWriter& logger, const StringBufferBase& testRootDir)
 	{
-		return RunLocal(logger, testRootDir, [](LoggerWithWriter& logger, SessionServer& session, const tchar* workingDir, const RunProcessFunction& runProcess, void* extraData)
+		return RunLocal(logger, testRootDir, [](LoggerWithWriter& logger, SessionServer& session, const tchar* workingDir, const RunProcessFunction& runProcess, uintptr_t extraData)
 			{
 				return RunSchedulerReuseTest(logger, session, workingDir, true);
 			});
@@ -82,7 +82,7 @@ namespace uba
 
 	bool TestRemoteScheduleReuse(LoggerWithWriter& logger, const StringBufferBase& testRootDir)
 	{
-		return RunRemote(logger, testRootDir, [](LoggerWithWriter& logger, SessionServer& session, const tchar* workingDir, const RunProcessFunction& runProcess, void* extraData)
+		return RunRemote(logger, testRootDir, [](LoggerWithWriter& logger, SessionServer& session, const tchar* workingDir, const RunProcessFunction& runProcess, uintptr_t extraData)
 			{
 				return RunSchedulerReuseTest(logger, session, workingDir, false);
 			});
diff --git a/Engine/Source/Programs/UnrealBuildAccelerator/Test/Private/UbaTestSession.cpp b/Engine/Source/Programs/UnrealBuildAccelerator/Test/Private/UbaTestSession.cpp
--- a/Engine/Source/Programs/UnrealBuildAccelerator/Test/Private/UbaTestSession.cpp
+++ b/Engine/Source/Programs/UnrealBuildAccelerator/Test/Private/UbaTestSession.cpp
@@ -21,7 +21,7 @@ namespace uba
 		return 10'000;
 	}
 
-	bool RunLocal(LoggerWithWriter& logger, const StringBufferBase& testRootDir, const Config& config, const TestSessionFunction& testFunc, void* extraData = nullptr, bool enableDetour = true)
+	bool RunLocal(LoggerWithWriter& logger, const StringBufferBase& testRootDir, const Config& config, const TestSessionFunction& testFunc, uintptr_t extraData = 0, bool enableDetour = true)
 	{
 		LogWriter& logWriter = logger.m_writer;
 
@@ -62,7 +62,7 @@ namespace uba
 		return testFunc(logger, session, workingDir.data, [&](ProcessStartInfo& pi) { return session.RunProcess(pi, true, enableDetour); }, extraData);
 	}
 
-	bool RunLocal(LoggerWithWriter& logger, const StringBufferBase& testRootDir, const TestSessionFunction& testFunc, void* extraData, bool enableDetour)
+	bool RunLocal(LoggerWithWriter& logger, const StringBufferBase& testRootDir, const TestSessionFunction& testFunc, uintptr_t extraData, bool enableDetour)
 	{
 		return RunLocal(logger, testRootDir, {}, testFunc, extraData, enableDetour);
 	}
@@ -74,7 +74,7 @@ namespace uba
 		out.Append(IsWindows ? TC("UbaTestApp.exe") : TC("UbaTestApp"));
 	}
 
-	bool RunTestApp(LoggerWithWriter& logger, const tchar* workingDir, const RunProcessFunction& runProcess, const tchar* arguments, ProcessHandle* outHandle = nullptr)
+	bool RunTestApp(LoggerWithWriter& logger, const tchar* workingDir, const RunProcessFunction& runProcess, const tchar* arguments, ProcessHandle* outHandle = nullptr, const tchar* overlayFiles = TC(""), u32 expectedExitCode = 0)
 	{
 		StringBuffer<> testApp;
 		GetTestAppPath(logger, testApp);
@@ -83,18 +83,20 @@ namespace uba
 		processInfo.application = testApp.data;
 		processInfo.workingDir = workingDir;
 		processInfo.arguments = arguments;
+		processInfo.overlayFiles = overlayFiles;
 		processInfo.logLineFunc = [](void* userData, const tchar* line, u32 length, LogEntryType type)
 			{
 				LoggerWithWriter(g_consoleLogWriter, TC("")).Info(line);
 			};
 
 		ProcessHandle process = runProcess(processInfo);
+		CHECK_TRUEF(process.IsValid(), TC("Failed to run UbaTestApp"));
 		if (!process.WaitForExit(GetTimeoutTime()))
 			return logger.Error(TC("UbaTestApp did not exit in 10 seconds"));
 		if (outHandle)
 			*outHandle = process;
 		u32 exitCode = process.GetExitCode();
-		if (exitCode == 0)
+		if (exitCode == expectedExitCode)
 			return true;
 		for (auto& logLine : process.GetLogLines())
 			logger.Error(logLine.text.c_str());
@@ -244,7 +246,7 @@ namespace uba
 			});
 	}
 
-	bool RunRemote(LoggerWithWriter& logger, const StringBufferBase& testRootDir, const Config& serverConfig, const TestSessionFunction& testFunc, void* extraData = nullptr, bool deleteAll = true, bool serverShouldListen = true)
+	bool RunRemote(LoggerWithWriter& logger, const StringBufferBase& testRootDir, const Config& serverConfig, const TestSessionFunction& testFunc, uintptr_t extraData = 0, bool deleteAll = true, bool serverShouldListen = true)
 	{
 		return SetupServerClientSession(logger, testRootDir, deleteAll, serverShouldListen, serverConfig,
 			[&](LoggerWithWriter& logger, const StringView& workingDir, SessionServer& sessionServer, SessionClient& sessionClient)
@@ -255,7 +257,7 @@ namespace uba
 
 	bool RunRemote(LoggerWithWriter& logger, const StringBufferBase& testRootDir, const TestSessionFunction& testFunc, bool deleteAll, bool serverShouldListen)
 	{
-		return RunRemote(logger, testRootDir, {}, testFunc, nullptr, deleteAll, serverShouldListen);
+		return RunRemote(logger, testRootDir, {}, testFunc, 0, deleteAll, serverShouldListen);
 	}
 
 	template<typename CharType>
@@ -272,7 +274,7 @@ namespace uba
 		return fr.Close();
 	}
 
-	bool RunTestAppTests(LoggerWithWriter& logger, SessionServer& session, const tchar* workingDir, const RunProcessFunction& runProcess, void* extraData)
+	bool RunTestAppTests(LoggerWithWriter& logger, SessionServer& session, const tchar* workingDir, const RunProcessFunction& runProcess, uintptr_t extraData)
 	{
 		StringBuffer<MaxPath> fileR;
 		if (!CreateTextFile(fileR, logger, workingDir, TC("FileR.h"), "Foo"))
@@ -343,7 +345,7 @@ namespace uba
 	}
 #endif
 
-	bool RunClang(LoggerWithWriter& logger, SessionServer& session, const tchar* workingDir, const RunProcessFunction& runProcess, void* extraData)
+	bool RunClang(LoggerWithWriter& logger, SessionServer& session, const tchar* workingDir, const RunProcessFunction& runProcess, uintptr_t extraData)
 	{
 		StringBuffer<MaxPath> sourceFile;
 		sourceFile.Append(workingDir).Append(TCV("Code.cpp"));
@@ -411,7 +413,7 @@ namespace uba
 		return true;
 	}
 
-	bool RunCustomService(LoggerWithWriter& logger, SessionServer& session, const tchar* workingDir, const RunProcessFunction& runProcess, void* extraData)
+	bool RunCustomService(LoggerWithWriter& logger, SessionServer& session, const tchar* workingDir, const RunProcessFunction& runProcess, uintptr_t extraData)
 	{
 		bool gotMessage = false;
 		session.RegisterCustomService([&](uba::Process& process, const void* recv, u32 recvSize, void* send, u32 sendCapacity)
@@ -473,7 +475,7 @@ namespace uba
 
 	bool TestDetouredTouch(LoggerWithWriter& logger, const StringBufferBase& testRootDir)
 	{
-		return RunLocal(logger, testRootDir, [](LoggerWithWriter& logger, SessionServer& session, const tchar* workingDir, const RunProcessFunction& runProcess, void* extraData)
+		return RunLocal(logger, testRootDir, [](LoggerWithWriter& logger, SessionServer& session, const tchar* workingDir, const RunProcessFunction& runProcess, uintptr_t extraData)
 			{
 				StringBuffer<> file;
 				file.Append(workingDir).Append(TCV("TouchFile.h"));
@@ -483,7 +485,7 @@ namespace uba
 				CHECK_TRUE(fr.Write("Foo", 4));
 				CHECK_TRUE(fr.Close());
 				FileInformation oldInfo;
-				CHECK_TRUE(GetFileInformation(oldInfo, logger, file.data));
+				CHECK_TRUE(GetFileInformation(oldInfo, logger, file));
 
 				Sleep(100);
 
@@ -500,7 +502,7 @@ namespace uba
 					return false;
 
 				FileInformation newInfo;
-				CHECK_TRUE(GetFileInformation(newInfo, logger, file.data));
+				CHECK_TRUE(GetFileInformation(newInfo, logger, file));
 				if (newInfo.lastWriteTime == oldInfo.lastWriteTime)
 					return logger.Error(TC("File time not changed after touch"));
 				return true;
@@ -510,7 +512,7 @@ namespace uba
 	bool TestDetouredPopen(LoggerWithWriter& logger, const StringBufferBase& testRootDir)
 	{
 		#if PLATFORM_LINUX
-		return RunLocal(logger, testRootDir, [](LoggerWithWriter& logger, SessionServer& session, const tchar* workingDir, const RunProcessFunction& runProcess, void* extraData)
+		return RunLocal(logger, testRootDir, [](LoggerWithWriter& logger, SessionServer& session, const tchar* workingDir, const RunProcessFunction& runProcess, uintptr_t extraData)
 			{
 				StringBuffer<> testApp;
 				GetTestAppPath(logger, testApp);
@@ -577,7 +579,7 @@ namespace uba
 
 	bool TestMultipleDetouredProcesses(LoggerWithWriter& logger, const StringBufferBase& testRootDir)
 	{
-		return RunLocal(logger, testRootDir, [](LoggerWithWriter& logger, SessionServer& session, const tchar* workingDir, const RunProcessFunction& runProcess, void* extraData)
+		return RunLocal(logger, testRootDir, [](LoggerWithWriter& logger, SessionServer& session, const tchar* workingDir, const RunProcessFunction& runProcess, uintptr_t extraData)
 			{
 				ProcessStartInfo processInfo;
 				processInfo.application = GetSystemApplication();
@@ -602,7 +604,7 @@ namespace uba
 			});
 	}
 
-	bool RunSystemApplicationAndLookForLog(LoggerWithWriter& logger, SessionServer& session, const tchar* workingDir, const RunProcessFunction& runProcess, void* extraData)
+	bool RunSystemApplicationAndLookForLog(LoggerWithWriter& logger, SessionServer& session, const tchar* workingDir, const RunProcessFunction& runProcess, uintptr_t extraData)
 	{
 		ProcessStartInfo processInfo;
 		processInfo.application = GetSystemApplication();
@@ -635,10 +637,10 @@ namespace uba
 
 	bool TestLogLinesNoDetour(LoggerWithWriter& logger, const StringBufferBase& testRootDir)
 	{
-		return RunLocal(logger, testRootDir, RunSystemApplicationAndLookForLog, nullptr, false);
+		return RunLocal(logger, testRootDir, RunSystemApplicationAndLookForLog, 0, false);
 	}
 
-	bool CheckAttributes(LoggerWithWriter& logger, SessionServer& session, const tchar* workingDir, const RunProcessFunction& runProcess, void* extraData)
+	bool CheckAttributes(LoggerWithWriter& logger, SessionServer& session, const tchar* workingDir, const RunProcessFunction& runProcess, uintptr_t extraData)
 	{
 		StringBuffer<MaxPath> testApp;
 		GetTestAppPath(logger, testApp);
@@ -652,7 +654,7 @@ namespace uba
 
 		auto GetAttributes = [&](const StringView& file) -> u32
 			{
-				StringBuffer<> arg(TC("-GetFileAttributes="));
+				StringBuffer<> arg(TC("-getFileAttributes="));
 				arg.Append(file);
 				processInfo.arguments = arg.data;
 				ProcessHandle process = runProcess(processInfo);
@@ -704,7 +706,7 @@ namespace uba
 		newDir2.Append(workingDir).Append(TCV("NewDir2"));
 		CHECK_TRUE(CreateDirectoryW(newDir2.data));
 		CHECK_TRUE(GetAttributes(newDir2) == INVALID_FILE_ATTRIBUTES);
-		CHECK_TRUE(session.RefreshDirectory(workingDir))
+		CHECK_TRUE(session.RefreshDirectory(workingDir));
 		CHECK_TRUE(GetAttributes(newDir2) != INVALID_FILE_ATTRIBUTES);
 		dirTable.ParseDirectoryTable(session.GetDirectoryTableSize(nullptr));
 		CHECK_TRUE(dirTable.EntryExists(newDir2) == DirectoryTable::Exists_Yes);
@@ -724,7 +726,7 @@ namespace uba
 
 	bool TestSharedReservedMemory(LoggerWithWriter& logger, const StringBufferBase& testRootDir)
 	{
-		return RunLocal(logger, testRootDir, [](LoggerWithWriter& logger, SessionServer& session, const tchar* workingDir, const RunProcessFunction& runProcess, void* extraData)
+		return RunLocal(logger, testRootDir, [](LoggerWithWriter& logger, SessionServer& session, const tchar* workingDir, const RunProcessFunction& runProcess, uintptr_t extraData)
 			{
 				StringBuffer<MaxPath> testApp;
 				GetTestAppPath(logger, testApp);
@@ -783,7 +785,7 @@ namespace uba
 		bool registerToDir = true;
 	};
 
-	bool RunCreateVirtualFileTest(LoggerWithWriter& logger, SessionServer& session, const tchar* workingDir, const RunProcessFunction& runProcess, void* extraData)
+	bool RunCreateVirtualFileTest(LoggerWithWriter& logger, SessionServer& session, const tchar* workingDir, const RunProcessFunction& runProcess, uintptr_t extraData)
 	{
 		auto testData = *(RunCreateVirtualFileTestData*)extraData;
 
@@ -855,15 +857,15 @@ namespace uba
 		config.AddTable(TC("Session")).AddValue(TC("ShouldWriteToDisk"), false);
 
 		RunCreateVirtualFileTestData testData { false };
-		if (!RunLocal(logger, testRootDir, config, RunCreateVirtualFileTest, &testData))
+		if (!RunLocal(logger, testRootDir, config, RunCreateVirtualFileTest, uintptr_t(&testData)))
 			return false;
 
 		testData.transient = true;
-		if (!RunLocal(logger, testRootDir, config, RunCreateVirtualFileTest, &testData))
+		if (!RunLocal(logger, testRootDir, config, RunCreateVirtualFileTest, uintptr_t(&testData)))
 			return false;
 
 		testData.registerToDir = false;
-		if (!RunLocal(logger, testRootDir, config, RunCreateVirtualFileTest, &testData))
+		if (!RunLocal(logger, testRootDir, config, RunCreateVirtualFileTest, uintptr_t(&testData)))
 			return false;
 		return true;
 	}
@@ -877,20 +879,20 @@ namespace uba
 		RunCreateVirtualFileTestData testData { true };
 
 		testData.transient = true;
-		if (!RunRemote(logger, testRootDir, config, RunCreateVirtualFileTest, &testData))
+		if (!RunRemote(logger, testRootDir, config, RunCreateVirtualFileTest, uintptr_t(&testData)))
 			return false;
 			
 		testData.transient = false;
-		if (!RunRemote(logger, testRootDir, config, RunCreateVirtualFileTest, &testData))
+		if (!RunRemote(logger, testRootDir, config, RunCreateVirtualFileTest, uintptr_t(&testData)))
 			return false;
 
 		testData.registerToDir = false;
-		if (!RunRemote(logger, testRootDir, config, RunCreateVirtualFileTest, &testData))
+		if (!RunRemote(logger, testRootDir, config, RunCreateVirtualFileTest, uintptr_t(&testData)))
 			return false;
 		return true;
 	}
 
-	bool RunRegisterVirtualFileTest(LoggerWithWriter& logger, SessionServer& session, const tchar* workingDir, const RunProcessFunction& runProcess, void* extraData)
+	bool RunRegisterVirtualFileTest(LoggerWithWriter& logger, SessionServer& session, const tchar* workingDir, const RunProcessFunction& runProcess, uintptr_t extraData)
 	{
 		StringBuffer<> bigFile;
 		if (!CreateTestFile(bigFile, logger, ToView(workingDir), TCV("BigFile.txt"), TCV("FooBar")))
@@ -921,7 +923,7 @@ namespace uba
 	}
 
 #if PLATFORM_MAC
-	bool RunXCodeSelect(LoggerWithWriter& logger, SessionServer& session, const tchar* workingDir, const RunProcessFunction& runProcess, void* extraData)
+	bool RunXCodeSelect(LoggerWithWriter& logger, SessionServer& session, const tchar* workingDir, const RunProcessFunction& runProcess, uintptr_t extraData)
 	{
 		#if UBA_DEBUG // Failing on farm for some reason... need to revisit
 		StringBuffer<> xcodeSelect;
@@ -984,42 +986,30 @@ namespace uba
 		return SetupServerSession(logger, testRootDir, true, true, config,
 			[&](LoggerWithWriter& logger, const StringView& workingDir, SessionServer& sessionServer)
 			{
-				if (!CreateTestFile(logger, workingDir, TCV("File.h"), TCV("0")))
-					return false;
-
+				CHECK_TRUE(CreateTestFile(logger, workingDir, TCV("File.h"), TCV("0")));
 				ProcessStartInfo processInfo;
 				processInfo.application = TC("cmd.exe");
 				processInfo.workingDir = workingDir.data;
 				processInfo.arguments = TC("/c copy /Y \"File.h\" \"File2.h\"");
 				ProcessHandle process = sessionServer.RunProcess(processInfo);
-				if (!process.IsValid())
-					return logger.Error(TC("Failed to start process"));
-				if (!process.WaitForExit(GetTimeoutTime()))
-					return logger.Error(TC("UbaTestApp did not exit in 10 seconds"));
-				if (!Equals(process.GetStartInfo().application, TC("ubacopy")))
-					return logger.Error(TC("Special copy was not used"));
+				CHECK_TRUEF(process.IsValid(), TC("Failed to start process"));
+				CHECK_TRUEF(process.WaitForExit(GetTimeoutTime()), TC("UbaTestApp did not exit in 10 seconds"));
+				CHECK_TRUEF(Equals(process.GetStartInfo().application, TC("ubacopy")), TC("Special copy was not used"));
 				u32 exitCode = process.GetExitCode();
-				if (exitCode != 0)
-					return logger.Error(TC("Special copy failed"));
-				if (!FileExists(logger, workingDir, TCV("File2.h")))
-					return false;
+				CHECK_TRUEF(exitCode == 0, TC("Special copy failed"));
+				CHECK_TRUE(FileExists(logger, workingDir, TCV("File2.h")));
 				return true;
 			});
 	}
 
-	bool WriteFileMap(LoggerWithWriter& logger, SessionServer& session, const tchar* workingDir, const RunProcessFunction& runProcess, void* extraData)
+	bool WriteFileMap(LoggerWithWriter& logger, SessionServer& session, const tchar* workingDir, const RunProcessFunction& runProcess, uintptr_t extraData)
 	{
 		StringBuffer<MaxPath> fileR;
-		if (!CreateTextFile(fileR, logger, workingDir, TC("File.h"), TC("Foo")))
-			return false;
-		if (!RunTestApp(logger, workingDir, runProcess, TC("-readFileMap=Foo")))
-			return false;
-		if (!RunTestApp(logger, workingDir, runProcess, TC("-writeFileMap=Foo")))
-			return false;
-		if (!RunTestApp(logger, workingDir, runProcess, TC("-writeFileMap=Bar")))
-			return false;
-		if (!RunTestApp(logger, workingDir, runProcess, TC("-readFileMap=Bar")))
-			return false;
+		CHECK_TRUE(CreateTextFile(fileR, logger, workingDir, TC("File.h"), TC("Foo")));
+		CHECK_TRUE(RunTestApp(logger, workingDir, runProcess, TC("-readFileMap=Foo")));
+		CHECK_TRUE(RunTestApp(logger, workingDir, runProcess, TC("-writeFileMap=Foo")));
+		CHECK_TRUE(RunTestApp(logger, workingDir, runProcess, TC("-writeFileMap=Bar")));
+		CHECK_TRUE(RunTestApp(logger, workingDir, runProcess, TC("-readFileMap=Bar")));
 		return true;
 	}
 
@@ -1032,4 +1022,43 @@ namespace uba
 	{
 		return RunRemote(logger, testRootDir, WriteFileMap);
 	}
+
+	bool RunOverlayAccessOrdering(LoggerWithWriter& logger, SessionServer& session, const tchar* workingDir, const RunProcessFunction& runProcess, uintptr_t extraData)
+	{
+		StringBuffer<> barFile;
+		barFile.Append(workingDir).Append(TCV("bar.txt"));
+		CHECK_TRUE(session.CreateVirtualFile(barFile.data, "", 1, true, false));
+
+		StringBuffer<> fooFile;
+		CHECK_TRUE(CreateTestFile(fooFile, logger, ToView(workingDir), TCV("foo.txt"), TCV("foo")));
+
+		if (extraData == 1)
+			CHECK_TRUE(RunTestApp(logger, workingDir, runProcess, TC("-getFileAttributes=foo.txt"), nullptr, TC(""), 32));
+		else if (extraData == 2)
+		{
+			CHECK_TRUE(RunTestApp(logger, workingDir, runProcess, TC("-getFileAttributes=bar.txt"), nullptr, barFile.data, 128));
+			CHECK_TRUE(RunTestApp(logger, workingDir, runProcess, TC("-getFileAttributes=foo.txt"), nullptr, barFile.data, 32));
+		}
+		else if (extraData == 3)
+			CHECK_TRUE(RunTestApp(logger, workingDir, runProcess, TC("-createFileEmptyFile=bar.txt -getFileAttributes=foo.txt"), nullptr, TC(""), 32));
+
+		return true;
+	}
+
+	bool TestOverlayAccessOrdering(LoggerWithWriter& logger, const StringBufferBase& testRootDir)
+	{
+		CHECK_TRUE(RunLocal(logger, testRootDir, RunOverlayAccessOrdering, 1));
+		CHECK_TRUE(RunLocal(logger, testRootDir, RunOverlayAccessOrdering, 2));
+		CHECK_TRUE(RunLocal(logger, testRootDir, RunOverlayAccessOrdering, 3));
+		return true;
+	}
+
+	bool TestOverlayAccessOrderingRemote(LoggerWithWriter& logger, const StringBufferBase& testRootDir)
+	{
+		Config c;
+		CHECK_TRUE(RunRemote(logger, testRootDir, c, RunOverlayAccessOrdering, 1));
+		CHECK_TRUE(RunRemote(logger, testRootDir, c, RunOverlayAccessOrdering, 2));
+		CHECK_TRUE(RunRemote(logger, testRootDir, c, RunOverlayAccessOrdering, 3));
+		return true;
+	}
 }
diff --git a/Engine/Source/Programs/UnrealBuildAccelerator/TestApp/Private/PosixOS/UbaTestAppPosix.cpp b/Engine/Source/Programs/UnrealBuildAccelerator/TestApp/Private/PosixOS/UbaTestAppPosix.cpp
--- a/Engine/Source/Programs/UnrealBuildAccelerator/TestApp/Private/PosixOS/UbaTestAppPosix.cpp
+++ b/Engine/Source/Programs/UnrealBuildAccelerator/TestApp/Private/PosixOS/UbaTestAppPosix.cpp
@@ -123,7 +123,7 @@ int main(int argc, char* argv[])
 		if (stat("FileW", &attrW2) == -1)
 			return LogError("stat for FileW failed");
 		if (S_ISREG(attrW2.st_mode) == 0)
-			return LogError("stat for FileW did not return normal file (%u)", attrW2.st_mode);
+			return LogError("stat for FileW did not return normal file (0x%x)", attrW2.st_mode);
 
 		if (rename("FileW", "FileW2") == -1)
 			return LogError("rename for FileW to FileW2 failed");
@@ -276,7 +276,7 @@ int main(int argc, char* argv[])
 			return LogError("stat for 'FileW' in child process failed");
 		return 0;
 	}
-	else if (strncmp(argv[1], "-GetFileAttributes=", 6) == 0)
+	else if (strncmp(argv[1], "-getFileAttributes=", 6) == 0)
 	{
 		const char* str = argv[1] + 19;
 		struct stat attr;
diff --git a/Engine/Source/Programs/UnrealBuildAccelerator/TestApp/Private/Windows/UbaTestAppWin.cpp b/Engine/Source/Programs/UnrealBuildAccelerator/TestApp/Private/Windows/UbaTestAppWin.cpp
--- a/Engine/Source/Programs/UnrealBuildAccelerator/TestApp/Private/Windows/UbaTestAppWin.cpp
+++ b/Engine/Source/Programs/UnrealBuildAccelerator/TestApp/Private/Windows/UbaTestAppWin.cpp
@@ -252,254 +252,265 @@ int wmain(int argc, wchar_t* argv[])
 				return LogError(L"Child process failed");
 			CloseHandle(pi.hProcess);
 		}
-
-	}
-	else if (wcscmp(argv[1], L"-child") == 0)
-	{
-		if (GetFileAttributes(L"FileW2") == INVALID_FILE_ATTRIBUTES)
-			return LogError(L"Child process could not get attributes of FileW2");
-		if (GetFileAttributes(L"FileWF") == INVALID_FILE_ATTRIBUTES)
-			return LogError(L"Child process could not get attributes of FileWF");
-		if (GetFileAttributes(L"FileW") != INVALID_FILE_ATTRIBUTES)
-			return LogError(L"Child process found FileW which should not exist anymore");
+		return 0;
 	}
-	else if (wcscmp(argv[1], L"-reuse") == 0)
+
+	for (int i = 1; i != argc; ++i)
 	{
-		if (detoursHandle)
-			return LogError(L"Did not find UbaDetours.dll in process!!!\n");
+		wchar_t* arg = argv[i];
+
+		if (wcscmp(arg, L"-child") == 0)
+		{
+			if (GetFileAttributes(L"FileW2") == INVALID_FILE_ATTRIBUTES)
+				return LogError(L"Child process could not get attributes of FileW2");
+			if (GetFileAttributes(L"FileWF") == INVALID_FILE_ATTRIBUTES)
+				return LogError(L"Child process could not get attributes of FileWF");
+			if (GetFileAttributes(L"FileW") != INVALID_FILE_ATTRIBUTES)
+				return LogError(L"Child process found FileW which should not exist anymore");
+		}
+		else if (wcscmp(arg, L"-reuse") == 0)
+		{
+			if (detoursHandle)
+				return LogError(L"Did not find UbaDetours.dll in process!!!\n");
 
-		wchar_t arguments[1024];
-		if (requestNextProcess(0, arguments, sizeof(arguments)))
-			return LogError(L"Didn't expect another process");
-		while (true)
+			wchar_t arguments[1024];
 			if (requestNextProcess(0, arguments, sizeof(arguments)))
-				break;
-	}
-	else if (wcsncmp(argv[1], L"-file=", 6) == 0)
-	{
-		wchar_t arguments[1024];
-		const wchar_t* file = argv[1] + 6;
-		while (true)
+				return LogError(L"Didn't expect another process");
+			while (true)
+				if (requestNextProcess(0, arguments, sizeof(arguments)))
+					break;
+		}
+		else if (wcsncmp(arg, L"-file=", 6) == 0)
+		{
+			wchar_t arguments[1024];
+			const wchar_t* file = arg + 6;
+			while (true)
+			{
+				HANDLE rh = CreateFileW(file, GENERIC_READ, FILE_SHARE_READ, NULL, OPEN_EXISTING, 0, NULL);
+				if (rh == INVALID_HANDLE_VALUE)
+					return LogError(L"Failed to open file %s", file);
+				char data[17] = {};
+				DWORD bytesRead;
+				if (!ReadFile(rh, data, 16, &bytesRead, NULL) || bytesRead != 16)
+					return LogError(L"Failed to read 16 bytes from file %s", file);
+				CloseHandle(rh);
+
+				srand(GetProcessId(GetCurrentProcess()));
+				Sleep(rand() % 2000);
+				wchar_t outFile[1024];
+				wcscpy_s(outFile, 1024, file);
+				outFile[wcslen(file)-3] = 0;
+				wcscat_s(outFile, 1024, L".out");
+
+				HANDLE wh = CreateFileW(outFile, GENERIC_WRITE, FILE_SHARE_WRITE, NULL, CREATE_ALWAYS, 0, NULL);
+				if (wh == INVALID_HANDLE_VALUE)
+					return LogError(L"Failed to create file File");
+				data[16] = 1;
+				DWORD bytesWritten;
+				if (!WriteFile(wh, data, 17, &bytesWritten, NULL) || bytesWritten != 17)
+					return LogError(L"Failed to read 16 bytes from file %s", file);
+
+				CloseHandle(wh);
+
+				// Request new process
+				if (!requestNextProcess(0, arguments, 1024))
+					break; // No process available, exit loop
+				file = arguments + 6;
+			}
+
+			return 0;
+		}
+		else if (wcsncmp(arg, L"-getFileAttributes=", 6) == 0)
+		{
+			const wchar_t* str = arg + 19;
+			DWORD attr = GetFileAttributesW(str);
+			return attr == INVALID_FILE_ATTRIBUTES ? 255 : attr;
+		}
+		else if (wcsncmp(arg, L"-stdout=", 8) == 0)
+		{
+			const wchar_t* str = arg + 8;
+			if (wcscmp(str, L"rootprocess") == 0)
+			{
+				STARTUPINFOW si;
+				memset(&si, 0, sizeof(si));
+				PROCESS_INFORMATION pi;
+				memset(&pi, 0, sizeof(pi));
+				wchar_t childArg[1024];
+				wcscpy_s(childArg, 1024, argv[0]);
+				wcscat_s(childArg, 1024, L" -stdout=childprocess");
+				//wcscpy_s(childArg, 1024, L"\"c:\\sdk\\AutoSDK/HostWin64/Win64/MetalDeveloperTools/4.1/metal/macos/bin/metal.exe\" -v --target=air64-apple-darwin18.7.0 16384");
+
+				SECURITY_ATTRIBUTES saAttr;
+				saAttr.nLength = sizeof(SECURITY_ATTRIBUTES);
+				saAttr.bInheritHandle = TRUE;
+				saAttr.lpSecurityDescriptor = NULL;
+				HANDLE readPipe;
+				HANDLE writePipe;
+				if (!CreatePipe(&readPipe, &writePipe, &saAttr, 0))
+					return 1;
+
+				if (!SetHandleInformation(readPipe, HANDLE_FLAG_INHERIT, 0))
+					return 2;
+
+				si.dwFlags = STARTF_USESTDHANDLES;
+				si.hStdInput = GetStdHandle(STD_INPUT_HANDLE);
+				si.hStdOutput = writePipe;
+				si.hStdError = writePipe;
+
+				DWORD flags = 0;//CREATE_NO_WINDOW;
+				if (!CreateProcessW(nullptr, childArg, nullptr, nullptr, TRUE, flags, nullptr, nullptr, &si, &pi))
+					return 3;
+				CloseHandle(pi.hThread);
+				CloseHandle(writePipe);
+
+				char buf[4096] = { 0 };
+				DWORD readCount = 0;
+				if (!::ReadFile(readPipe, buf, sizeof(buf) - 1, &readCount, NULL))
+				{
+					LogError(L"Failed to read pipe %u %u", GetLastError(), readCount);
+					return 4;
+				}
+				buf[readCount] = 0;
+				if (strncmp(buf, "childprocess", 12) != 0)
+					return 5;
+
+				if (WaitForSingleObject(pi.hProcess, INFINITE) != WAIT_OBJECT_0)
+					return 6;
+				CloseHandle(pi.hProcess);
+			}
+			wprintf(L"%s\n", str);
+		}
+		else if (wcsncmp(arg, L"-virtualFile", 12) == 0)
 		{
+			const wchar_t* file = L"VirtualFile.in";
 			HANDLE rh = CreateFileW(file, GENERIC_READ, FILE_SHARE_READ, NULL, OPEN_EXISTING, 0, NULL);
 			if (rh == INVALID_HANDLE_VALUE)
 				return LogError(L"Failed to open file %s", file);
-			char data[17] = {};
+			char data[4] = {};
 			DWORD bytesRead;
-			if (!ReadFile(rh, data, 16, &bytesRead, NULL) || bytesRead != 16)
-				return LogError(L"Failed to read 16 bytes from file %s", file);
+			if (!ReadFile(rh, data, 4, &bytesRead, NULL) || bytesRead != 3)
+				return LogError(L"Failed to read 3 bytes from file %s", file);
 			CloseHandle(rh);
+			if (memcmp(data, "FOO", 3) != 0)
+				return LogError(L"File %s has wrong content", file);
 
-			srand(GetProcessId(GetCurrentProcess()));
-			Sleep(rand() % 2000);
-			wchar_t outFile[1024];
-			wcscpy_s(outFile, 1024, file);
-			outFile[wcslen(file)-3] = 0;
-			wcscat_s(outFile, 1024, L".out");
-
-			HANDLE wh = CreateFileW(outFile, GENERIC_WRITE, FILE_SHARE_WRITE, NULL, CREATE_ALWAYS, 0, NULL);
+			file = L"VirtualFile.out";
+			HANDLE wh = CreateFileW(file, GENERIC_WRITE, FILE_SHARE_WRITE, NULL, CREATE_ALWAYS, 0, NULL);
 			if (wh == INVALID_HANDLE_VALUE)
-				return LogError(L"Failed to create file File");
-			data[16] = 1;
-			DWORD bytesWritten;
-			if (!WriteFile(wh, data, 17, &bytesWritten, NULL) || bytesWritten != 17)
-				return LogError(L"Failed to read 16 bytes from file %s", file);
-
+				return LogError(L"Failed to open file %s for write", file);
+			if (!WriteFile(wh, "BAR", 3, NULL, NULL))
+				return LogError(L"Failed to write 3 bytes to file %s", file);
 			CloseHandle(wh);
-
-			// Request new process
-			if (!requestNextProcess(0, arguments, 1024))
-				break; // No process available, exit loop
-			file = arguments + 6;
 		}
-
-		return 0;
-	}
-	else if (wcsncmp(argv[1], L"-GetFileAttributes=", 6) == 0)
-	{
-		const wchar_t* str = argv[1] + 19;
-		DWORD attr = GetFileAttributesW(str);
-		return attr == INVALID_FILE_ATTRIBUTES ? 255 : attr;
-	}
-	else if (wcsncmp(argv[1], L"-stdout=", 8) == 0)
-	{
-		const wchar_t* str = argv[1] + 8;
-		if (wcscmp(str, L"rootprocess") == 0)
+		else if (wcsncmp(arg, L"-read=", 6) == 0)
 		{
-			STARTUPINFOW si;
-			memset(&si, 0, sizeof(si));
-			PROCESS_INFORMATION pi;
-			memset(&pi, 0, sizeof(pi));
-			wchar_t arg[1024];
-			wcscpy_s(arg, 1024, argv[0]);
-			wcscat_s(arg, 1024, L" -stdout=childprocess");
-			//wcscpy_s(arg, 1024, L"\"c:\\sdk\\AutoSDK/HostWin64/Win64/MetalDeveloperTools/4.1/metal/macos/bin/metal.exe\" -v --target=air64-apple-darwin18.7.0 16384");
-
-			SECURITY_ATTRIBUTES saAttr;
-			saAttr.nLength = sizeof(SECURITY_ATTRIBUTES);
-			saAttr.bInheritHandle = TRUE;
-			saAttr.lpSecurityDescriptor = NULL;
-			HANDLE readPipe;
-			HANDLE writePipe;
-			if (!CreatePipe(&readPipe, &writePipe, &saAttr, 0))
-				return 1;
-
-			if (!SetHandleInformation(readPipe, HANDLE_FLAG_INHERIT, 0))
-				return 2;
-
-			si.dwFlags = STARTF_USESTDHANDLES;
-			si.hStdInput = GetStdHandle(STD_INPUT_HANDLE);
-			si.hStdOutput = writePipe;
-			si.hStdError = writePipe;
-
-			DWORD flags = 0;//CREATE_NO_WINDOW;
-			if (!CreateProcessW(nullptr, arg, nullptr, nullptr, TRUE, flags, nullptr, nullptr, &si, &pi))
-				return 3;
-			CloseHandle(pi.hThread);
-			CloseHandle(writePipe);
-
-			char buf[4096] = { 0 };
-			DWORD readCount = 0;
-			if (!::ReadFile(readPipe, buf, sizeof(buf) - 1, &readCount, NULL))
-			{
-				LogError(L"Failed to read pipe %u %u", GetLastError(), readCount);
-				return 4;
-			}
-			buf[readCount] = 0;
-			if (strncmp(buf, "childprocess", 12) != 0)
-				return 5;
-
-			if (WaitForSingleObject(pi.hProcess, INFINITE) != WAIT_OBJECT_0)
-				return 6;
-			CloseHandle(pi.hProcess);
+			int callIndex = _wtoi(arg + 6);
+
+			wchar_t buf[32];
+			int read = ReadTestFile(buf, 32*2, L"SpecialFile1");
+			if (read != 2)
+				return LogError(L"Expected to read two bytes but got %i", read);
+			buf[1] = 0;
+
+			int readIndex = _wtoi(buf);
+			if (callIndex != readIndex)
+				return LogError(L"ReadTestFile callIndex (%i) different from readIndex (%i)", callIndex, readIndex);
+			return 0;
+		}
+		else if (wcsncmp(arg, L"-write=", 7) == 0)
+		{
+			int callIndex = _wtoi(arg + 7);
+			if (callIndex > 9)
+				return LogError(L"CallIndex too large (%i)", callIndex);
+			wchar_t buf[32];
+			buf[0] = wchar_t('1' + callIndex);
+			return WriteTestFile(buf, 2, L"SpecialFile1");
+		}
+		else if (wcsncmp(arg, L"-readwrite=", 11) == 0)
+		{
+			int callIndex = _wtoi(arg + 11);
+			if (callIndex > 9)
+				return LogError(L"CallIndex too large (%i)", callIndex);
+
+			wchar_t buf[32];
+			int read = ReadTestFile(buf, 32*2, L"SpecialFile1");
+			if (read != 2)
+				return LogError(L"Expected to read two bytes but got %i", read);
+			buf[1] = 0;
+
+			int readIndex = _wtoi(buf);
+			if (callIndex != readIndex)
+				return LogError(L"ReadWriteTestFile callIndex (%i) different from readIndex (%i)", callIndex, readIndex);
+
+			buf[0] = wchar_t('1' + readIndex);
+			return WriteTestFile(buf, 2, L"SpecialFile1");
+		}
+		else if (wcscmp(arg, L"-testRegisterVirtualFile") == 0)
+		{
+			wchar_t data[3];
+			if (ReadTestFile(data, 6, L"First.txt") == -1)
+				return -1;
+			if (memcmp(data, L"Foo", 6) != 0)
+				return -1;
+			if (ReadTestFile(data, 6, L"Second.txt") == -1)
+				return -1;
+			if (memcmp(data, L"Bar", 6) != 0)
+				return -1;
+		}
+		else if (wcsncmp(arg, L"-writeFileMap=", 14) == 0)
+		{
+			wchar_t* toWrite = arg + 14;
+			if (WriteTestFile(toWrite, (wcslen(toWrite) + 1)*sizeof(wchar_t), L"File.h") == -1)
+				return -1;
+		}
+		else if (wcsncmp(arg, L"-readFileMap=", 13) == 0)
+		{
+			wchar_t* toExpect = arg + 13;
+			wchar_t buf[32];
+			if (ReadTestFile(buf, sizeof(buf), L"File.h") == -1)
+				return -1;
+			if (wcscmp(toExpect, buf) != 0)
+				return LogError(L"-readFileMap did not read what was expected");
+		}
+		else if (wcscmp(arg, L"-readWriteFileMap") == 0)
+		{
+			wchar_t buf[32];
+			if (ReadTestFile(buf, 1, L"File.h") == -1)
+				return -1;
+			buf[0] = 'M';
+			if (WriteTestFile(buf, 1, L"File.h") == -1)
+				return -1;
+			buf[0] = 0;
+			if (ReadTestFile(buf, 1, L"File.h") == -1)
+				return -1;
+			if (buf[0] != 'M')
+				return -1;
+		}
+		else if (wcsncmp(arg, L"-createFileEmptyFile=", 21) == 0)
+		{
+			wchar_t* fileName = arg + 22;
+			if (WriteTestFile("", 0, fileName) == -1)
+				return -1;
+		}
+		else
+		{
+			if (!detoursHandle)
+				return LogError(L"Did not find UbaDetours.dll in process!!!\n");
+
+			using u32 = unsigned int;
+			using UbaSendCustomMessageFunc = u32(const void* send, u32 sendSize, void* recv, u32 recvCapacity);
+
+			UbaSendCustomMessageFunc* sendMessage = (UbaSendCustomMessageFunc*)GetProcAddress(detoursHandle, "UbaSendCustomMessage");
+			if (!sendMessage)
+				return LogError(L"Couldn't find UbaSendCustomMessage function in UbaDetours.dll");
+
+			const wchar_t* helloMsg = L"Hello from client";
+			wchar_t response[256];
+			u32 responseSize = (*sendMessage)(helloMsg, u32(wcslen(helloMsg)) * 2, response, 256 * 2);
+			if (responseSize == 0)
+				return LogError(L"Didn't get proper response from session");
+			//wprintf(L"Recv: %.*s\n", responseSize / 2, response);
 		}
-		wprintf(L"%s\n", str);
-	}
-	else if (wcsncmp(argv[1], L"-virtualFile", 12) == 0)
-	{
-		const wchar_t* file = L"VirtualFile.in";
-		HANDLE rh = CreateFileW(file, GENERIC_READ, FILE_SHARE_READ, NULL, OPEN_EXISTING, 0, NULL);
-		if (rh == INVALID_HANDLE_VALUE)
-			return LogError(L"Failed to open file %s", file);
-		char data[4] = {};
-		DWORD bytesRead;
-		if (!ReadFile(rh, data, 4, &bytesRead, NULL) || bytesRead != 3)
-			return LogError(L"Failed to read 3 bytes from file %s", file);
-		CloseHandle(rh);
-		if (memcmp(data, "FOO", 3) != 0)
-			return LogError(L"File %s has wrong content", file);
-
-		file = L"VirtualFile.out";
-		HANDLE wh = CreateFileW(file, GENERIC_WRITE, FILE_SHARE_WRITE, NULL, CREATE_ALWAYS, 0, NULL);
-		if (wh == INVALID_HANDLE_VALUE)
-			return LogError(L"Failed to open file %s for write", file);
-		if (!WriteFile(wh, "BAR", 3, NULL, NULL))
-			return LogError(L"Failed to write 3 bytes to file %s", file);
-		CloseHandle(wh);
-	}
-	else if (wcsncmp(argv[1], L"-read=", 6) == 0)
-	{
-		int callIndex = _wtoi(argv[1] + 6);
-
-		wchar_t buf[32];
-		int read = ReadTestFile(buf, 32*2, L"SpecialFile1");
-		if (read != 2)
-			return LogError(L"Expected to read two bytes but got %i", read);
-		buf[1] = 0;
-
-		int readIndex = _wtoi(buf);
-		if (callIndex != readIndex)
-			return LogError(L"ReadTestFile callIndex (%i) different from readIndex (%i)", callIndex, readIndex);
-		return 0;
-	}
-	else if (wcsncmp(argv[1], L"-write=", 7) == 0)
-	{
-		int callIndex = _wtoi(argv[1] + 7);
-		if (callIndex > 9)
-			return LogError(L"CallIndex too large (%i)", callIndex);
-		wchar_t buf[32];
-		buf[0] = wchar_t('1' + callIndex);
-		return WriteTestFile(buf, 2, L"SpecialFile1");
-	}
-	else if (wcsncmp(argv[1], L"-readwrite=", 11) == 0)
-	{
-		int callIndex = _wtoi(argv[1] + 11);
-		if (callIndex > 9)
-			return LogError(L"CallIndex too large (%i)", callIndex);
-
-		wchar_t buf[32];
-		int read = ReadTestFile(buf, 32*2, L"SpecialFile1");
-		if (read != 2)
-			return LogError(L"Expected to read two bytes but got %i", read);
-		buf[1] = 0;
-
-		int readIndex = _wtoi(buf);
-		if (callIndex != readIndex)
-			return LogError(L"ReadWriteTestFile callIndex (%i) different from readIndex (%i)", callIndex, readIndex);
-
-		buf[0] = wchar_t('1' + readIndex);
-		return WriteTestFile(buf, 2, L"SpecialFile1");
-	}
-	else if (wcscmp(argv[1], L"-testRegisterVirtualFile") == 0)
-	{
-		wchar_t data[3];
-		if (ReadTestFile(data, 6, L"First.txt") == -1)
-			return -1;
-		if (memcmp(data, L"Foo", 6) != 0)
-			return -1;
-		if (ReadTestFile(data, 6, L"Second.txt") == -1)
-			return -1;
-		if (memcmp(data, L"Bar", 6) != 0)
-			return -1;
-	}
-	else if (wcsncmp(argv[1], L"-writeFileMap=", 14) == 0)
-	{
-		wchar_t* toWrite = argv[1] + 14;
-		if (WriteTestFile(toWrite, (wcslen(toWrite) + 1)*sizeof(wchar_t), L"File.h") == -1)
-			return -1;
-	}
-	else if (wcsncmp(argv[1], L"-readFileMap=", 13) == 0)
-	{
-		wchar_t* toExpect = argv[1] + 13;
-		wchar_t buf[32];
-		if (ReadTestFile(buf, sizeof(buf), L"File.h") == -1)
-			return -1;
-		if (wcscmp(toExpect, buf) != 0)
-			return LogError(L"-readFileMap did not read what was expected");
-	}
-	else if (wcscmp(argv[1], L"-readWriteFileMap") == 0)
-	{
-		wchar_t buf[32];
-		if (ReadTestFile(buf, 1, L"File.h") == -1)
-			return -1;
-		buf[0] = 'M';
-		if (WriteTestFile(buf, 1, L"File.h") == -1)
-			return -1;
-		buf[0] = 0;
-		if (ReadTestFile(buf, 1, L"File.h") == -1)
-			return -1;
-		if (buf[0] != 'M')
-			return -1;
-	}
-	else
-	{
-		if (!detoursHandle)
-			return LogError(L"Did not find UbaDetours.dll in process!!!\n");
-
-		using u32 = unsigned int;
-		using UbaSendCustomMessageFunc = u32(const void* send, u32 sendSize, void* recv, u32 recvCapacity);
-
-		UbaSendCustomMessageFunc* sendMessage = (UbaSendCustomMessageFunc*)GetProcAddress(detoursHandle, "UbaSendCustomMessage");
-		if (!sendMessage)
-			return LogError(L"Couldn't find UbaSendCustomMessage function in UbaDetours.dll");
-
-		const wchar_t* helloMsg = L"Hello from client";
-		wchar_t response[256];
-		u32 responseSize = (*sendMessage)(helloMsg, u32(wcslen(helloMsg)) * 2, response, 256 * 2);
-		if (responseSize == 0)
-			return LogError(L"Didn't get proper response from session");
-		//wprintf(L"Recv: %.*s\n", responseSize / 2, response);
 	}
-
 	return 0;
 }
