[UBA]
* Improved compress test logic in ubacli
* Fixed so coordinator is destroyed before scheduler since it can depend on it
* Changed some default settings to match better how uba runs in ue

diff --git a/Engine/Source/Programs/UnrealBuildAccelerator/Cli/Private/UbaCli.cpp b/Engine/Source/Programs/UnrealBuildAccelerator/Cli/Private/UbaCli.cpp
--- a/Engine/Source/Programs/UnrealBuildAccelerator/Cli/Private/UbaCli.cpp
+++ b/Engine/Source/Programs/UnrealBuildAccelerator/Cli/Private/UbaCli.cpp
@@ -998,19 +998,25 @@ namespace uba
 			void* decoderMem = malloc(decoredMemSize);
 			auto mg = MakeGuard([decoderMem]() { free(decoderMem); });
 
+			u64 decompressedLeft = decompressedSize;
 			u32 blockCounter = 0;
 			while (reader.GetLeft())
 			{
 				u32 compressedBlockSize = reader.ReadU32();
 				u32 decompressedBlockSize = reader.ReadU32();
 
-				logger.Info(TC("Decompressing block %u. %u bytes to %u bytes"), blockCounter++, compressedBlockSize, decompressedBlockSize);
+				if (decompressedLeft < decompressedBlockSize)
+					return logger.Error(TC("File %s is corrupt. Destination size is %llu but block %u overflows (%llu)"), fileToDecompress.c_str(), decompressedSize, blockCounter, decompressedBlockSize);
+
+				logger.Info(TC("Decompressing block %u. %u bytes to %u bytes"), blockCounter, compressedBlockSize, decompressedBlockSize);
 
 				OO_SINTa decompLen = OodleLZ_Decompress(reader.GetPositionData(), (OO_SINTa)compressedBlockSize, destMem, (OO_SINTa)decompressedBlockSize, OodleLZ_FuzzSafe_Yes, OodleLZ_CheckCRC_No, OodleLZ_Verbosity_None, NULL, 0, NULL, NULL, decoderMem, decoredMemSize);
 				if (decompLen != decompressedBlockSize)
 					return logger.Error(TC("Failed to decompress %s (CompressedSize: %llu DecompressedSize: %llu ReadPos: %llu CompressedBlock: %u DecompressedBlock: %u)"), fileToDecompress.c_str(), fileSize, decompressedSize, reader.GetPosition(), compressedBlockSize, decompressedBlockSize);
+				decompressedLeft -= decompressedBlockSize;
 				destMem += decompressedBlockSize;
 				reader.Skip(compressedBlockSize);
+				blockCounter++;
 			}
 
 			if (!faDest.Close())
@@ -1101,7 +1107,7 @@ namespace uba
 		if (!symbolFile.empty())
 		{
 			FileInformation info;
-			GetFileInformation(info, logger, symbolFile.c_str());
+			GetFileInformation(info, logger, symbolFile);
 			u64 fakeStartAddress = 0x800000;
 
 			StackBinaryWriter<128> writer;
@@ -1290,6 +1296,7 @@ namespace uba
 		RootsHandle rootsHandle = 0;
 		CacheClient* cacheClient = nullptr;
 		Function<bool(Scheduler*)> createCoordinator;
+		Function<void()> destroyCoordinator;
 
 		bool isRemote = commandType == CommandType_Remote || commandType == CommandType_Agent;
 
@@ -1336,6 +1343,9 @@ namespace uba
 			info.readIntermediateFilesCompressed = true;
 			//info.treatTempDirAsEmpty = false;
 			//info.extractObjFilesSymbols = true;
+			info.writeFilesBottleneck = 24;
+			info.writeFilesFileMapMaxMb = 1;
+
 			#if UBA_DEBUG_LOG_ENABLED
 			info.remoteLogEnabled = true;
 			#endif
@@ -1462,6 +1472,7 @@ namespace uba
 						return false;
 					return true;
 				};
+			destroyCoordinator = [&]() { coordinator.Destroy(); };
 
 			auto stopServer = MakeGuard([&]() { networkServer.DisconnectClients(); });
 			auto stopListen = MakeGuard([&]() { networkBackend.StopListen(); });
@@ -1563,6 +1574,7 @@ namespace uba
 				{
 					u32 maxProcessor = Min(maxProcessCount/clientCount, 32u);
 					ClientInitInfo cii { logWriter, networkBackend, g_rootDir.data, TC("127.0.0.1"), port, clientZone, maxProcessor, i};
+					cii.useDependencyCrawler = true;
 					if (!clients.emplace_back().Init(cii))
 						return false;
 				}
@@ -1633,13 +1645,14 @@ namespace uba
 			info.cacheClients = cacheClients;
 			info.cacheClientCount = cacheClient ? 1 : 0;
 			info.writeToCache = populateCache;
-			//info.writeTrackedInputs = true;
+			info.writeTrackedInputs = true;
 			Scheduler scheduler(info);
 
 			//agentCount = 2;
 
 			if (!createCoordinator(&scheduler))
 				return false;
+			auto coordinatorGuard = MakeGuard([&]() { destroyCoordinator(); });
 
 			if (StringView(app).EndsWith(TCV(".yaml")))
 			{
