[UBA]
* Changed most test code to use CHECK_TRUE/CHECK_TRUEF to get better errors when failing

diff --git a/Engine/Source/Programs/UnrealBuildAccelerator/Test/Private/UbaTest.cpp b/Engine/Source/Programs/UnrealBuildAccelerator/Test/Private/UbaTest.cpp
--- a/Engine/Source/Programs/UnrealBuildAccelerator/Test/Private/UbaTest.cpp
+++ b/Engine/Source/Programs/UnrealBuildAccelerator/Test/Private/UbaTest.cpp
@@ -24,8 +24,7 @@ namespace uba
 
 		u64 bytes = content.count*sizeof(tchar);
 		FileAccessor file(logger, outFile.data);
-		if (!file.CreateMemoryWrite(false, attributes, bytes))
-			return logger.Error(TC("Failed to create file for write"));
+		CHECK_TRUEF(file.CreateMemoryWrite(false, attributes, bytes), TC("Failed to create file for write"));
 		memcpy(file.GetData(), content.data, bytes);
 		return file.Close();
 	}
@@ -34,8 +33,7 @@ namespace uba
 	{
 		StringBuffer<> testFileName;
 		testFileName.Append(testRootDir).EnsureEndsWithSlash().Append(fileName).FixPathSeparators();
-		if (!DeleteFileW(testFileName.data))
-			return logger.Error(TC("Failed to delete test file %s"), fileName.data);
+		CHECK_TRUEF(DeleteFileW(testFileName.data), TC("Failed to delete test file %s"), fileName.data);
 		return true;
 	}
 
@@ -49,13 +47,10 @@ namespace uba
 	bool FilesEqual(Logger& logger, const tchar* file1, const tchar* file2)
 	{
 		FileAccessor f1(logger, file1);
-		if (!f1.OpenMemoryRead())
-			return false;
+		CHECK_TRUE(f1.OpenMemoryRead());
 		FileAccessor f2(logger, file2);
-		if (!f2.OpenMemoryRead())
-			return false;
-		if (f1.GetSize() != f2.GetSize())
-			return false;
+		CHECK_TRUE(f2.OpenMemoryRead());
+		CHECK_TRUE(f1.GetSize() == f2.GetSize());
 		return memcmp(f1.GetData(), f2.GetData(), f1.GetSize()) == 0;
 	}
 
diff --git a/Engine/Source/Programs/UnrealBuildAccelerator/Test/Private/UbaTestBasics.cpp b/Engine/Source/Programs/UnrealBuildAccelerator/Test/Private/UbaTestBasics.cpp
--- a/Engine/Source/Programs/UnrealBuildAccelerator/Test/Private/UbaTestBasics.cpp
+++ b/Engine/Source/Programs/UnrealBuildAccelerator/Test/Private/UbaTestBasics.cpp
@@ -42,18 +42,15 @@ namespace uba
 
 		u64 v = 0;
 		buf.Clear().AppendBase62(v);
-		if (v != StringToValueBase62(buf.data, buf.count))
-			return false;
+		CHECK_TRUE(v == StringToValueBase62(buf.data, buf.count));
 
 		v = 5473789457;
 		buf.Clear().AppendBase62(v);
-		if (v != StringToValueBase62(buf.data, buf.count))
-			return false;
+		CHECK_TRUE(v == StringToValueBase62(buf.data, buf.count));
 
 		v = ~u64(0);
 		buf.Clear().AppendBase62(v);
-		if (v != StringToValueBase62(buf.data, buf.count))
-			return false;
+		CHECK_TRUE(v == StringToValueBase62(buf.data, buf.count));
 		return true;
 	}
 
@@ -82,8 +79,7 @@ namespace uba
 
 		u64 totalTime;
 		u64 maxTime;
-		if (!GetCpuTime(totalTime, maxTime))
-			return logger.Error(TC("GetCpuTime failed"));
+		CHECK_TRUEF(GetCpuTime(totalTime, maxTime), TC("GetCpuTime failed"));
 
 		return true;
 	}
@@ -100,29 +96,23 @@ namespace uba
 				return true;
 			});
 
-		if (ev.IsSet(1))
-			return logger.Error(TC("Event was set after 1ms timeout where it should take 500ms"));
+		CHECK_TRUEF(!ev.IsSet(1), TC("Event was set after 1ms timeout where it should take 500ms"));
 
-		if (ev.IsSet(0))
-			return logger.Error(TC("Event was set after no timeout where it should take 500ms"));
+		CHECK_TRUEF(!ev.IsSet(0), TC("Event was set after no timeout where it should take 500ms"));
 
 		//logger.Info(TC("Waiting for event"));
-		if (!ev.IsSet(2000))
-			return logger.Error(TC("Event was not set after 2000ms where it should take 500ms"));
+		CHECK_TRUEF(ev.IsSet(2000), TC("Event was not set after 2000ms where it should take 500ms"));
 		//logger.Info(TC("Event was set"));
 
-		if (t.Wait(0))
-			return logger.Error(TC("Thread wait timed out. Should already be done after 2000ms"));
+		CHECK_TRUEF(!t.Wait(0), TC("Thread wait timed out. Should already be done after 2000ms"));
 
-		if (!t.Wait(2000))
-			return logger.Error(TC("Thread wait did not timed out should be done after 2000ms"));
+		CHECK_TRUEF(t.Wait(2000), TC("Thread wait did not timed out should be done after 2000ms"));
 
 		#if 0 // Long time test... disabled by default
 		u64 time = GetTime();
 		EventType longEv(true);
 		longEv.IsSet(10 * 60 * 1000);
-		if (TimeToMs((GetTime() - time) < 9 * 60 * 1000))
-			return logger.Error(TC("Event timeout was way too fast"));
+		CHECK_TRUEF(!TimeToMs((GetTime() - time) < 9 * 60 * 1000), TC("Event timeout was way too fast"));
 		#endif
 		return true;
 	}
@@ -130,10 +120,8 @@ namespace uba
 	bool TestEvents(LoggerWithWriter& logger, const StringBufferBase& rootDir)
 	{
 		Event ev;
-		if (!ev.Create(EventResetType_Manual))
-			return logger.Error(TC("Failed to create event"));
-		if (!TestEventsImpl(logger, rootDir, ev, ev))
-			return false;
+		CHECK_TRUEF(ev.Create(EventResetType_Manual), TC("Failed to create event"));
+		CHECK_TRUE(TestEventsImpl(logger, rootDir, ev, ev));
 		return true;
 	}
 
@@ -148,10 +136,8 @@ namespace uba
 		mappingBuffer.Init(TC(""), 64*1024, 64 * 1024);
 		MappedView view = mappingBuffer.AllocAndMapView(1024, 1, TC("Foo"));
 		SharedEvent* ev = new (view.memory) SharedEvent();
-		if (!ev->Create(EventResetType_Manual))
-			return logger.Error(TC("Failed to create event"));
-		if (!TestEventsImpl(logger, rootDir, *ev, *ev))
-			return false;
+		CHECK_TRUEF(ev->Create(EventResetType_Manual), TC("Failed to create event"));
+		CHECK_TRUE(TestEventsImpl(logger, rootDir, *ev, *ev));
 		ev->~SharedEvent();
 		mappingBuffer.UnmapView(view, TC("Foo"));
 		#endif
@@ -167,8 +153,7 @@ namespace uba
 		ev.ToString(tmp);
 		SharedEvent ev2;
 		ev2.Create(tmp, false);
-		if (!TestEventsImpl(logger, rootDir, ev, ev2))
-			return false;
+		CHECK_TRUE(TestEventsImpl(logger, rootDir, ev, ev2));
 		#endif
 		return true;
 	}
@@ -182,33 +167,25 @@ namespace uba
 		auto TestPath = [&](const tchar* path) { return FixPath2(path, ~0u, workingDir, TStrlen(workingDir), buffer, sizeof_array(buffer), &lengthResult); };
 
 #if PLATFORM_WINDOWS
-		if (!FixPath2(TC("\"e:\\temp\""), ~0u, workingDir, TStrlen(workingDir), buffer, sizeof_array(buffer), &lengthResult))
-			return logger.Error(TC("FixPath2 (1) failed"));
-		if (!FixPath2(TC("\\??\\c:\\windows:::"), 14, workingDir, TStrlen(workingDir), buffer, sizeof_array(buffer), &lengthResult) || lengthResult != 10 || TStrlen(buffer) != 10)
-			return logger.Error(TC("FixPath2 (2) failed"));
-		if (!FixPath2(TC("Z:/UEVFS/Root/Engine/Binaries/Win64/UnrealEditor-InterchangeCore-Win64-Debug.pdb::"), 80, workingDir, TStrlen(workingDir), buffer, sizeof_array(buffer), &lengthResult) || lengthResult != 80 || TStrlen(buffer) != 80)
-			return logger.Error(TC("FixPath2 (2) failed"));
-		if (FixPath2(TC("'Small vector' optimization: store up to a small number of items on the stack"), ~0u, workingDir, TStrlen(workingDir), buffer, sizeof_array(buffer), &lengthResult))
-			return logger.Error(TC("FixPath2 (2) failed"));
+		CHECK_TRUEF(FixPath2(TC("\"e:\\temp\""), ~0u, workingDir, TStrlen(workingDir), buffer, sizeof_array(buffer), &lengthResult), TC("FixPath2 (1) failed"));
+		CHECK_TRUEF(!(!FixPath2(TC("\\??\\c:\\windows:::"), 14, workingDir, TStrlen(workingDir), buffer, sizeof_array(buffer), &lengthResult) || lengthResult != 10 || TStrlen(buffer) != 10), TC("FixPath2 (2) failed"));
+		CHECK_TRUEF(!(!FixPath2(TC("Z:/UEVFS/Root/Engine/Binaries/Win64/UnrealEditor-InterchangeCore-Win64-Debug.pdb::"), 80, workingDir, TStrlen(workingDir), buffer, sizeof_array(buffer), &lengthResult) || lengthResult != 80 || TStrlen(buffer) != 80), TC("FixPath2 (2) failed"));
+		CHECK_TRUEF(!FixPath2(TC("'Small vector' optimization: store up to a small number of items on the stack"), ~0u, workingDir, TStrlen(workingDir), buffer, sizeof_array(buffer), &lengthResult), TC("FixPath2 (2) failed"));
 		#if 0
 		const tchar longPath[] = TC("c:\\Program Files (x86)\\Microsoft Visual Studio\\Shared\\Entity Framework Tools\\NuGet Packages");
 		tchar shortPath[MAX_PATH];
 		DWORD len = GetShortPathNameW(longPath, shortPath, MAX_PATH);
-		if (len == 0 || len >= sizeof_array(longPath))
-			return logger.Error(TC("GetShortPathNameW failed"));
-		if (!FixPath2(shortPath, len, workingDir, TStrlen(workingDir), buffer, sizeof_array(buffer), &lengthResult) || lengthResult != sizeof_array(longPath) - 1)
-			return logger.Error(TC("FixPath2 for short path failed"));
+		CHECK_TRUEF(!(len == 0 || len >= sizeof_array(longPath)), TC("GetShortPathNameW failed"));
+		CHECK_TRUEF(!(!FixPath2(shortPath, len, workingDir, TStrlen(workingDir), buffer, sizeof_array(buffer), &lengthResult) || lengthResult != sizeof_array(longPath) - 1), TC("FixPath2 for short path failed"));
 		#endif
 
 
 #else
 
-		if (!TestPath(TC("/..")))
-			return logger.Error(TC("FixPath2 should have failed"));
+		CHECK_TRUEF(TestPath(TC("/..")), TC("FixPath2 should have failed"));
 		UBA_TEST_CHECK(Equals(buffer, TC("/")), "Should not contain ..");
 
-		if (!FixPath2(TC("/../Foo"), ~0u, workingDir, TStrlen(workingDir), buffer, sizeof_array(buffer), &lengthResult))
-			return logger.Error(TC("FixPath2 should have failed"));
+		CHECK_TRUEF(FixPath2(TC("/../Foo"), ~0u, workingDir, TStrlen(workingDir), buffer, sizeof_array(buffer), &lengthResult), TC("FixPath2 should have failed"));
 		UBA_TEST_CHECK(Equals(buffer, TC("/Foo")), "Should not contain ..");
 
 		if (!FixPath2(TC("/usr/bin//clang++"), ~0u, workingDir, TStrlen(workingDir), buffer, sizeof_array(buffer), &lengthResult))
@@ -216,16 +193,13 @@ namespace uba
 		UBA_TEST_CHECK(!Contains(buffer, TC("//")), "Should not contain //");
 #endif
 
-		if (!FixPath2(TC("../Foo"), ~0u, workingDir, TStrlen(workingDir), buffer, sizeof_array(buffer), &lengthResult))
-			return logger.Error(TC("FixPath2 (1) failed"));
+		CHECK_TRUEF(FixPath2(TC("../Foo"), ~0u, workingDir, TStrlen(workingDir), buffer, sizeof_array(buffer), &lengthResult), TC("FixPath2 (1) failed"));
 		UBA_TEST_CHECK(!Contains(buffer, TC("..")), "Should not contain ..");
 
-		if (!FixPath2(TC("@../Foo"), ~0u, workingDir, TStrlen(workingDir), buffer, sizeof_array(buffer), &lengthResult))
-			return logger.Error(TC("FixPath2 (1) failed"));
+		CHECK_TRUEF(FixPath2(TC("@../Foo"), ~0u, workingDir, TStrlen(workingDir), buffer, sizeof_array(buffer), &lengthResult), TC("FixPath2 (1) failed"));
 		UBA_TEST_CHECK(Contains(buffer, TC("..")), "Should contain ..");
 
-		if (!FixPath2(TC("..@/Foo"), ~0u, workingDir, TStrlen(workingDir), buffer, sizeof_array(buffer), &lengthResult))
-			return logger.Error(TC("FixPath2 (1) failed"));
+		CHECK_TRUEF(FixPath2(TC("..@/Foo"), ~0u, workingDir, TStrlen(workingDir), buffer, sizeof_array(buffer), &lengthResult), TC("FixPath2 (1) failed"));
 		UBA_TEST_CHECK(Contains(buffer, TC("..")), "Should contain ..");
 
 		return true;
@@ -237,23 +211,18 @@ namespace uba
 		testFileName.Append(TCV("UbaTestFile"));
 
 		FileAccessor fileHandle(logger, testFileName.data);
-		if (!fileHandle.CreateWrite())
-			return logger.Error(TC("Failed to create file for write"));
+		CHECK_TRUEF(fileHandle.CreateWrite(), TC("Failed to create file for write"));
 
 		u8 byte = 'H';
-		if (!fileHandle.Write(&byte, 1))
-			return false;
+		CHECK_TRUE(fileHandle.Write(&byte, 1));
 
-		if (!fileHandle.Close())
-			return false;
+		CHECK_TRUE(fileHandle.Close());
 
 		FileHandle fileHandle2;
-		if (!OpenFileSequentialRead(logger, testFileName.data, fileHandle2))
-			return logger.Error(TC("Failed to create file for read"));
+		CHECK_TRUEF(OpenFileSequentialRead(logger, testFileName.data, fileHandle2), TC("Failed to create file for read"));
 
 		u64 writeTime = 0;
-		if (!GetFileLastWriteTime(writeTime, fileHandle2))
-			return logger.Error(TC("Failed to get last written time"));
+		CHECK_TRUEF(GetFileLastWriteTime(writeTime, fileHandle2), TC("Failed to get last written time"));
 
 		u64 writeTime2 = 0;
 		TraverseDir(logger, rootDir, [&](const DirectoryEntry& de)
@@ -262,88 +231,65 @@ namespace uba
 					writeTime2 = de.lastWritten;
 			});
 
-		if (writeTime != writeTime2)
-			return logger.Error(TC("GetFileLastWriteTime and TraverseDir are returning different last write time for same file"));
+		CHECK_TRUEF(!(writeTime != writeTime2), TC("GetFileLastWriteTime and TraverseDir are returning different last write time for same file"));
 
 		u64 systemTime = GetSystemTimeAsFileTime();
-		if (systemTime < writeTime)
-			return logger.Error(TC("System time is lower than last written time"));
-		if (GetFileTimeAsSeconds(systemTime) - GetFileTimeAsSeconds(writeTime) > 3)
-			return logger.Error(TC("System time or last written time is wrong (system: %llu, write: %llu, diffInSec: %llu)"), systemTime, writeTime, GetFileTimeAsSeconds(systemTime) - GetFileTimeAsSeconds(writeTime));
+		CHECK_TRUEF(!(systemTime < writeTime), TC("System time is lower than last written time"));
+		CHECK_TRUEF(!(GetFileTimeAsSeconds(systemTime) - GetFileTimeAsSeconds(writeTime) > 3), TC("System time or last written time is wrong (system: %llu, write: %llu, diffInSec: %llu)"), systemTime, writeTime, GetFileTimeAsSeconds(systemTime) - GetFileTimeAsSeconds(writeTime));
 
 
 		u8 byte2 = 0;
-		if (!ReadFile(logger, testFileName.data, fileHandle2, &byte2, 1))
-			return false;
+		CHECK_TRUE(ReadFile(logger, testFileName.data, fileHandle2, &byte2, 1));
 
-		if (!CloseFile(testFileName.data, fileHandle2))
-			return false;
+		CHECK_TRUE(CloseFile(testFileName.data, fileHandle2));
 
 		FileHandle fileHandle3;
-		if (!OpenFileSequentialRead(logger, TC("NonExistingFile"), fileHandle3, false))
-			return logger.Error(TC("OpenFileSequentialRead failed with non existing file"));
-		if (fileHandle3 != InvalidFileHandle)
-			return logger.Error(TC("OpenFileSequentialRead found file that doesn't exist"));
+		CHECK_TRUEF(OpenFileSequentialRead(logger, TC("NonExistingFile"), fileHandle3, false), TC("OpenFileSequentialRead failed with non existing file"));
+		CHECK_TRUEF(!(fileHandle3 != InvalidFileHandle), TC("OpenFileSequentialRead found file that doesn't exist"));
 
-		if (RemoveDirectoryW(TC("TestDir")))
-			return logger.Error(TC("Did not fail to remove non-existing TestDir (or were things not cleaned before test)"));
-		else if (GetLastError() != ERROR_FILE_NOT_FOUND)
-			return logger.Error(TC("GetLastError did not return correct error failing to remove non-existing directory TestDir"));
+		CHECK_TRUEF(!RemoveDirectoryW(TC("TestDir")), TC("Did not fail to remove non-existing TestDir (or were things not cleaned before test)"));
+		CHECK_TRUEF(GetLastError() == ERROR_FILE_NOT_FOUND, TC("GetLastError did not return correct error failing to remove non-existing directory TestDir"));
 
-		if (!CreateDirectoryW(TC("TestDir")))
-			return logger.Error(TC("Failed to create dir"));
+		CHECK_TRUEF(CreateDirectoryW(TC("TestDir")), TC("Failed to create dir"));
 
-		if (CreateDirectoryW(TC("TestDir")))
-			return logger.Error(TC("Did not fail to create existing dir"));
+		CHECK_TRUEF(!CreateDirectoryW(TC("TestDir")), TC("Did not fail to create existing dir"));
 
 		FileHandle fileHandle4;
-		if (OpenFileSequentialRead(logger, TC("TestDir"), fileHandle4))
-			return logger.Error(TC("This should return fail"));
+		CHECK_TRUEF(!OpenFileSequentialRead(logger, TC("TestDir"), fileHandle4), TC("This should return fail"));
 
-		if (!RemoveDirectoryW(TC("TestDir")))
-			return logger.Error(TC("Fail to remove TestDir"));
+		CHECK_TRUEF(RemoveDirectoryW(TC("TestDir")), TC("Fail to remove TestDir"));
 
 		u64 size = 0;
-		if (!FileExists(logger, testFileName.data, &size) || size != 1)
-			return logger.Error(TC("UbaTestFile not found"));
+		CHECK_TRUEF(FileExists(logger, testFileName.data, &size) || size != 1, TC("UbaTestFile not found"));
 
 		StringBuffer<> testFileName2(rootDir);
 		testFileName2.Append(TCV("UbaTestFile2"));
 
 		DeleteFileW(testFileName2.data);
 
-		if (DeleteFileW(testFileName2.data))
-			return logger.Error(TC("Did not fail to delete non-existing UbaTestFile2 (or were things not cleaned before test)"));
-		else if (GetLastError() != ERROR_FILE_NOT_FOUND)
-			return logger.Error(TC("GetLastError did not return correct error failing to delete non-existing file UbaTestFile2"));
+		CHECK_TRUEF(!DeleteFileW(testFileName2.data), TC("Did not fail to delete non-existing UbaTestFile2 (or were things not cleaned before test)"));
+		CHECK_TRUEF(GetLastError() == ERROR_FILE_NOT_FOUND, TC("GetLastError did not return correct error failing to delete non-existing file UbaTestFile2"));
 
-		if (!CreateHardLinkW(testFileName2.data, testFileName.data))
-			return logger.Error(TC("Failed to create hardlink from UbaTestFile to UbaTestFile2"));
+		CHECK_TRUEF(CreateHardLinkW(testFileName2.data, testFileName.data), TC("Failed to create hardlink from UbaTestFile to UbaTestFile2"));
 
-		if (!DeleteFileW(testFileName.data))
-			return logger.Error(TC("Failed to delete UbaTestFile"));
+		CHECK_TRUEF(DeleteFileW(testFileName.data), TC("Failed to delete UbaTestFile"));
 
-		if (FileExists(logger, testFileName.data))
-			return logger.Error(TC("Found non-existing file UbaTestFile"));
+		CHECK_TRUEF(!FileExists(logger, testFileName.data), TC("Found non-existing file UbaTestFile"));
 
 		// CreateHardLinkW is a symbolic link on non-windows.. need to revisit
 		#if PLATFORM_WINDOWS
-		if (!FileExists(logger, testFileName2.data))
-			return logger.Error(TC("Failed to find file UbaTestFile2"));
+		CHECK_TRUEF(FileExists(logger, testFileName2.data), TC("Failed to find file UbaTestFile2"));
 
 		StringBuffer<> currentDir;
-		if (!GetCurrentDirectoryW(currentDir))
-			return logger.Error(TC("GetCurrentDirectoryW failed"));
+		CHECK_TRUEF(GetCurrentDirectoryW(currentDir), TC("GetCurrentDirectoryW failed"));
 
 		bool foundFile = false;
 		if (!TraverseDir(logger, rootDir, [&](const DirectoryEntry& de) { foundFile |= TStrcmp(de.name, TC("UbaTestFile2")) == 0; }, true))
 			return logger.Error(TC("Failed to TraverseDir '.'"));
 
-		if (!foundFile)
-			return logger.Error(TC("Did not find UbaTestFile2 with TraverseDir"));
+		CHECK_TRUEF(foundFile, TC("Did not find UbaTestFile2 with TraverseDir"));
 
-		if (!DeleteFileW(testFileName2.data))
-			return false;
+		CHECK_TRUE(DeleteFileW(testFileName2.data));
 		#endif
 
 		return true;
@@ -391,8 +337,7 @@ namespace uba
 		u64 left = 2*1024*1024;
 
 		FileAccessor fileHandle(logger, testFileName.data);
-		if (!fileHandle.CreateWrite(false, DefaultAttributes() | FILE_FLAG_OVERLAPPED, left))
-			return logger.Error(TC("Failed to create file for write"));
+		CHECK_TRUEF(fileHandle.CreateWrite(false, DefaultAttributes() | FILE_FLAG_OVERLAPPED, left), TC("Failed to create file for write"));
 
 		constexpr u64 oddSize = 277872ull;
 		u8 buffer[oddSize];
@@ -400,12 +345,10 @@ namespace uba
 		while (left)
 		{
 			u64 toWrite = Min(left, oddSize);
-			if (!fileHandle.Write(buffer, toWrite))
-				return logger.Error(TC("Failed to create file for write"));
+			CHECK_TRUEF(fileHandle.Write(buffer, toWrite), TC("Failed to create file for write"));
 			left -= toWrite;
 		}
-		if (!fileHandle.Close())
-			return false;
+		CHECK_TRUE(fileHandle.Close());
 		return true;
 	}
 
@@ -421,8 +364,7 @@ namespace uba
 		if (GetHugePageCount())
 		{
 			MemoryBlock block(TC("TestHugePage"));
-			if (!block.Init(1024 * 1024, nullptr, true))
-				return logger.Error(TC("Failed to allocate huge pages even though system says they exists"));
+			CHECK_TRUEF(block.Init(1024 * 1024, nullptr, true), TC("Failed to allocate huge pages even though system says they exists"));
 			u64* mem = (u64*)block.Allocate(8, 1, TC("Foo"));
 			*mem = 0x1234;
 			block.Free(mem);
@@ -504,23 +446,18 @@ namespace uba
 			writer.WriteString(str);
 			BinaryReader reader(mem);
 			TString s = reader.ReadString();
-			if (s.size() != TStrlen(str))
-				return logger.Error(TC("Serialized string '%s' has wrong strlen"), str);
-			if (s != str)
-				return logger.Error(TC("Serialized string '%s' is different from source"), str);
+			CHECK_TRUEF(!(s.size() != TStrlen(str)), TC("Serialized string '%s' has wrong strlen"), str);
+			CHECK_TRUEF(!(s != str), TC("Serialized string '%s' is different from source"), str);
 			return true;
 		};
 
-		if (!testString(TC("Foo")))
-			return false;
+		CHECK_TRUE(testString(TC("Foo")));
 
 		#if PLATFORM_WINDOWS
 		tchar str1[] = { 54620, 44544, 0 };
-		if (!testString(str1))
-			return false;
+		CHECK_TRUE(testString(str1));
 		tchar str2[] = { 'f', 54620, 'o', 44544, 0 };
-		if (!testString(str2))
-			return false;
+		CHECK_TRUE(testString(str2));
 		#endif
 
 		if (!IsDebug)
@@ -532,11 +469,9 @@ namespace uba
 			BinaryReader reader(mem);
 			StringBuffer<3> str;
 			reader.ReadString(str);
-			if (!str.Equals(TC("Fo")))
-				return false;
+			CHECK_TRUE(str.Equals(TC("Fo")));
 			reader.ReadString(str.Clear());
-			if (!str.Equals(TC("Ba")))
-				return false;
+			CHECK_TRUE(str.Equals(TC("Ba")));
 		}
 
 		return true;
@@ -546,10 +481,8 @@ namespace uba
 	bool TestKnownSystemFiles(LoggerWithWriter& logger, const StringBufferBase& rootDir)
 	{
 		for (auto systemFile : g_knownSystemFiles)
-			if (!IsKnownSystemFile(systemFile))
-				return logger.Error(TC("IsKnownSystemFile returned false for %s which is a system file"), systemFile);
-		if (IsKnownSystemFile(TC("Fooo.dll")))
-			return logger.Error(TC("IsKnownSystemFile returned true for Fooo.dll which is not a system file"));
+			CHECK_TRUEF(IsKnownSystemFile(systemFile), TC("IsKnownSystemFile returned false for %s which is a system file"), systemFile);
+		CHECK_TRUEF(!IsKnownSystemFile(TC("Fooo.dll")), TC("IsKnownSystemFile returned true for Fooo.dll which is not a system file"));
 		return true;
 	}
 	#endif
@@ -567,10 +500,8 @@ namespace uba
 		#endif
 
 		RootPaths paths;
-		if (!paths.RegisterRoot(logger, root1))
-			return false;
-		if (!paths.RegisterRoot(logger, root2))
-			return false;
+		CHECK_TRUE(paths.RegisterRoot(logger, root1));
+		CHECK_TRUE(paths.RegisterRoot(logger, root2));
 
 		bool success = true;
 		StringBuffer<> temp;
@@ -593,16 +524,13 @@ namespace uba
 				}
 			}, false, TC(""));
 
-		if (!res)
-			return false;
-		if (!success)
-			return false;
+		CHECK_TRUE(res);
+		CHECK_TRUE(success);
 
 		StringBuffer<> newStr;
 		auto& root = paths.GetRoot(rootPos - RootPaths::RootStartByte);
 		newStr.Append(root.c_str()).Append(temp);
-		if (!newStr.Equals(str))
-			return false;
+		CHECK_TRUE(newStr.Equals(str));
 
 		#if PLATFORM_WINDOWS
 		const tchar str2[] = TC("file://e:/temp/");
@@ -622,8 +550,7 @@ namespace uba
 				}
 
 			}, false, TC(""));
-		if (!res || !foundPath || !success)
-			return false;
+		CHECK_TRUE(res && foundPath && success);
 		#endif
 
 		return true;
@@ -659,8 +586,7 @@ namespace uba
 		{
 			importKernel |= isKnown && Contains(import, TC("KERNEL32.dll"));
 		}, error);
-		if (!importKernel)
-			return logger.Error(TC("Failed to find Kernel32 as import"));
+		CHECK_TRUEF(importKernel, TC("Failed to find Kernel32 as import"));
 #elif PLATFORM_MAC
 		//StringBuffer<> error;
 		//FindImportsMac(TC("[REDACTED_USER_HOME]/p4/fn/Engine/Binaries/Mac/ShaderCompileWorker"), LogImports, error);
@@ -671,8 +597,7 @@ namespace uba
 	bool TestVolumeCache(LoggerWithWriter& logger, const StringBufferBase& rootDir)
 	{
 		VolumeCache cache;
-		if (!cache.Init(logger))
-			return false;
+		CHECK_TRUE(cache.Init(logger));
 		return true;
 	}
 
diff --git a/Engine/Source/Programs/UnrealBuildAccelerator/Test/Private/UbaTestCache.cpp b/Engine/Source/Programs/UnrealBuildAccelerator/Test/Private/UbaTestCache.cpp
--- a/Engine/Source/Programs/UnrealBuildAccelerator/Test/Private/UbaTestCache.cpp
+++ b/Engine/Source/Programs/UnrealBuildAccelerator/Test/Private/UbaTestCache.cpp
@@ -24,14 +24,13 @@ namespace uba
 		u32 expectedExtraCount;
 	};
 
-	bool CheckCacheEntry(CacheEntries& entries, CacheEntry& entry, const TestRecord* rec, bool checkOffsets, bool checkExpectedCounts)
+	bool CheckCacheEntry(Logger& logger, CacheEntries& entries, CacheEntry& entry, const TestRecord* rec, bool checkOffsets, bool checkExpectedCounts)
 	{
 		Set<u32> sharedOffsets;
 		{
 			BinaryReader sharedReader(entries.sharedInputCasKeyOffsets);
 			while (sharedReader.GetLeft())
-				if (!sharedOffsets.insert(u32(sharedReader.Read7BitEncoded())).second)
-					return false;
+				CHECK_TRUE(sharedOffsets.insert(u32(sharedReader.Read7BitEncoded())).second);
 		}
 
 		Set<u32> entryOffsets;
@@ -47,10 +46,8 @@ namespace uba
 			while (sharedReader.GetLeft())
 			{
 				u64 offset;
-				if (!sharedReader.TryRead7BitEncoded(offset))
-					return false;
-				if (!entryOffsets.insert(u32(offset)).second)
-					return false;
+				CHECK_TRUE(sharedReader.TryRead7BitEncoded(offset));
+				CHECK_TRUE(entryOffsets.insert(u32(offset)).second);
 			}
 		}
 
@@ -60,31 +57,25 @@ namespace uba
 		{
 			++extraCount;
 			u32 offset = u32(extraReader.Read7BitEncoded());
-			if (sharedOffsets.find(offset) != sharedOffsets.end())
-				return false;
-			if (!entryOffsets.insert(offset).second)
-				return false;
+			CHECK_TRUE(sharedOffsets.find(offset) == sharedOffsets.end());
+			CHECK_TRUE(entryOffsets.insert(offset).second);
 		}
 
 		if (rec)
 		{
 			if (checkExpectedCounts)
 			{
-				if (rangeCount != rec->expectedRangeCount)
-					return false;
-				if (extraCount != rec->expectedExtraCount)
-					return false;
+				CHECK_TRUE(rangeCount == rec->expectedRangeCount);
+				CHECK_TRUE(extraCount == rec->expectedExtraCount);
 			}
 
 			if (checkOffsets)
 			{
-				if (rec->offsets != entryOffsets)
-					return false;
+				CHECK_TRUE(rec->offsets == entryOffsets);
 			}
 			else
 			{
-				if (rec->offsets.size() != entryOffsets.size())
-					return false;
+				CHECK_TRUE(rec->offsets.size() == entryOffsets.size());
 			}
 		}
 		return true;
@@ -109,7 +100,7 @@ namespace uba
 		CacheEntries entries;
 
 		using Record = TestRecord;
-		auto CheckEntry = [&](CacheEntry& entry, const Record* rec, bool checkOffsets, bool checkExpectedCounts) { return CheckCacheEntry(entries, entry, rec, checkOffsets, checkExpectedCounts); };
+		auto CheckEntry = [&](CacheEntry& entry, const Record* rec, bool checkOffsets, bool checkExpectedCounts) { return CheckCacheEntry(logger, entries, entry, rec, checkOffsets, checkExpectedCounts); };
 
 		auto AddEntry = [&](const Record& rec, bool checkExpectedCount)
 			{
@@ -336,10 +327,8 @@ namespace uba
 				u32 offsets2 = table2.AddNoLock(context, offsets[i]);
 
 				StringBuffer<> temp;
-				if (!table2.GetString(temp, offsets2))
-					return logger.Error(TC("Error getting offset %u from table2"), offsets2);
-				if (!temp.Equals(paths[i].data))
-					return logger.Error(TC("Error adding %s to table2. Found %s"), paths[i].data, temp.data);
+				CHECK_TRUEF(table2.GetString(temp, offsets2), TC("Error getting offset %u from table2"), offsets2);
+				CHECK_TRUEF(temp.Equals(paths[i].data), TC("Error adding %s to table2. Found %s"), paths[i].data, temp.data);
 			}
 		}
 
@@ -358,19 +347,15 @@ namespace uba
 					return str.Equals(str2.data);
 				};
 
-			if (!testStr(str.Append("foo")))
-				return false;
-			if (!testStr(str.Clear().Append("foo").EnsureEndsWithSlash().Append("bar.h")))
-				return false;
-			if (!testStr(str.Clear().Append(PathSeparator).Append("foo").Append(PathSeparator).Append("bar.h")))
-				return false;
+			CHECK_TRUE(testStr(str.Append("foo")));
+			CHECK_TRUE(testStr(str.Clear().Append("foo").EnsureEndsWithSlash().Append("bar.h")));
+			CHECK_TRUE(testStr(str.Clear().Append(PathSeparator).Append("foo").Append(PathSeparator).Append("bar.h")));
 
 			CompactPathTable table2(CaseInsensitiveFs, 0, 0, version);
 			BinaryReader reader(table.GetMemory(), 0, table.GetSize());
 			table2.ReadMem(reader, true);
 			u32 offset2 = table2.Add(str.data, str.count);
-			if (offset != offset2)
-				return false;
+			CHECK_TRUE(offset == offset2);
 		}
 		return true;
 	}
@@ -385,12 +370,9 @@ namespace uba
 		CompactCasKeyTable table2;
 		BinaryReader reader(table.GetMemory(), 0, table.GetSize());
 		table2.ReadMem(reader, true);
-		if (table2.AddNoLock(CasKeyZero, 0) != offset0)
-			return false;
-		if (table2.AddNoLock(CasKeyZero, 1) != offset1)
-			return false;
-		if (table2.AddNoLock(CasKeyZero, 2) != offset2)
-			return false;
+		CHECK_TRUE(table2.AddNoLock(CasKeyZero, 0) == offset0);
+		CHECK_TRUE(table2.AddNoLock(CasKeyZero, 1) == offset1);
+		CHECK_TRUE(table2.AddNoLock(CasKeyZero, 2) == offset2);
 
 		for (u32 i=0;i!=32;++i)
 		{
@@ -407,39 +389,31 @@ namespace uba
 			MemoryBlock memoryBlock(TC("TestHashTable"), 1024*1024);
 			HashMap<u32, u32> casMap;
 			casMap.Init(memoryBlock, 3);
-			if (casMap.Find(1))
-				return false;
+			CHECK_TRUE(!casMap.Find(1));
 			casMap.Insert(1) = 2;
-			if (*casMap.Find(1) != 2)
-				return false;
+			CHECK_TRUE(*casMap.Find(1) == 2);
 			casMap.Insert(1) = 3;
-			if (*casMap.Find(1) != 3)
-				return false;
+			CHECK_TRUE(*casMap.Find(1) == 3);
 
 			HashMap<u32, u32, true> casMap2;
 			casMap2.Init(4);
 			for (u32 i=0;i!=4; ++i)
 			{
 				casMap2.Insert(i) = i;
-				if (*casMap2.Find(i) != i)
-					return false;
+				CHECK_TRUE(*casMap2.Find(i) == i);
 			}
 			for (u32 i=0;i!=4; ++i)
-				if (*casMap2.Find(i) != i)
-					return false;
+				CHECK_TRUE(*casMap2.Find(i) == i);
 
 			casMap2.Insert(4) = 4;
 			for (u32 i=0;i!=5; ++i)
-				if (*casMap2.Find(i) != i)
-					return false;
+				CHECK_TRUE(*casMap2.Find(i) == i);
 
 			for (u32 i=5;i!=1000; ++i)
 			{
-				if (casMap2.Find(i))
-					return false;
+				CHECK_TRUE(!casMap2.Find(i));
 				casMap2.Insert(i) = i;
-				if (*casMap2.Find(i) != i)
-					return false;
+				CHECK_TRUE(*casMap2.Find(i) == i);
 			}
 		}
 
@@ -589,7 +563,8 @@ namespace uba
 
 				{
 					CacheResult result;
-					CHECK_TRUE(!cacheClient.FetchFromCache(result, RootPaths(), 0, psi) && !result.hit);
+					CHECK_TRUE(!cacheClient.FetchFromCache(result, RootPaths(), 0, psi));
+					CHECK_TRUE(!result.hit);
 					CHECK_TRUE(cacheClient.WriteToCache(RootPaths(), 0, psi, inputs.GetData(), inputs.GetPosition(), outputs.GetData(), outputs.GetPosition(), logLines.GetData(), logLines.GetPosition()));
 					CHECK_TRUE(DeleteFileW(outputFile.data));
 					CHECK_TRUE(!FileExists(logger, outputFile.data));
@@ -605,7 +580,8 @@ namespace uba
 					InvalidateCachedInfo(storageServer, inputFile);
 
 					CacheResult result;
-					CHECK_TRUE(!cacheClient.FetchFromCache(result, RootPaths(), 0, psi) && !result.hit);
+					CHECK_TRUE(!cacheClient.FetchFromCache(result, RootPaths(), 0, psi));
+					CHECK_TRUE(!result.hit);
 
 					CHECK_TRUE(cacheClient.WriteToCache(RootPaths(), 0, psi, inputs.GetData(), inputs.GetPosition(), outputs.GetData(), outputs.GetPosition(), logLines.GetData(), logLines.GetPosition()));
 
@@ -681,16 +657,14 @@ namespace uba
 		#endif
 
 		FileAccessor bucketFile(logger, fileName);
-		if (!bucketFile.OpenMemoryRead())
-			return false;
+		CHECK_TRUE(bucketFile.OpenMemoryRead());
 		BinaryReader reader(bucketFile.GetData(), 0, bucketFile.GetSize());
 		u32 bucketVersion = reader.ReadU32();
 		CacheBucket::LoadStats stats;
 		StorageServer* storage = nullptr;
 
 		CacheBucket bucket(0, 0);
-		if (!bucket.Load(logger, reader, bucketVersion, stats, *storage))
-			return false;
+		CHECK_TRUE(bucket.Load(logger, reader, bucketVersion, stats, *storage));
 		bucketFile.Close();
 
 		LoggerWithWriter lg(g_consoleLogWriter);
diff --git a/Engine/Source/Programs/UnrealBuildAccelerator/Test/Private/UbaTestConfig.cpp b/Engine/Source/Programs/UnrealBuildAccelerator/Test/Private/UbaTestConfig.cpp
--- a/Engine/Source/Programs/UnrealBuildAccelerator/Test/Private/UbaTestConfig.cpp
+++ b/Engine/Source/Programs/UnrealBuildAccelerator/Test/Private/UbaTestConfig.cpp
@@ -3,6 +3,7 @@
 #include "UbaConfig.h"
 #include "UbaLogger.h"
 #include "UbaPlatform.h"
+#include "UbaTest.h"
 
 namespace uba
 {
@@ -16,25 +17,18 @@ namespace uba
 			"";
 
 		Config config;
-		if (!config.LoadFromText(logger, configText, strlen(configText)))
-			return false;
+		CHECK_TRUE(config.LoadFromText(logger, configText, strlen(configText)));
 
 		const ConfigTable* tablePtr = config.GetTable(TC("CacheClient"));
-		if (!tablePtr)
-			return false;
+		CHECK_TRUE(tablePtr);
 		const ConfigTable& table = *tablePtr;
 		bool test = false;
-		if (!table.GetValueAsBool(test, TC("UseDirectoryPreparsing")))
-			return false;
-		if (test != true)
-			return false;
+		CHECK_TRUE(table.GetValueAsBool(test, TC("UseDirectoryPreparsing")));
+		CHECK_TRUE(test == true);
 		const tchar* str = nullptr;
-		if (!table.GetValueAsString(str, TC("RootDir")))
-			return false;
-		if (TStrcmp(str, TC("e:\\foo")) != 0)
-			return false;
-		if (table.GetValueAsBool(test, TC("Comment")))
-			return false;
+		CHECK_TRUE(table.GetValueAsString(str, TC("RootDir")));
+		CHECK_TRUE(TStrcmp(str, TC("e:\\foo")) == 0);
+		CHECK_TRUE(!table.GetValueAsBool(test, TC("Comment")));
 		return true;
 	}
 
@@ -49,18 +43,14 @@ namespace uba
 		}
 
 		Config config;
-		if (!config.LoadFromText(logger, data.data(), data.size()))
-			return false;
+		CHECK_TRUE(config.LoadFromText(logger, data.data(), data.size()));
 
 		const ConfigTable* tablePtr = config.GetTable(TC("TestTable"));
-		if (!tablePtr)
-			return false;
+		CHECK_TRUE(tablePtr);
 		const ConfigTable& table = *tablePtr;
 		int foo = 0;
-		if (!table.GetValueAsInt(foo, TC("Foo")))
-			return false;
-		if (foo != 42)
-			return false;
+		CHECK_TRUE(table.GetValueAsInt(foo, TC("Foo")));
+		CHECK_TRUE(foo == 42);
 		return true;
 	}
 }
\ No newline at end of file
diff --git a/Engine/Source/Programs/UnrealBuildAccelerator/Test/Private/UbaTestCrypto.cpp b/Engine/Source/Programs/UnrealBuildAccelerator/Test/Private/UbaTestCrypto.cpp
--- a/Engine/Source/Programs/UnrealBuildAccelerator/Test/Private/UbaTestCrypto.cpp
+++ b/Engine/Source/Programs/UnrealBuildAccelerator/Test/Private/UbaTestCrypto.cpp
@@ -3,6 +3,7 @@
 #include "UbaCrypto.h"
 #include "UbaLogger.h"
 #include "UbaStringBuffer.h"
+#include "UbaTest.h"
 
 namespace uba
 {
@@ -26,18 +27,14 @@ namespace uba
 			for (u32 i=0; i!=3; ++i)
 			{
 				Guid iv;
-				if (!Crypto::Encrypt(logger, encryptKey, encryptedData, dataSize, iv))
-					return false;
+				CHECK_TRUE(Crypto::Encrypt(logger, encryptKey, encryptedData, dataSize, iv));
 
-				if (memcmp(encryptedData, originalData, dataSize) == 0)
-					return false;
+				CHECK_TRUE(memcmp(encryptedData, originalData, dataSize) != 0);
 
 				Guid iv2;
-				if (!Crypto::Decrypt(logger, decryptKey, encryptedData, dataSize, iv2))
-					return false;
+				CHECK_TRUE(Crypto::Decrypt(logger, decryptKey, encryptedData, dataSize, iv2));
 
-				if (memcmp(encryptedData, originalData, dataSize) != 0)
-					return false;
+				CHECK_TRUE(memcmp(encryptedData, originalData, dataSize) == 0);
 
 				if (i == 1)
 					decryptKey = Crypto::DuplicateKey(logger, encryptKey);
diff --git a/Engine/Source/Programs/UnrealBuildAccelerator/Test/Private/UbaTestDependencyCrawler.cpp b/Engine/Source/Programs/UnrealBuildAccelerator/Test/Private/UbaTestDependencyCrawler.cpp
--- a/Engine/Source/Programs/UnrealBuildAccelerator/Test/Private/UbaTestDependencyCrawler.cpp
+++ b/Engine/Source/Programs/UnrealBuildAccelerator/Test/Private/UbaTestDependencyCrawler.cpp
@@ -17,8 +17,7 @@ namespace uba
 		auto CreateFileFunc = [&](TrackWorkScope& tracker, const StringView& fileName, const DependencyCrawler::AccessFileFunc& func)
 			{
 				FileAccessor fa(logger, fileName.data);
-				if (!fa.OpenMemoryRead())
-					return false;
+				CHECK_TRUE(fa.OpenMemoryRead());
 				return func(fa.GetData(), fa.GetSize());
 			};
 
diff --git a/Engine/Source/Programs/UnrealBuildAccelerator/Test/Private/UbaTestFileMappings.cpp b/Engine/Source/Programs/UnrealBuildAccelerator/Test/Private/UbaTestFileMappings.cpp
--- a/Engine/Source/Programs/UnrealBuildAccelerator/Test/Private/UbaTestFileMappings.cpp
+++ b/Engine/Source/Programs/UnrealBuildAccelerator/Test/Private/UbaTestFileMappings.cpp
@@ -23,14 +23,12 @@ namespace uba
 			u64 v = i + 1337;
 			*(u64*)view.memory = v;
 			u64 value = *(u64*)view.memory;
-			if (value != v)
-				return false;
+			CHECK_TRUE(value == v);
 			mappingBuffer.UnmapView(view, TC("Foo"));
 
 			view = mappingBuffer.MapView(view.handle, view.offset, 1024, TC("Foo"));
 			value = *(u64*)view.memory;
-			if (value != v)
-				return false;
+			CHECK_TRUE(value == v);
 			mappingBuffer.UnmapView(view, TC("Foo"));
 		}
 
@@ -47,13 +45,11 @@ namespace uba
 			MappedView view = mappingBuffer.AllocAndMapView(1024, 1, TC("Foo"), true);
 			*(u64*)view.memory = 1337;
 			u64 value = *(u64*)view.memory;
-			if (value != 1337)
-				return false;
+			CHECK_TRUE(value == 1337);
 			mappingBuffer.UnmapView(view, TC("Foo"), 1024);
 			view = mappingBuffer.MapView(view.handle, view.offset, 1024, TC("Foo"));
 			value = *(u64*)view.memory;
-			if (value != 1337)
-				return false;
+			CHECK_TRUE(value == 1337);
 			mappingBuffer.UnmapView(view, TC("Foo"), 1024);
 		}
 		return true;
diff --git a/Engine/Source/Programs/UnrealBuildAccelerator/Test/Private/UbaTestNetwork.cpp b/Engine/Source/Programs/UnrealBuildAccelerator/Test/Private/UbaTestNetwork.cpp
--- a/Engine/Source/Programs/UnrealBuildAccelerator/Test/Private/UbaTestNetwork.cpp
+++ b/Engine/Source/Programs/UnrealBuildAccelerator/Test/Private/UbaTestNetwork.cpp
@@ -55,8 +55,7 @@ namespace uba
 		Sleep(200);
 		t.Wait();
 
-		if (result != 42)
-			return logger.Error(TC("Failed to receive data"));
+		CHECK_TRUEF(!(result != 42), TC("Failed to receive data"));
 
 		return true;
 	}
@@ -79,26 +78,22 @@ namespace uba
 			});
 
 		auto ds = MakeGuard([&]() { networkBackend.StopListen(); server.DisconnectClients(); });
-		if (!server.StartListen(networkBackend, 1234))
-			return logger.Error(TC("Failed to listen"));
+		CHECK_TRUEF(server.StartListen(networkBackend, 1234), TC("Failed to listen"));
 		Sleep(100);
 		auto dg = MakeGuard([&]() { client.Disconnect(); });
-		if (!client.Connect(networkBackend, TC("127.0.0.1"), 1234))
-			return logger.Error(TC("Failed to connect"));
+		CHECK_TRUEF(client.Connect(networkBackend, TC("127.0.0.1"), 1234), TC("Failed to connect"));
 
 
 		StackBinaryWriter<128> writer;
 		NetworkMessage msg(client, 1, SessionMessageType_Ping, writer);
 		writer.WriteU32(32);
 		StackBinaryReader<32> reader;
-		if (!msg.Send(reader))
-			return logger.Error(TC("Failed to get message"));
+		CHECK_TRUEF(msg.Send(reader), TC("Failed to get message"));
 
 		u8 value = reader.ReadByte();
 		logger.Info(TC("Got value %u"), value);
 
-		if (!server.SendKeepAlive())
-			return logger.Error(TC("Failed to send keep alive"));
+		CHECK_TRUEF(server.SendKeepAlive(), TC("Failed to send keep alive"));
 		Sleep(100);
 		return true;
 	}
@@ -120,21 +115,17 @@ namespace uba
 				return true;
 			});
 
-		if (!client.StartListen(networkBackend, 1239))
-			return logger.Error(TC("Client failed to listen"));
+		CHECK_TRUEF(client.StartListen(networkBackend, 1239), TC("Client failed to listen"));
 		Sleep(100);
 		auto ds = MakeGuard([&]() { networkBackend.StopListen(); server.DisconnectClients(); });
 		auto dg = MakeGuard([&]() { client.Disconnect(); });
-		if (!server.AddClient(networkBackend, TC("127.0.0.1"), 1239))
-			return logger.Error(TC("Server failed to connect"));
-		if (!server.AddClient(networkBackend, TC("127.0.0.1"), 1239))
-			return logger.Error(TC("Server failed to connect second"));
+		CHECK_TRUEF(server.AddClient(networkBackend, TC("127.0.0.1"), 1239), TC("Server failed to connect"));
+		CHECK_TRUEF(server.AddClient(networkBackend, TC("127.0.0.1"), 1239), TC("Server failed to connect second"));
 
 		u64 time = GetTime();
 		while (!client.GetConnectionCount())
 		{
-			if (TimeToMs(GetTime() - time) > 4000)
-				return logger.Error(TC("Client failed to establish connection"));
+			CHECK_TRUEF(!(TimeToMs(GetTime() - time) > 4000), TC("Client failed to establish connection"));
 			Sleep(100);
 		}
 
@@ -142,8 +133,7 @@ namespace uba
 		NetworkMessage msg(client, 1, SessionMessageType_Ping, writer);
 		writer.WriteU32(32);
 		StackBinaryReader<32> reader;
-		if (!msg.Send(reader))
-			return logger.Error(TC("Failed to get message"));
+		CHECK_TRUEF(msg.Send(reader), TC("Failed to get message"));
 
 		u8 value = reader.ReadByte();
 		logger.Info(TC("Got value %u"), value);
@@ -168,24 +158,20 @@ namespace uba
 			});
 
 		auto ds = MakeGuard([&]() { server.DisconnectClients(); });
-		if (!server.StartListen(networkBackend, 1234))
-			return logger.Error(TC("Failed to listen"));
+		CHECK_TRUEF(server.StartListen(networkBackend, 1234), TC("Failed to listen"));
 		Sleep(100);
 		auto dg = MakeGuard([&]() { client.Disconnect(); });
-		if (!client.Connect(networkBackend, TC("127.0.0.1"), 1234))
-			return logger.Error(TC("Failed to connect"));
+		CHECK_TRUEF(client.Connect(networkBackend, TC("127.0.0.1"), 1234), TC("Failed to connect"));
 
 
 		StackBinaryWriter<128> writer;
 		NetworkMessage msg(client, 1, SessionMessageType_Ping, writer);
 		writer.WriteU32(32);
 		StackBinaryReader<32> reader;
-		if (!msg.Send(reader))
-			return logger.Error(TC("Failed to get message"));
+		CHECK_TRUEF(msg.Send(reader), TC("Failed to get message"));
 
 		u8 value = reader.ReadByte();
-		if (value != 42)
-			return logger.Error(TC("Got wrong value %u, expected 42"), value);
+		CHECK_TRUEF(!(value != 42), TC("Got wrong value %u, expected 42"), value);
 		logger.Info(TC("Got value %u"), value);
 		return true;
 	}
@@ -217,12 +203,10 @@ namespace uba
 			});
 
 		auto ds = MakeGuard([&]() { networkBackend.StopListen(); server.DisconnectClients(); });
-		if (!server.StartListen(networkBackend, 1234, TC("0.0.0.0"), true))
-			return logger.Error(TC("Failed to listen"));
+		CHECK_TRUEF(server.StartListen(networkBackend, 1234, TC("0.0.0.0"), true), TC("Failed to listen"));
 		Sleep(100);
 		auto dg = MakeGuard([&]() { client.Disconnect(); });
-		if (!client.Connect(networkBackend, TC("127.0.0.1"), 1234))
-			return logger.Error(TC("Failed to connect"));
+		CHECK_TRUEF(client.Connect(networkBackend, TC("127.0.0.1"), 1234), TC("Failed to connect"));
 
 
 		{
@@ -231,8 +215,7 @@ namespace uba
 			writer.WriteU32(32);
 			writer.WriteGuid({});
 			StackBinaryReader<32> reader;
-			if (!msg.Send(reader))
-				return logger.Error(TC("Failed to get message"));
+			CHECK_TRUEF(msg.Send(reader), TC("Failed to get message"));
 			u8 value = reader.ReadByte();
 			CHECK_TRUE(value == 42);
 		}
@@ -246,14 +229,12 @@ namespace uba
 			if (!msg.SendAsync(reader, [](bool error, void* userData) { ((Event*)userData)->Set(); }, &ev))
 				return logger.Error(TC("Failed to get message"));
 			ev.IsSet();
-			if (!msg.ProcessAsyncResults(reader))
-				return false;
+			CHECK_TRUE(msg.ProcessAsyncResults(reader));
 			u8 value = reader.ReadByte();
 			CHECK_TRUE(value == 42);
 		}
 
-		if (!server.SendKeepAlive())
-			return logger.Error(TC("Failed to send keep alive"));
+		CHECK_TRUEF(server.SendKeepAlive(), TC("Failed to send keep alive"));
 		Sleep(100);
 		return true;
 	}
diff --git a/Engine/Source/Programs/UnrealBuildAccelerator/Test/Private/UbaTestSession.cpp b/Engine/Source/Programs/UnrealBuildAccelerator/Test/Private/UbaTestSession.cpp
--- a/Engine/Source/Programs/UnrealBuildAccelerator/Test/Private/UbaTestSession.cpp
+++ b/Engine/Source/Programs/UnrealBuildAccelerator/Test/Private/UbaTestSession.cpp
@@ -27,8 +27,7 @@ namespace uba
 
 		StringBuffer<MaxPath> rootDir;
 		rootDir.Append(testRootDir).Append(TCV("Uba"));
-		if (!DeleteAllFiles(logger, rootDir.data))
-			return false;
+		CHECK_TRUE(DeleteAllFiles(logger, rootDir.data));
 
 		bool ctorSuccess = true;
 		NetworkServer server(ctorSuccess, { logWriter });
@@ -51,13 +50,10 @@ namespace uba
 
 		StringBuffer<MaxPath> workingDir;
 		workingDir.Append(testRootDir).Append(TCV("WorkingDir"));
-		if (!DeleteAllFiles(logger, workingDir.data))
-			return false;
+		CHECK_TRUE(DeleteAllFiles(logger, workingDir.data));
 
-		if (!storage.CreateDirectory(workingDir.data))
-			return false;
-		if (!DeleteAllFiles(logger, workingDir.data, false))
-			return false;
+		CHECK_TRUE(storage.CreateDirectory(workingDir.data));
+		CHECK_TRUE(DeleteAllFiles(logger, workingDir.data, false));
 		workingDir.EnsureEndsWithSlash();
 		return testFunc(logger, session, workingDir.data, [&](ProcessStartInfo& pi) { return session.RunProcess(pi, true, enableDetour); }, extraData);
 	}
@@ -91,8 +87,7 @@ namespace uba
 
 		ProcessHandle process = runProcess(processInfo);
 		CHECK_TRUEF(process.IsValid(), TC("Failed to run UbaTestApp"));
-		if (!process.WaitForExit(GetTimeoutTime()))
-			return logger.Error(TC("UbaTestApp did not exit in 10 seconds"));
+		CHECK_TRUEF(process.WaitForExit(GetTimeoutTime()), TC("UbaTestApp did not exit in 10 seconds"));
 		if (outHandle)
 			*outHandle = process;
 		u32 exitCode = process.GetExitCode();
@@ -120,8 +115,7 @@ namespace uba
 		StringBuffer<MaxPath> toDelete(rootDir);
 		if (!deleteAll)
 			toDelete.Append(PathSeparator).Append(TCV("sessions"));
-		if (!DeleteAllFiles(logger, toDelete.data))
-			return false;
+		CHECK_TRUE(DeleteAllFiles(logger, toDelete.data));
 
 		StorageServerCreateInfo storageServerInfo(server, rootDir.data, logWriter);
 		storageServerInfo.casCapacityBytes = 1024ull * 1024 * 1024;
@@ -152,12 +146,9 @@ namespace uba
 
 		StringBuffer<MaxPath> workingDir;
 		workingDir.Append(testRootDir).Append(TCV("WorkingDir"));
-		if (deleteAll && !DeleteAllFiles(logger, workingDir.data))
-			return false;
-		if (!storageServer.CreateDirectory(workingDir.data))
-			return false;
-		if (deleteAll && !DeleteAllFiles(logger, workingDir.data, false))
-			return false;
+		CHECK_TRUE(!deleteAll || DeleteAllFiles(logger, workingDir.data));
+		CHECK_TRUE(storageServer.CreateDirectory(workingDir.data));
+		CHECK_TRUE(!deleteAll || DeleteAllFiles(logger, workingDir.data, false));
 
 		workingDir.EnsureEndsWithSlash();
 		return testFunc(logger, workingDir, sessionServer);
@@ -179,27 +170,21 @@ namespace uba
 
 		if (serverShouldListen)
 		{
-			if (!server.StartListen(tcpBackend, port))
-				return logger.Error(TC("Failed to start listen"));
-			if (!client.Connect(tcpBackend, TC("127.0.0.1"), port))
-				return logger.Error(TC("Failed to connect"));
-			if (!client.Connect(tcpBackend, TC("127.0.0.1"), port))
-				return logger.Error(TC("Failed to connect"));
+			CHECK_TRUEF(server.StartListen(tcpBackend, port), TC("Failed to start listen"));
+			CHECK_TRUEF(client.Connect(tcpBackend, TC("127.0.0.1"), port), TC("Failed to connect"));
+			CHECK_TRUEF(client.Connect(tcpBackend, TC("127.0.0.1"), port), TC("Failed to connect"));
 		}
 		else
 		{
-			if (!client.StartListen(tcpBackend, port))
-				return logger.Error(TC("Failed to listen"));
-			if (!server.AddClient(tcpBackend, TC("127.0.0.1"), port))
-				return logger.Error(TC("Failed to connect"));
+			CHECK_TRUEF(client.StartListen(tcpBackend, port), TC("Failed to listen"));
+			CHECK_TRUEF(server.AddClient(tcpBackend, TC("127.0.0.1"), port), TC("Failed to connect"));
 			while (!client.IsConnected())
 				Sleep(1);
 		}
 		auto disconnectGuard = MakeGuard([&]() { tcpBackend.StopListen(); client.Disconnect(); server.RemoveDisconnectedConnections(); });
 
 		Config config;
-		if (!client.FetchConfig(config))
-			return logger.Error(TC("Failed to fetch config"));
+		CHECK_TRUEF(client.FetchConfig(config), TC("Failed to fetch config"));
 
 		StringBuffer<MaxPath> rootDir;
 		rootDir.Append(testRootDir).Append(TCV("UbaClient")).AppendValue(port);
@@ -207,8 +192,7 @@ namespace uba
 		StringBuffer<MaxPath> toDelete(rootDir);
 		if (!deleteAll)
 			toDelete.Append(PathSeparator).Append(TCV("sessions"));
-		if (!DeleteAllFiles(logger, toDelete.data))
-			return logger.Error(TC("Failed to delete all files"));
+		CHECK_TRUEF(DeleteAllFiles(logger, toDelete.data), TC("Failed to delete all files"));
 
 		StorageClientCreateInfo storageClientInfo(client, rootDir.data);
 		storageClientInfo.Apply(config);
@@ -265,8 +249,7 @@ namespace uba
 	{
 		outPath.Clear().Append(workingDir).EnsureEndsWithSlash().Append(fileName);
 		FileAccessor fr(logger, outPath.data);
-		if (!fr.CreateWrite())
-			return false;
+		CHECK_TRUE(fr.CreateWrite());
 		u32 len = 0;
 		for (const CharType* it=text; *it; ++it)
 			++len;
@@ -277,50 +260,38 @@ namespace uba
 	bool RunTestAppTests(LoggerWithWriter& logger, SessionServer& session, const tchar* workingDir, const RunProcessFunction& runProcess, uintptr_t extraData)
 	{
 		StringBuffer<MaxPath> fileR;
-		if (!CreateTextFile(fileR, logger, workingDir, TC("FileR.h"), "Foo"))
-			return false;
+		CHECK_TRUE(CreateTextFile(fileR, logger, workingDir, TC("FileR.h"), "Foo"));
 
 		{
 			StringBuffer<MaxPath> dir;
 			dir.Append(workingDir).Append(TCV("Dir1"));
-			if (!CreateDirectoryW(dir.data))
-				return logger.Error(TC("Failed to create dir %s"), dir.data);
+			CHECK_TRUEF(CreateDirectoryW(dir.data), TC("Failed to create dir %s"), dir.data);
 
 			dir.Clear().Append(workingDir).Append(TCV("Dir2"));
-			if (!CreateDirectoryW(dir.data))
-				return logger.Error(TC("Failed to create dir %s"), dir.data);
+			CHECK_TRUEF(CreateDirectoryW(dir.data), TC("Failed to create dir %s"), dir.data);
 			dir.EnsureEndsWithSlash().Append(TCV("Dir3"));
-			if (!CreateDirectoryW(dir.data))
-				return logger.Error(TC("Failed to create dir %s"), dir.data);
+			CHECK_TRUEF(CreateDirectoryW(dir.data), TC("Failed to create dir %s"), dir.data);
 			dir.EnsureEndsWithSlash().Append(TCV("Dir4"));
-			if (!CreateDirectoryW(dir.data))
-				return logger.Error(TC("Failed to create dir %s"), dir.data);
+			CHECK_TRUEF(CreateDirectoryW(dir.data), TC("Failed to create dir %s"), dir.data);
 			dir.EnsureEndsWithSlash().Append(TCV("Dir5"));
-			if (!CreateDirectoryW(dir.data))
-				return logger.Error(TC("Failed to create dir %s"), dir.data);
+			CHECK_TRUEF(CreateDirectoryW(dir.data), TC("Failed to create dir %s"), dir.data);
 		}
 
-		if (!CreateTestFile(logger, ToView(workingDir), TCV("File4.out"), TCV("0")))
-			return false;
-		if (!CreateTestFile(logger, ToView(workingDir), TCV("File4.obj"), TCV("9")))
-			return false;
-		if (!CreateTestFile(logger, ToView(workingDir), TCV("File4.lib"), TCV("9")))
-			return false;
+		CHECK_TRUE(CreateTestFile(logger, ToView(workingDir), TCV("File4.out"), TCV("0")));
+		CHECK_TRUE(CreateTestFile(logger, ToView(workingDir), TCV("File4.obj"), TCV("9")));
+		CHECK_TRUE(CreateTestFile(logger, ToView(workingDir), TCV("File4.lib"), TCV("9")));
 
-		if (!RunTestApp(logger, workingDir, runProcess, TC("")))
-			return false;
+		CHECK_TRUE(RunTestApp(logger, workingDir, runProcess, TC("")));
 
 		{
 			StringBuffer<MaxPath> fileW2;
 			fileW2.Append(workingDir).Append(TCV("FileW2"));
-			if (!FileExists(logger, fileW2.data))
-				return logger.Error(TC("Can't find file %s"), fileW2.data);
+			CHECK_TRUEF(FileExists(logger, fileW2.data), TC("Can't find file %s"), fileW2.data);
 		}
 		{
 			StringBuffer<MaxPath> fileWF;
 			fileWF.Append(workingDir).Append(TCV("FileWF"));
-			if (!FileExists(logger, fileWF.data))
-				return logger.Error(TC("Can't find file %s"), fileWF.data);
+			CHECK_TRUEF(FileExists(logger, fileWF.data), TC("Can't find file %s"), fileWF.data);
 		}
 		return true;
 	}
@@ -350,13 +321,10 @@ namespace uba
 		StringBuffer<MaxPath> sourceFile;
 		sourceFile.Append(workingDir).Append(TCV("Code.cpp"));
 		FileAccessor codeFile(logger, sourceFile.data);
-		if (!codeFile.CreateWrite())
-			return false;
+		CHECK_TRUE(codeFile.CreateWrite());
 		char code[] = "#include <stdio.h>\n int main() { printf(\"Hello world\\n\"); return 0; }";
-		if (!codeFile.Write(code, sizeof(code) - 1))
-			return false;
-		if (!codeFile.Close())
-			return false;
+		CHECK_TRUE(codeFile.Write(code, sizeof(code) - 1));
+		CHECK_TRUE(codeFile.Close());
 
 #if PLATFORM_WINDOWS
 		const tchar* clangPath = TC("c:\\Program Files\\Microsoft Visual Studio\\2022\\Professional\\VC\\Tools\\Llvm\\bin\\clang-cl.exe");
@@ -405,11 +373,9 @@ namespace uba
 		processInfo.workingDir = workingDir;
 		//processInfo.logFile = TC("/mnt/e/temp/ttt/RunClang.log");
 		ProcessHandle process = runProcess(processInfo);
-		if (!process.WaitForExit(GetTimeoutTime()))
-			return logger.Error(TC("clang++ timed out"));
+		CHECK_TRUEF(process.WaitForExit(GetTimeoutTime()), TC("clang++ timed out"));
 		u32 exitCode = process.GetExitCode();
-		if (exitCode != 0)
-			return logger.Error(TC("clang++ returned exit code %u"), exitCode);
+		CHECK_TRUEF(!(exitCode != 0), TC("clang++ returned exit code %u"), exitCode);
 		return true;
 	}
 
@@ -426,10 +392,8 @@ namespace uba
 				return u32(helloBytes);
 			});
 
-		if (!RunTestApp(logger, workingDir, runProcess, TC("Whatever")))
-			return false;
-		if (!gotMessage)
-			return logger.Error(TC("Never got message from UbaTestApp"));
+		CHECK_TRUE(RunTestApp(logger, workingDir, runProcess, TC("Whatever")));
+		CHECK_TRUEF(gotMessage, TC("Never got message from UbaTestApp"));
 		return true;
 	}
 
@@ -468,8 +432,7 @@ namespace uba
 	bool TestRemoteDetouredClang(LoggerWithWriter& logger, const StringBufferBase& testRootDir)
 	{
 		// Run twice to test LoadCasTable/SaveCasTable etc
-		if (!RunRemote(logger, testRootDir, RunClang))
-			return false;
+		CHECK_TRUE(RunRemote(logger, testRootDir, RunClang));
 		return RunRemote(logger, testRootDir, RunClang, false);
 	}
 
@@ -495,16 +458,13 @@ namespace uba
 				processInfo.arguments = file.data;
 				processInfo.logFile = TC("[REDACTED_USER_HOME]/Touch.log");
 				ProcessHandle process = runProcess(processInfo);
-				if (!process.WaitForExit(GetTimeoutTime()))
-					return logger.Error(TC("UbaTestApp did not exit in 10 seconds"));
+				CHECK_TRUEF(process.WaitForExit(GetTimeoutTime()), TC("UbaTestApp did not exit in 10 seconds"));
 				u32 exitCode = process.GetExitCode();
-				if (exitCode != 0)
-					return false;
+				CHECK_TRUE(exitCode == 0);
 
 				FileInformation newInfo;
 				CHECK_TRUE(GetFileInformation(newInfo, logger, file));
-				if (newInfo.lastWriteTime == oldInfo.lastWriteTime)
-					return logger.Error(TC("File time not changed after touch"));
+				CHECK_TRUEF(!(newInfo.lastWriteTime == oldInfo.lastWriteTime), TC("File time not changed after touch"));
 				return true;
 			});
 	}
@@ -527,8 +487,7 @@ namespace uba
 					};
 
 				ProcessHandle process = runProcess(processInfo);
-				if (!process.WaitForExit(GetTimeoutTime()))
-					return logger.Error(TC("UbaTestApp did not exit in 10 seconds"));
+				CHECK_TRUEF(process.WaitForExit(GetTimeoutTime()), TC("UbaTestApp did not exit in 10 seconds"));
 				u32 exitCode = process.GetExitCode();
 
 				if (exitCode != 0)
@@ -593,11 +552,9 @@ namespace uba
 
 				for (auto& process : processes)
 				{
-					if (!process.WaitForExit(GetTimeoutTime()))
-						return logger.Error(TC("UbaTestApp did not exit in 10 seconds"));
+					CHECK_TRUEF(process.WaitForExit(GetTimeoutTime()), TC("UbaTestApp did not exit in 10 seconds"));
 					u32 exitCode = process.GetExitCode();
-					if (exitCode != 0)
-						return logger.Error(TC("UbaTestApp exited with code %u"), exitCode);
+					CHECK_TRUEF(!(exitCode != 0), TC("UbaTestApp exited with code %u"), exitCode);
 				}
 
 				return true;
@@ -620,13 +577,10 @@ namespace uba
 
 		ProcessHandle process = runProcess(processInfo);
 
-		if (!process.WaitForExit(GetTimeoutTime()))
-			return logger.Error(TC("UbaTestApp did not exit in 10 seconds"));
+		CHECK_TRUEF(process.WaitForExit(GetTimeoutTime()), TC("UbaTestApp did not exit in 10 seconds"));
 		u32 exitCode = process.GetExitCode();
-		if (exitCode != 0)
-			return logger.Error(TC("Got exit code %u"), exitCode);
-		if (!foundPingString)
-			return logger.Error(TC("Did not log string containing \"%s\""), GetSystemExpectedLogLine());
+		CHECK_TRUEF(!(exitCode != 0), TC("Got exit code %u"), exitCode);
+		CHECK_TRUEF(foundPingString, TC("Did not log string containing \"%s\""), GetSystemExpectedLogLine());
 		return true;
 	}
 
@@ -658,8 +612,7 @@ namespace uba
 				arg.Append(file);
 				processInfo.arguments = arg.data;
 				ProcessHandle process = runProcess(processInfo);
-				if (!process.WaitForExit(GetTimeoutTime()))
-					return logger.Error(TC("UbaTestApp did not exit in 10 seconds"));
+				CHECK_TRUEF(process.WaitForExit(GetTimeoutTime()), TC("UbaTestApp did not exit in 10 seconds"));
 				u32 exitCode =  process.GetExitCode();
 				return exitCode == 255 ? INVALID_FILE_ATTRIBUTES : exitCode;
 			};
@@ -742,11 +695,9 @@ namespace uba
 
 				for (auto& process : processes)
 				{
-					if (!process.WaitForExit(100000))
-						return logger.Error(TC("UbaTestApp did not exit in 10 seconds"));
+					CHECK_TRUEF(process.WaitForExit(100000), TC("UbaTestApp did not exit in 10 seconds"));
 					u32 exitCode = process.GetExitCode();
-					if (exitCode != 0)
-						return false;
+					CHECK_TRUE(exitCode == 0);
 				}
 
 				return true;
@@ -791,8 +742,7 @@ namespace uba
 
 		StringBuffer<> inFile(workingDir);
 		inFile.Append(TCV("VirtualFile.in"));
-		if (!session.CreateVirtualFile(inFile.data, "FOO", 3, testData.transient, testData.registerToDir))
-			return false;
+		CHECK_TRUE(session.CreateVirtualFile(inFile.data, "FOO", 3, testData.transient, testData.registerToDir));
 		ProcessHandle ph;
 
 
@@ -807,47 +757,33 @@ namespace uba
 				return testData.remote ? session.RunProcessRemote(pi) : session.RunProcess(pi, true, true);
 			};
 
-		if (!RunTestApp(logger, workingDir, runProcess2, TC("-virtualFile"), &ph))
-			return false;
+		CHECK_TRUE(RunTestApp(logger, workingDir, runProcess2, TC("-virtualFile"), &ph));
 
 		StringBuffer<> outFile(workingDir);
 		outFile.Append(TCV("VirtualFile.out"));
 
 		bool success = false;
 		ph.TraverseOutputFiles([&](StringView file) { success = file.Equals(outFile); });
-		if (!success)
-			return false;
+		CHECK_TRUE(success);
 
-		if (!session.DeleteVirtualFile(inFile.data, testData.registerToDir))
-			return false;
+		CHECK_TRUE(session.DeleteVirtualFile(inFile.data, testData.registerToDir));
 
-		if (FileExists(logger, outFile.data))
-			return false;
+		CHECK_TRUE(!FileExists(logger, outFile.data));
 		u64 outSize;
-		if (!session.GetOutputFileSize(outSize, outFile.data))
-			return false;
-		if (outSize != 3)
-			return false;
+		CHECK_TRUE(session.GetOutputFileSize(outSize, outFile.data));
+		CHECK_TRUE(outSize == 3);
 		u8 data[3];
-		if (!session.GetOutputFileData(data, outFile.data, false))
-			return false;
-		if (memcmp(data, "BAR", 3) != 0)
-			return false;
-		if (FileExists(logger, outFile.data))
-			return false;
-		if (!session.WriteOutputFile(outFile.data, true))
-			return false;
-		if (!FileExists(logger, outFile.data))
-			return false;
+		CHECK_TRUE(session.GetOutputFileData(data, outFile.data, false));
+		CHECK_TRUE(memcmp(data, "BAR", 3) == 0);
+		CHECK_TRUE(!FileExists(logger, outFile.data));
+		CHECK_TRUE(session.WriteOutputFile(outFile.data, true));
+		CHECK_TRUE(FileExists(logger, outFile.data));
 
 		if (ph.IsRemote())
-			if (!session.GetStorage().DeleteCasForFile(inFile.data))
-				return false;
+			CHECK_TRUE(session.GetStorage().DeleteCasForFile(inFile.data));
 
-		if (session.GetOutputFileSize(outSize, outFile.data))
-			return false;
-		if (session.GetOutputFileData(data, outFile.data, true))
-			return false;
+		CHECK_TRUE(!session.GetOutputFileSize(outSize, outFile.data));
+		CHECK_TRUE(!session.GetOutputFileData(data, outFile.data, true));
 		return true;
 	}
 
@@ -857,16 +793,13 @@ namespace uba
 		config.AddTable(TC("Session")).AddValue(TC("ShouldWriteToDisk"), false);
 
 		RunCreateVirtualFileTestData testData { false };
-		if (!RunLocal(logger, testRootDir, config, RunCreateVirtualFileTest, uintptr_t(&testData)))
-			return false;
+		CHECK_TRUE(RunLocal(logger, testRootDir, config, RunCreateVirtualFileTest, uintptr_t(&testData)));
 
 		testData.transient = true;
-		if (!RunLocal(logger, testRootDir, config, RunCreateVirtualFileTest, uintptr_t(&testData)))
-			return false;
+		CHECK_TRUE(RunLocal(logger, testRootDir, config, RunCreateVirtualFileTest, uintptr_t(&testData)));
 
 		testData.registerToDir = false;
-		if (!RunLocal(logger, testRootDir, config, RunCreateVirtualFileTest, uintptr_t(&testData)))
-			return false;
+		CHECK_TRUE(RunLocal(logger, testRootDir, config, RunCreateVirtualFileTest, uintptr_t(&testData)));
 		return true;
 	}
 
@@ -879,46 +812,37 @@ namespace uba
 		RunCreateVirtualFileTestData testData { true };
 
 		testData.transient = true;
-		if (!RunRemote(logger, testRootDir, config, RunCreateVirtualFileTest, uintptr_t(&testData)))
-			return false;
+		CHECK_TRUE(RunRemote(logger, testRootDir, config, RunCreateVirtualFileTest, uintptr_t(&testData)));
 			
 		testData.transient = false;
-		if (!RunRemote(logger, testRootDir, config, RunCreateVirtualFileTest, uintptr_t(&testData)))
-			return false;
+		CHECK_TRUE(RunRemote(logger, testRootDir, config, RunCreateVirtualFileTest, uintptr_t(&testData)));
 
 		testData.registerToDir = false;
-		if (!RunRemote(logger, testRootDir, config, RunCreateVirtualFileTest, uintptr_t(&testData)))
-			return false;
+		CHECK_TRUE(RunRemote(logger, testRootDir, config, RunCreateVirtualFileTest, uintptr_t(&testData)));
 		return true;
 	}
 
 	bool RunRegisterVirtualFileTest(LoggerWithWriter& logger, SessionServer& session, const tchar* workingDir, const RunProcessFunction& runProcess, uintptr_t extraData)
 	{
 		StringBuffer<> bigFile;
-		if (!CreateTestFile(bigFile, logger, ToView(workingDir), TCV("BigFile.txt"), TCV("FooBar")))
-			return false;
-		if (!session.RegisterVirtualFile(StringBuffer(workingDir).Append(TCV("First.txt")).data, bigFile.data, 0, sizeof(tchar)*3))
-			return false;
-		if (!session.RegisterVirtualFile(StringBuffer(workingDir).Append(TCV("Second.txt")).data, bigFile.data, sizeof(tchar) * 3, sizeof(tchar) * 3))
-			return false;
-		if (!RunTestApp(logger, workingDir, runProcess, TC("-testRegisterVirtualFile")))
-			return false;
+		CHECK_TRUE(CreateTestFile(bigFile, logger, ToView(workingDir), TCV("BigFile.txt"), TCV("FooBar")));
+		CHECK_TRUE(session.RegisterVirtualFile(StringBuffer(workingDir).Append(TCV("First.txt")).data, bigFile.data, 0, sizeof(tchar)*3));
+		CHECK_TRUE(session.RegisterVirtualFile(StringBuffer(workingDir).Append(TCV("Second.txt")).data, bigFile.data, sizeof(tchar) * 3, sizeof(tchar) * 3));
+		CHECK_TRUE(RunTestApp(logger, workingDir, runProcess, TC("-testRegisterVirtualFile")));
 		return true;
 	}
 
 	bool TestRegisterVirtualFile(LoggerWithWriter& logger, const StringBufferBase& testRootDir)
 	{
 		Config config;
-		if (!RunLocal(logger, testRootDir, config, RunRegisterVirtualFileTest))
-			return false;
+		CHECK_TRUE(RunLocal(logger, testRootDir, config, RunRegisterVirtualFileTest));
 		return true;
 	}
 
 	bool TestRemoteRegisterVirtualFile(LoggerWithWriter& logger, const StringBufferBase& testRootDir)
 	{
 		Config config;
-		if (!RunRemote(logger, testRootDir, config, RunRegisterVirtualFileTest))
-			return false;
+		CHECK_TRUE(RunRemote(logger, testRootDir, config, RunRegisterVirtualFileTest));
 		return true;
 	}
 
@@ -951,26 +875,22 @@ namespace uba
 			[&](LoggerWithWriter& logger, const StringView& workingDir, SessionServer& sessionServer)
 			{
 				auto runProcess = [&](ProcessStartInfo& pi) { return sessionServer.RunProcessRemote(pi); };
-				if (!CreateTestFile(logger, workingDir, TCV("SpecialFile1"), TCV("0")))
-					return false;
+				CHECK_TRUE(CreateTestFile(logger, workingDir, TCV("SpecialFile1"), TCV("0")));
 
 				const StringView casFile(TCV("UbaClient1357/cas/4d/4d067153ac729a4a7e8220c97935ffba67487800"));
 
 				if (!SetupClientSession(logger, testRootDir, true, false, sessionServer.GetServer(), 1357, [&](LoggerWithWriter& logger, SessionClient& sessionClient)
 					{
-						if (!CreateTestFile(logger, testRootDir, casFile, TCV("0")))
-							return false;
+						CHECK_TRUE(CreateTestFile(logger, testRootDir, casFile, TCV("0")));
 						return RunTestApp(logger, workingDir.data, runProcess, TC("-readwrite=0"));
 					}))
 					return false;
 
-				if (!DeleteTestFile(logger, testRootDir, casFile))
-					return false;
+				CHECK_TRUE(DeleteTestFile(logger, testRootDir, casFile));
 
 				if (!SetupClientSession(logger, testRootDir, false, false, sessionServer.GetServer(), 1357, [&](LoggerWithWriter& logger, SessionClient& sessionClient)
 					{
-						if (!RunTestApp(logger, workingDir.data, runProcess, TC("-readwrite=1")))
-							return false;
+						CHECK_TRUE(RunTestApp(logger, workingDir.data, runProcess, TC("-readwrite=1")));
 						return true;
 					}))
 					return false;
diff --git a/Engine/Source/Programs/UnrealBuildAccelerator/Test/Private/UbaTestStdOut.cpp b/Engine/Source/Programs/UnrealBuildAccelerator/Test/Private/UbaTestStdOut.cpp
--- a/Engine/Source/Programs/UnrealBuildAccelerator/Test/Private/UbaTestStdOut.cpp
+++ b/Engine/Source/Programs/UnrealBuildAccelerator/Test/Private/UbaTestStdOut.cpp
@@ -1,8 +1,10 @@
 // Copyright Epic Games, Inc. All Rights Reserved.
 
 #include "UbaNetworkBackendTcp.h"
+#include "UbaProcessStartInfo.h"
 #include "UbaSessionServer.h"
 #include "UbaStorageServer.h"
+#include "UbaTest.h"
 
 namespace uba
 {
@@ -18,8 +20,7 @@ namespace uba
 
 		StringBuffer<> rootDir;
 		rootDir.Append(testRootDir).Append(TCV("Uba"));
-		if (!DeleteAllFiles(logger, rootDir.data))
-			return false;
+		CHECK_TRUE(DeleteAllFiles(logger, rootDir.data));
 
 		StorageServerCreateInfo storageServerInfo(networkServer, rootDir.data, logWriter);
 		storageServerInfo.casCapacityBytes = 1024ull * 1024 * 1024;
@@ -35,16 +36,12 @@ namespace uba
 
 		StringBuffer<> workingDir;
 		workingDir.Append(testRootDir).Append(TCV("WorkingDir"));
-		if (!DeleteAllFiles(logger, workingDir.data))
-			return false;
-		if (!storageServer.CreateDirectory(workingDir.data))
-			return false;
-		if (!DeleteAllFiles(logger, workingDir.data, false))
-			return false;
+		CHECK_TRUE(DeleteAllFiles(logger, workingDir.data));
+		CHECK_TRUE(storageServer.CreateDirectory(workingDir.data));
+		CHECK_TRUE(DeleteAllFiles(logger, workingDir.data, false));
 		workingDir.EnsureEndsWithSlash();
 
-		if (!networkServer.StartListen(networkBackend))
-			return logger.Error(TC("Failed to listen"));
+		CHECK_TRUEF(networkServer.StartListen(networkBackend), TC("Failed to listen"));
 
 		StringBuffer<> testApp;
 		if (!app)
@@ -66,8 +63,7 @@ namespace uba
 		pi.logFile = TC("Log");
 
 		auto ph = remote ? sessionServer.RunProcessRemote(pi) : sessionServer.RunProcess(pi);
-		if (!ph.WaitForExit(5000))
-			return logger.Error(TC("Timed out waiting for process"));
+		CHECK_TRUEF(ph.WaitForExit(5000), TC("Timed out waiting for process"));
 		if (auto ec = ph.GetExitCode())
 			return logger.Error(TC("Process exited with error code %u"), ec);
 
@@ -76,10 +72,8 @@ namespace uba
 		sessionServer.WaitOnAllTasks();
 
 		auto& logLines = ph.GetLogLines();
-		if (logLines.size() != 1)
-			return logger.Error(TC("Application %s produced %u log line(s) but expected 1"), app, u32(logLines.size()));
-		if (!Equals(logLines[0].text.c_str(), expectedOut))
-			return logger.Error(TC("Application %s produced non-matching log line: %s"), app, logLines[0].text.c_str());
+		CHECK_TRUEF(!(logLines.size() != 1), TC("Application %s produced %u log line(s) but expected 1"), app, u32(logLines.size()));
+		CHECK_TRUEF(Equals(logLines[0].text.c_str(), expectedOut), TC("Application %s produced non-matching log line: %s"), app, logLines[0].text.c_str());
 
 		return true;
 	}
diff --git a/Engine/Source/Programs/UnrealBuildAccelerator/Test/Private/UbaTestStorage.cpp b/Engine/Source/Programs/UnrealBuildAccelerator/Test/Private/UbaTestStorage.cpp
--- a/Engine/Source/Programs/UnrealBuildAccelerator/Test/Private/UbaTestStorage.cpp
+++ b/Engine/Source/Programs/UnrealBuildAccelerator/Test/Private/UbaTestStorage.cpp
@@ -30,37 +30,29 @@ namespace uba
 
 		CasKey key;
 		bool deferCreation = false;
-		if (!storage.StoreCasFile(key, detoursLib.data, CasKeyZero, deferCreation))
-			return logger.Error(TC("Failed to store file %s"), detoursLib.data);
-		if (key == CasKeyZero)
-			return logger.Error(TC("Failed to find file %s"), detoursLib.data);
+		CHECK_TRUEF(storage.StoreCasFile(key, detoursLib.data, CasKeyZero, deferCreation), TC("Failed to store file %s"), detoursLib.data);
+		CHECK_TRUEF(!(key == CasKeyZero), TC("Failed to find file %s"), detoursLib.data);
 
 		StringBuffer<> detoursLibCopy(detoursLib);
 		detoursLibCopy.Append(TCV(".tmp"));
 
 		auto deleteFile = MakeGuard([&]() { return DeleteFileW(detoursLibCopy.data); });
 
-		if (!storage.CopyOrLink(key, detoursLibCopy.data, DefaultAttributes()))
-			return logger.Error(TC("Failed to copy cas to file %s"), detoursLibCopy.data);
+		CHECK_TRUEF(storage.CopyOrLink(key, detoursLibCopy.data, DefaultAttributes()), TC("Failed to copy cas to file %s"), detoursLibCopy.data);
 
 		FileHandle original;
-		if (!OpenFileSequentialRead(logger, detoursLib.data, original))
-			return logger.Error(TC("Failed to open %s for read"), detoursLib.data);
+		CHECK_TRUEF(OpenFileSequentialRead(logger, detoursLib.data, original), TC("Failed to open %s for read"), detoursLib.data);
 		auto closeOriginal = MakeGuard([&]() { return CloseFile(detoursLib.data, original); });
 
 		FileHandle copy;
-		if (!OpenFileSequentialRead(logger, detoursLibCopy.data, copy))
-			return logger.Error(TC("Failed to open %s for read"), detoursLibCopy.data);
+		CHECK_TRUEF(OpenFileSequentialRead(logger, detoursLibCopy.data, copy), TC("Failed to open %s for read"), detoursLibCopy.data);
 		auto closeCopy = MakeGuard([&]() { return CloseFile(detoursLibCopy.data, copy); });
 
 		u64 originalSize;
-		if (!GetFileSizeEx(originalSize, original))
-			return logger.Error(TC("Failed to get size of %s"), detoursLib.data);
+		CHECK_TRUEF(GetFileSizeEx(originalSize, original), TC("Failed to get size of %s"), detoursLib.data);
 		u64 copySize;
-		if (!GetFileSizeEx(copySize, copy))
-			return logger.Error(TC("Failed to get size of %s"), detoursLibCopy.data);
-		if (originalSize != copySize)
-			return logger.Error(TC("Size mismatch between %s and %s (%llu vs %llu)"), detoursLib.data, detoursLibCopy.data, originalSize, copySize);
+		CHECK_TRUEF(GetFileSizeEx(copySize, copy), TC("Failed to get size of %s"), detoursLibCopy.data);
+		CHECK_TRUEF(!(originalSize != copySize), TC("Size mismatch between %s and %s (%llu vs %llu)"), detoursLib.data, detoursLibCopy.data, originalSize, copySize);
 
 		u8 originalBuffer[PageSize];
 		u8 copyBuffer[PageSize];
@@ -68,24 +60,18 @@ namespace uba
 		while (left)
 		{
 			u64 toRead = Min(left, (u64)sizeof(originalBuffer));
-			if (!ReadFile(logger, detoursLib.data, original, originalBuffer, toRead))
-				return logger.Error(TC("Failed to read %u from %s"), detoursLib.data);
-			if (!ReadFile(logger, detoursLibCopy.data, copy, copyBuffer, toRead))
-				return logger.Error(TC("Failed to read %u from %s"), detoursLibCopy.data);
-			if (memcmp(originalBuffer, copyBuffer, toRead) != 0)
-				return logger.Error(TC("Data mismatch between %s and %s"), detoursLib.data, detoursLibCopy.data);
+			CHECK_TRUEF(ReadFile(logger, detoursLib.data, original, originalBuffer, toRead), TC("Failed to read %u from %s"), detoursLib.data);
+			CHECK_TRUEF(ReadFile(logger, detoursLibCopy.data, copy, copyBuffer, toRead), TC("Failed to read %u from %s"), detoursLibCopy.data);
+			CHECK_TRUEF(!(memcmp(originalBuffer, copyBuffer, toRead) != 0), TC("Data mismatch between %s and %s"), detoursLib.data, detoursLibCopy.data);
 
 			left -= toRead;
 		}
 
-		if (!closeOriginal.Execute())
-			return logger.Error(TC("Failed to close %s"), detoursLib.data);
+		CHECK_TRUEF(closeOriginal.Execute(), TC("Failed to close %s"), detoursLib.data);
 
-		if (!closeCopy.Execute())
-			return logger.Error(TC("Failed to close %s"), detoursLibCopy.data);
+		CHECK_TRUEF(closeCopy.Execute(), TC("Failed to close %s"), detoursLibCopy.data);
 
-		if (!deleteFile.Execute())
-			return logger.Error(TC("Failed to delete %s"), detoursLibCopy.data);
+		CHECK_TRUEF(deleteFile.Execute(), TC("Failed to delete %s"), detoursLibCopy.data);
 
 		return true;
 	}
@@ -103,8 +89,7 @@ namespace uba
 
 		StringBuffer<> rootDir;
 		rootDir.Append(testRootDir).Append(TCV("Uba"));
-		if (!DeleteAllFiles(logger, rootDir.data))
-			return false;
+		CHECK_TRUE(DeleteAllFiles(logger, rootDir.data));
 
 		StorageServerCreateInfo storageServerInfo(server, rootDir.data, logger.m_writer);
 		storageServerInfo.casCapacityBytes = 1024ull * 1024 * 1024;
@@ -113,36 +98,29 @@ namespace uba
 		auto g = MakeGuard([&]() { server.DisconnectClients(); });
 
 		rootDir.Append(TCV("Client"));
-		if (!DeleteAllFiles(logger, rootDir.data))
-			return false;
+		CHECK_TRUE(DeleteAllFiles(logger, rootDir.data));
 
 		StorageClientCreateInfo storageClientInfo(client, rootDir.data);
 		StorageClient storageClient(storageClientInfo);
 		storageClient.Start();
 
-		if (!storageClient.LoadCasTable(true))
-			return false;
+		CHECK_TRUE(storageClient.LoadCasTable(true));
 
 		rootDir.EnsureEndsWithSlash();
 
-		if (!server.StartListen(serverTcp, 1234))
-			return logger.Error(TC("Failed to listen"));
+		CHECK_TRUEF(server.StartListen(serverTcp, 1234), TC("Failed to listen"));
 		Sleep(100);
-		if (!client.Connect(clientTcp, TC("127.0.0.1"), 1234))
-			return logger.Error(TC("Failed to connect"));
+		CHECK_TRUEF(client.Connect(clientTcp, TC("127.0.0.1"), 1234), TC("Failed to connect"));
 		auto cg = MakeGuard([&]() { client.Disconnect(); });
 
 		StringBuffer<> fileName;
 		{
 			fileName.Append(rootDir).Append(TCV("UbaTestFile"));
 			FileAccessor fileHandle(logger, fileName.data);
-			if (!fileHandle.CreateWrite())
-				return logger.Error(TC("Failed to create file for write"));
+			CHECK_TRUEF(fileHandle.CreateWrite(), TC("Failed to create file for write"));
 			u8 byte = 'H';
-			if (!fileHandle.Write(&byte, 1))
-				return false;
-			if (!fileHandle.Close())
-				return false;
+			CHECK_TRUE(fileHandle.Write(&byte, 1));
+			CHECK_TRUE(fileHandle.Close());
 		}
 
 		CasKey key;
@@ -151,8 +129,7 @@ namespace uba
 			return logger.Error(TC("Failed to store file %s"), fileName.data);
 
 		fileName.Clear().Append(testRootDir).Append(TCV("Uba")).Append(PathSeparator).Append(TCV("UbaTestFile"));
-		if (!storageServer.CopyOrLink(key, fileName.data, DefaultAttributes()))
-			return logger.Error(TC("Failed to copy cas to file %s"), fileName.data);
+		CHECK_TRUEF(storageServer.CopyOrLink(key, fileName.data, DefaultAttributes()), TC("Failed to copy cas to file %s"), fileName.data);
 		return true;
 	}
 
@@ -169,8 +146,7 @@ namespace uba
 
 		StringBuffer<> rootDir;
 		rootDir.Append(testRootDir).Append(TCV("Uba"));
-		if (!DeleteAllFiles(logger, rootDir.data))
-			return false;
+		CHECK_TRUE(DeleteAllFiles(logger, rootDir.data));
 
 		StorageServerCreateInfo storageServerInfo(server, rootDir.data, logger.m_writer);
 		storageServerInfo.casCapacityBytes = 1024ull * 1024 * 1024;
@@ -180,23 +156,19 @@ namespace uba
 		auto g = MakeGuard([&]() { server.DisconnectClients(); });
 
 		rootDir.Append(TCV("Client"));
-		if (!DeleteAllFiles(logger, rootDir.data))
-			return false;
+		CHECK_TRUE(DeleteAllFiles(logger, rootDir.data));
 
 		StorageClientCreateInfo storageClientInfo(client, rootDir.data);
 		StorageClient storageClient(storageClientInfo);
 		storageClient.Start();
 
-		if (!storageClient.LoadCasTable(true))
-			return false;
+		CHECK_TRUE(storageClient.LoadCasTable(true));
 
 		rootDir.EnsureEndsWithSlash();
 
-		if (!server.StartListen(serverTcp, 1234))
-			return logger.Error(TC("Failed to listen"));
+		CHECK_TRUEF(server.StartListen(serverTcp, 1234), TC("Failed to listen"));
 		Sleep(100);
-		if (!client.Connect(clientTcp, TC("127.0.0.1"), 1234))
-			return logger.Error(TC("Failed to connect"));
+		CHECK_TRUEF(client.Connect(clientTcp, TC("127.0.0.1"), 1234), TC("Failed to connect"));
 		auto cg = MakeGuard([&]() { client.Disconnect(); });
 
 		StringBuffer<> fileName;
@@ -204,13 +176,10 @@ namespace uba
 		{
 			fileName.Append(rootDir).Append(TCV("UbaTestFile"));
 			FileAccessor fileHandle(logger, fileName.data);
-			if (!fileHandle.CreateWrite())
-				return logger.Error(TC("Failed to create file for write"));
+			CHECK_TRUEF(fileHandle.CreateWrite(), TC("Failed to create file for write"));
 			u8 byte = 'H';
-			if (!fileHandle.Write(&byte, 1))
-				return false;
-			if (!fileHandle.Close())
-				return false;
+			CHECK_TRUE(fileHandle.Write(&byte, 1));
+			CHECK_TRUE(fileHandle.Close());
 		}
 #else
 		fileName.Append(TCV("e:\\dev\\fn\\QAGame\\Saved\\StagedBuilds\\PS5_Temp\\Split\\qagame\\content\\paks\\qagame-ps5.ucas\\.copied"));
@@ -244,8 +213,7 @@ namespace uba
 
 			storageServer.RegisterExternalFileMappingsProvider([&](Storage::ExternalFileMapping& out, StringKey fileNameKey, const tchar* fileName)
 				{
-					if (fileNameKey != testFileNameKey)
-						return false;
+					CHECK_TRUE(fileNameKey == testFileNameKey);
 					out.handle = sharedMem.DuplicateHandle(memHandle, TC(""));
 					out.size = 1;
 					out.dropCasAfterUse = true;
@@ -255,13 +223,11 @@ namespace uba
 
 			storageClient.DropCasFile(casKey, true, fileName.data, false);
 			storageServer.DropCasFile(casKey, true, fileName.data, false);
-			if (!storageClient.RetrieveCasFile(result, casKey, testFileNameKey, TC("KnownInput")))
-				return logger.Error(TC("Failed to retrieve file %s"), fileName.data);
+			CHECK_TRUEF(storageClient.RetrieveCasFile(result, casKey, testFileNameKey, TC("KnownInput")), TC("Failed to retrieve file %s"), fileName.data);
 
 			storageClient.DropCasFile(casKey, true, fileName.data, false);
 			storageServer.DropCasFile(casKey, true, fileName.data, false);
-			if (!storageClient.RetrieveCasFile(result, casKey, testFileNameKey, TC("KnownInput")))
-				return logger.Error(TC("Failed to retrieve file %s"), fileName.data);
+			CHECK_TRUEF(storageClient.RetrieveCasFile(result, casKey, testFileNameKey, TC("KnownInput")), TC("Failed to retrieve file %s"), fileName.data);
 		}
 
 		return true;
@@ -280,8 +246,7 @@ namespace uba
 
 		StringBuffer<> rootDir;
 		rootDir.Append(testRootDir).Append(TCV("Uba"));
-		if (!DeleteAllFiles(logger, rootDir.data))
-			return false;
+		CHECK_TRUE(DeleteAllFiles(logger, rootDir.data));
 
 		StorageServerCreateInfo storageServerInfo(server, rootDir.data, logger.m_writer);
 		storageServerInfo.casCapacityBytes = 1024ull * 1024 * 1024;
@@ -291,36 +256,29 @@ namespace uba
 		auto g = MakeGuard([&]() { server.DisconnectClients(); });
 
 		rootDir.Append(TCV("Client"));
-		if (!DeleteAllFiles(logger, rootDir.data))
-			return false;
+		CHECK_TRUE(DeleteAllFiles(logger, rootDir.data));
 
 		StorageClientCreateInfo storageClientInfo(client, rootDir.data);
 		StorageClient storageClient(storageClientInfo);
 		storageClient.Start();
 
-		if (!storageClient.LoadCasTable(true))
-			return false;
+		CHECK_TRUE(storageClient.LoadCasTable(true));
 
 		rootDir.EnsureEndsWithSlash();
 
-		if (!server.StartListen(serverTcp, 1234))
-			return logger.Error(TC("Failed to listen"));
+		CHECK_TRUEF(server.StartListen(serverTcp, 1234), TC("Failed to listen"));
 		Sleep(100);
-		if (!client.Connect(clientTcp, TC("127.0.0.1"), 1234))
-			return logger.Error(TC("Failed to connect"));
+		CHECK_TRUEF(client.Connect(clientTcp, TC("127.0.0.1"), 1234), TC("Failed to connect"));
 		auto cg = MakeGuard([&]() { client.Disconnect(); });
 
 		StringBuffer<> sourceFileName;
 		{
 			sourceFileName.Append(rootDir).Append(TCV("UbaTestFile"));
 			FileAccessor fileHandle(logger, sourceFileName.data);
-			if (!fileHandle.CreateWrite())
-				return logger.Error(TC("Failed to create file for write"));
+			CHECK_TRUEF(fileHandle.CreateWrite(), TC("Failed to create file for write"));
 			u8 byte = 'H';
-			if (!fileHandle.Write(&byte, 1))
-				return false;
-			if (!fileHandle.Close())
-				return false;
+			CHECK_TRUE(fileHandle.Write(&byte, 1));
+			CHECK_TRUE(fileHandle.Close());
 		}
 
 		StringBuffer<> destFileName;
@@ -336,39 +294,27 @@ namespace uba
 		auto copy = [&]()
 			{
 				destFileName.Clear().Append(testRootDir).Append(TCV("Uba")).Append(PathSeparator).Append(TCV("UbaTestFile"));
-				if (!storageServer.WaitForCasFile(key, destFileName.data))
-					return false;
+				CHECK_TRUE(storageServer.WaitForCasFile(key, destFileName.data));
 				if (!storageServer.CopyOrLink(key, destFileName.data, DefaultAttributes(), false, 1, {}, false, true))
 					return logger.Error(TC("Failed to copy cas to file %s"), destFileName.data);
-				if (!FilesEqual(logger, sourceFileName.data, destFileName.data))
-					return logger.Error(TC("Files %s and %s not equal"), sourceFileName.data, destFileName.data);
+				CHECK_TRUEF(FilesEqual(logger, sourceFileName.data, destFileName.data), TC("Files %s and %s not equal"), sourceFileName.data, destFileName.data);
 				return true;
 			};
 
-		if (!store())
-			return false;
-		if (!copy())
-			return false;
-		if (!store())
-			return false;
-		if (!copy())
-			return false;
+		CHECK_TRUE(store());
+		CHECK_TRUE(copy());
+		CHECK_TRUE(store());
+		CHECK_TRUE(copy());
 
-		if (!store())
-			return false;
-		if (!store())
-			return false;
-		if (!copy())
-			return false;
-		if (!copy())
-			return false;
+		CHECK_TRUE(store());
+		CHECK_TRUE(store());
+		CHECK_TRUE(copy());
+		CHECK_TRUE(copy());
 
-		if (!storageServer.StoreCasFile(key, sourceFileName.data, {}, true))
-			return false;
+		CHECK_TRUE(storageServer.StoreCasFile(key, sourceFileName.data, {}, true));
 
 		//client.Fetch
-		if (!store())
-			return false;
+		CHECK_TRUE(store());
 		//if (!copy())
 		//	return false;
 
@@ -376,8 +322,7 @@ namespace uba
 		{
 			sourceFileName.Clear().Append(rootDir).Append(TCV("UbaTestFile2"));
 			FileAccessor fileHandle(logger, sourceFileName.data);
-			if (!fileHandle.CreateWrite())
-				return logger.Error(TC("Failed to create file for write"));
+			CHECK_TRUEF(fileHandle.CreateWrite(), TC("Failed to create file for write"));
 			
 			// 256*64*1024 = 16mb file
 			u32 data[16*1024];
@@ -385,17 +330,13 @@ namespace uba
 			{
 				for (u32 j=0;j!=sizeof_array(data); ++j)
 					data[j] = u32(rand());
-				if (!fileHandle.Write(data, sizeof(data)))
-					return false;
+				CHECK_TRUE(fileHandle.Write(data, sizeof(data)));
 			}
-			if (!fileHandle.Close())
-				return false;
+			CHECK_TRUE(fileHandle.Close());
 		}
 
-		if (!store())
-			return false;
-		if (!copy())
-			return false;
+		CHECK_TRUE(store());
+		CHECK_TRUE(copy());
 
 		return true;
 	}
diff --git a/Engine/Source/Programs/UnrealBuildAccelerator/Test/Private/UbaTestStress.cpp b/Engine/Source/Programs/UnrealBuildAccelerator/Test/Private/UbaTestStress.cpp
--- a/Engine/Source/Programs/UnrealBuildAccelerator/Test/Private/UbaTestStress.cpp
+++ b/Engine/Source/Programs/UnrealBuildAccelerator/Test/Private/UbaTestStress.cpp
@@ -6,6 +6,7 @@
 #include "UbaScheduler.h"
 #include "UbaSessionServer.h"
 #include "UbaStorageServer.h"
+#include "UbaTest.h"
 
 namespace uba
 {
@@ -21,8 +22,7 @@ namespace uba
 
 		StringBuffer<> rootDir;
 		rootDir.Append(testRootDir).Append(TCV("Uba"));
-		if (!DeleteAllFiles(logger, rootDir.data))
-			return false;
+		CHECK_TRUE(DeleteAllFiles(logger, rootDir.data));
 
 		StorageServerCreateInfo storageServerInfo(networkServer, rootDir.data, logWriter);
 		storageServerInfo.casCapacityBytes = 1024ull * 1024 * 1024;
@@ -38,12 +38,9 @@ namespace uba
 
 		StringBuffer<> workingDir;
 		workingDir.Append(testRootDir).Append(TCV("WorkingDir"));
-		if (!DeleteAllFiles(logger, workingDir.data))
-			return false;
-		if (!storageServer.CreateDirectory(workingDir.data))
-			return false;
-		if (!DeleteAllFiles(logger, workingDir.data, false))
-			return false;
+		CHECK_TRUE(DeleteAllFiles(logger, workingDir.data));
+		CHECK_TRUE(storageServer.CreateDirectory(workingDir.data));
+		CHECK_TRUE(DeleteAllFiles(logger, workingDir.data, false));
 		workingDir.EnsureEndsWithSlash();
 
 		SchedulerCreateInfo info(sessionServer);
@@ -53,8 +50,7 @@ namespace uba
 		Scheduler scheduler(info);
 		scheduler.Start();
 
-		if (!networkServer.StartListen(networkBackend))
-			return logger.Error(TC("Failed to listen"));
+		CHECK_TRUEF(networkServer.StartListen(networkBackend), TC("Failed to listen"));
 
 		CoordinatorWrapper coordinator;
 
@@ -66,8 +62,7 @@ namespace uba
 			StringBuffer<512> coordinatorWorkDir(testRootDir);
 			coordinatorWorkDir.EnsureEndsWithSlash().Append(TCV("Horde"));
 			StringBuffer<512> binariesDir;
-			if (!GetDirectoryOfCurrentModule(logger, binariesDir))
-				return false;
+			CHECK_TRUE(GetDirectoryOfCurrentModule(logger, binariesDir));
 
 			CoordinatorCreateInfo cinfo;
 			cinfo.workDir = coordinatorWorkDir.data;
@@ -78,8 +73,7 @@ namespace uba
 
 			cinfo.maxCoreCount = 400;
 			cinfo.logging = true;
-			if (!coordinator.Create(logger, TC("Horde"), cinfo, networkBackend, networkServer))
-				return false;
+			CHECK_TRUE(coordinator.Create(logger, TC("Horde"), cinfo, networkBackend, networkServer));
 		}
 		auto cg = MakeGuard([&]() { coordinator.Destroy(); });
 
