[UBA]
* Re-added __try/__except around memcpy in NtWriteFile since that is usually how oom manifests on wine and we want to handle it more cleanly

diff --git a/Engine/Source/Programs/UnrealBuildAccelerator/Detours/Private/Windows/UbaDetoursFunctionsKernelBase.inl b/Engine/Source/Programs/UnrealBuildAccelerator/Detours/Private/Windows/UbaDetoursFunctionsKernelBase.inl
--- a/Engine/Source/Programs/UnrealBuildAccelerator/Detours/Private/Windows/UbaDetoursFunctionsKernelBase.inl
+++ b/Engine/Source/Programs/UnrealBuildAccelerator/Detours/Private/Windows/UbaDetoursFunctionsKernelBase.inl
@@ -611,15 +611,6 @@ HANDLE Detoured_CreateToolhelp32Snapshot(DWORD dwFlags, DWORD th32ProcessID)
 	//return True_CreateToolhelp32Snapshot(dwFlags, th32ProcessID);
 }
 
-int WriteFileAccessViolation(PEXCEPTION_POINTERS ep)
-{
-	// This likely is OOM in the file mapping system on wine (tmpfs running out of space or something).
-	// We need to improve how uba tracks memory usage on wine to catch these before it is too late. but for now, just terminate
-	if (g_isRunningWine && ep->ExceptionRecord->ExceptionCode == EXCEPTION_ACCESS_VIOLATION)
-		TerminateProcess(GetCurrentProcess(), EmergencyShutdownExitCode);
-	return EXCEPTION_CONTINUE_SEARCH;
-}
-
 BOOL Detoured_FlushFileBuffers(HANDLE hFile)
 {
 	DETOURED_CALL(FlushFileBuffers);
@@ -4454,13 +4445,13 @@ HANDLE Detoured_OpenProcess(DWORD dwDesiredAccess, BOOL bInheritHandle, DWORD dw
 PVOID Detoured_AddVectoredExceptionHandler(ULONG First, PVECTORED_EXCEPTION_HANDLER Handler)
 {
 	DEBUG_LOG_TRUE(L"AddVectoredExceptionHandler", L"");
-	return nullptr;
+	return True_AddVectoredExceptionHandler(First, Handler);
 }
 
 ULONG Detoured_RemoveVectoredExceptionHandler(PVOID Handle)
 {
 	DEBUG_LOG_TRUE(L"RemoveVectoredExceptionHandler", L"");
-	return 0;
+	return True_RemoveVectoredExceptionHandler(Handle);
 }
 
 VOID Detoured_RaiseFailFastException(PEXCEPTION_RECORD pExceptionRecord, PCONTEXT pContextRecord, DWORD dwFlags)
diff --git a/Engine/Source/Programs/UnrealBuildAccelerator/Detours/Private/Windows/UbaDetoursFunctionsNtDll.inl b/Engine/Source/Programs/UnrealBuildAccelerator/Detours/Private/Windows/UbaDetoursFunctionsNtDll.inl
--- a/Engine/Source/Programs/UnrealBuildAccelerator/Detours/Private/Windows/UbaDetoursFunctionsNtDll.inl
+++ b/Engine/Source/Programs/UnrealBuildAccelerator/Detours/Private/Windows/UbaDetoursFunctionsNtDll.inl
@@ -2154,6 +2154,15 @@ void WriteStdFile(LPCVOID lpBuffer, DWORD nNumberOfBytesToWrite, bool isError)
 		g_stdFile.Appendf(L"%.*hs", left, bufferStr + start);
 }
 
+int WriteFileAccessViolation(PEXCEPTION_POINTERS ep)
+{
+	// This likely is OOM in the file mapping system on wine (tmpfs running out of space or something).
+	// We need to improve how uba tracks memory usage on wine to catch these before it is too late. but for now, just terminate
+	if (g_isRunningWine && ep->ExceptionRecord->ExceptionCode == EXCEPTION_ACCESS_VIOLATION)
+		TerminateProcess(GetCurrentProcess(), EmergencyShutdownExitCode);
+	return EXCEPTION_CONTINUE_SEARCH;
+}
+
 NTSTATUS NTAPI Detoured_NtWriteFile(HANDLE FileHandle, HANDLE Event, PIO_APC_ROUTINE ApcRoutine, PVOID ApcContext, PIO_STATUS_BLOCK IoStatusBlock, PVOID Buffer, ULONG Length, PLARGE_INTEGER ByteOffset, PULONG Key)
 {
 	DETOURED_CALL(NtWriteFile);
@@ -2237,7 +2246,15 @@ NTSTATUS NTAPI Detoured_NtWriteFile(HANDLE FileHandle, HANDLE Event, PIO_APC_ROU
 	mf.writtenSize = newSize;
 	mf.isReported = false;
 	lock.Leave();
-	memcpy(mf.baseAddress + offset, Buffer, Length);
+
+	__try
+	{
+		memcpy(mf.baseAddress + offset, Buffer, Length);
+	}
+	__except(WriteFileAccessViolation(GetExceptionInformation()))
+	{
+	}
+
 	return successRun(TC("(Memory)"));
 }
 
