[UBA]
* Fixed so UbaWine.dll.so exports SocketRecv2.

diff --git a/Engine/Source/Programs/UnrealBuildAccelerator/Wine/UbaWineWindows.cpp b/Engine/Source/Programs/UnrealBuildAccelerator/Wine/UbaWineWindows.cpp
--- a/Engine/Source/Programs/UnrealBuildAccelerator/Wine/UbaWineWindows.cpp
+++ b/Engine/Source/Programs/UnrealBuildAccelerator/Wine/UbaWineWindows.cpp
@@ -130,6 +130,10 @@ extern "C" int WINAPI SocketRecv(uba::Socket s, void* data, int len)
 {
 	return uba::SocketRecv(s, data, len);
 }
+extern "C" int WINAPI SocketRecv2(uba::Socket s, uba::RecvBuf* bufs, int bufCount)
+{
+	return uba::SocketRecv2(s, bufs, bufCount);
+}
 extern "C" int WINAPI SocketSend(uba::Socket s, uba::SendBuf* bufs, int bufCount)
 {
 	return uba::SocketSend(s, bufs, bufCount);
diff --git a/Engine/Source/Programs/UnrealBuildAccelerator/Wine/UbaWineWindows.def b/Engine/Source/Programs/UnrealBuildAccelerator/Wine/UbaWineWindows.def
--- a/Engine/Source/Programs/UnrealBuildAccelerator/Wine/UbaWineWindows.def
+++ b/Engine/Source/Programs/UnrealBuildAccelerator/Wine/UbaWineWindows.def
@@ -21,6 +21,7 @@ EXPORTS
 	SocketSetSendBuf
 	SocketSetPriority
 	SocketRecv
+	SocketRecv2
 	SocketSend
 	SocketCheckConnect
 	SocketShouldPoll
