[UBA]
* Fixed bug in scheduler that where it didn't honor force remote execution properly

diff --git a/Engine/Source/Programs/UnrealBuildAccelerator/Common/Private/UbaScheduler.cpp b/Engine/Source/Programs/UnrealBuildAccelerator/Common/Private/UbaScheduler.cpp
--- a/Engine/Source/Programs/UnrealBuildAccelerator/Common/Private/UbaScheduler.cpp
+++ b/Engine/Source/Programs/UnrealBuildAccelerator/Common/Private/UbaScheduler.cpp
@@ -1002,11 +1002,11 @@ namespace uba
 				if (ProcessEntry* entry = PeekFromQueue(m_readyToRunLocal)) // Check local only first
 					if (entry->info->weight <= weightLeft)
 						return PopFromQueue(m_readyToRunLocal);
-				if (!m_forceRemote || !m_hadRemoteRequest)
+				if (!m_forceRemote || (m_hadRemoteRequest && m_hadCrossSystemRequest)) // If force remote but we've had cross system requests we let local machine pop non-cross system queue
 					if (ProcessEntry* entry = PeekFromQueue(m_readyToRunRemote))
 						if (entry->info->weight <= weightLeft)
 							return PopFromQueue(m_readyToRunRemote);
-				if (m_forceRemote && m_hadCrossSystemRequest)
+				if (m_forceRemote && (!m_hadRemoteRequest || m_hadCrossSystemRequest))
 					return ~0u;
 				if (ProcessEntry* entry = PeekFromQueue(m_readyToRunCrossArchitecture))
 					if (entry->info->weight <= weightLeft)
