[UBA]
* Fixed bug in scheduler where a processes that dependent on a process that failed not always got skipped
* Fixed bug where processes that had higher weight than maxweight never got scheduled

diff --git a/Engine/Source/Programs/UnrealBuildAccelerator/Common/Private/UbaScheduler.cpp b/Engine/Source/Programs/UnrealBuildAccelerator/Common/Private/UbaScheduler.cpp
--- a/Engine/Source/Programs/UnrealBuildAccelerator/Common/Private/UbaScheduler.cpp
+++ b/Engine/Source/Programs/UnrealBuildAccelerator/Common/Private/UbaScheduler.cpp
@@ -10,6 +10,7 @@
 #include "UbaCacheClient.h"
 #include "UbaConfig.h"
 #include "UbaEnvironment.h"
+#include "UbaFileAccessor.h"
 #include "UbaNetworkMessage.h"
 #include "UbaNetworkServer.h"
 #include "UbaProcess.h"
@@ -403,6 +404,7 @@ namespace uba
 		entry.canExecuteRemotely = canExecuteRemotely;
 		entry.canCrossArchitecture = canCrossArchitecture;
 		entry.canCrossPlatform = canCrossPlatform;
+		entry.shouldSkip = false;
 		return entry;
 	}
 
@@ -417,15 +419,12 @@ namespace uba
 		return m_readyToRunRemote;
 	}
 
-	void Scheduler::AddToQueueBackNoLock(ProcessQueue& queue, u32 entryIndex)
+	void Scheduler::AddToQueue(ProcessQueue& queue, u32 entryIndex)
 	{
-		queue.insert(entryIndex);
+		UBA_FOR_ASSERT(auto res =) queue.insert(entryIndex);
+		UBA_ASSERT(res.second);
 	}
 
-	void Scheduler::AddToQueueFrontNoLock(ProcessQueue& queue, u32 entryIndex)
-	{
-		queue.insert(entryIndex);
-	}
 	u32 Scheduler::PopFromQueue(ProcessQueue& queue)
 	{
 		if (queue.empty())
@@ -464,7 +463,7 @@ namespace uba
 				continue;
 			if (dependency.runStatus == RunStatus_Failed || dependency.runStatus == RunStatus_Skipped)
 			{
-				AddToQueueBackNoLock(m_readyToSkip, entryIndex);
+				AddToQueue(m_readyToSkip, entryIndex);
 				return;
 			}
 			else
@@ -474,9 +473,9 @@ namespace uba
 		if (!dependencyCount)
 		{
 			if (entry.cacheStatus == CacheStatus_QueuedForTest)
-				AddToQueueBackNoLock(m_readyToTestCache, entryIndex);
+				AddToQueue(m_readyToTestCache, entryIndex);
 			else
-				AddToQueueBackNoLock(GetReadyToRunQueue(entry), entryIndex);
+				AddToQueue(GetReadyToRunQueue(entry), entryIndex);
 			return;
 		}
 
@@ -493,23 +492,27 @@ namespace uba
 
 	void Scheduler::UpdateDependentsNoLock(RunStatus newRunStatus, ProcessLink* firstDependent)
 	{
+		bool shouldSkip = newRunStatus == RunStatus_Failed || newRunStatus == RunStatus_Skipped;
 		for (ProcessLink* it = firstDependent; it; it = it->next)
 		{
 			u32 dependentIndex = it->entryIndex;
 			auto& dependent = m_processEntries[dependentIndex];
 			if (--dependent.dependencyCount)
+			{
+				if (shouldSkip)
+					dependent.shouldSkip = true;
 				continue;
-			if (newRunStatus == RunStatus_Success)
+			}
+			if (newRunStatus == RunStatus_Success && !dependent.shouldSkip)
 			{
 				if (dependent.cacheStatus == CacheStatus_QueuedForTest)
-					AddToQueueBackNoLock(m_readyToTestCache, dependentIndex);
+					AddToQueue(m_readyToTestCache, dependentIndex);
 				else
-					AddToQueueBackNoLock(GetReadyToRunQueue(dependent), dependentIndex);
+					AddToQueue(GetReadyToRunQueue(dependent), dependentIndex);
 				continue;
 			}
-			UBA_ASSERT(newRunStatus == RunStatus_Failed || newRunStatus == RunStatus_Skipped);
 			if (dependent.runStatus == RunStatus_QueuedForRun)
-				AddToQueueBackNoLock(m_readyToSkip, dependentIndex);
+				AddToQueue(m_readyToSkip, dependentIndex);
 
 		}
 	}
@@ -775,7 +778,7 @@ namespace uba
 		if (!m_cancelled)
 		{
 			entry.runStatus = RunStatus_QueuedForRun;
-			AddToQueueFrontNoLock(GetReadyToRunQueue(entry), processIndex);
+			AddToQueue(GetReadyToRunQueue(entry), processIndex);
 			UpdateQueueCounter(1, entry.canExecuteRemotely);
 			lock.Leave();
 		}
@@ -821,7 +824,7 @@ namespace uba
 		}
 
 		UBA_ASSERT(entry.runStatus == RunStatus_QueuedForRun);
-		AddToQueueBackNoLock(GetReadyToRunQueue(entry), processIndex);
+		AddToQueue(GetReadyToRunQueue(entry), processIndex);
 		lock.Leave();
 
 		m_updateThreadWakeup.Set();
@@ -1024,6 +1027,8 @@ namespace uba
 		float maxWeight = float(m_maxLocalProcessors) + 0.9f;
 		if (activeWeight > 1.0f && (activeWeight >= maxWeight || !LocalProcessOkToSpawn()))
 			return ~0u;
+		if (activeWeight == 0)
+			maxWeight = 100000.0f; // If we are not running anything we must pick up the last process regardless of weight
 
 		// If there are queued getnext, then we try to hand over running queued processes
 		// Setting the event will wake up the getnext thread.. and If the getnext thread is timing out we will retry
@@ -1085,7 +1090,7 @@ namespace uba
 
 						SCOPED_FUTEX(m_processEntriesLock, lock);
 						m_processEntries[entryToProcess].cacheStatus = CacheStatus_QueuedForDownload;
-						AddToQueueBackNoLock(m_readyToDownloadCache, entryToProcess);
+						AddToQueue(m_readyToDownloadCache, entryToProcess);
 						--m_activeCacheQueries;
 						break;
 					}
@@ -1445,7 +1450,7 @@ namespace uba
 			entry.canExecuteRemotely = false;
 			entry.canDetour = exitedResponse != ProcessExitedResponse_RerunNative;
 			entry.runStatus = RunStatus_QueuedForRun;
-			AddToQueueFrontNoLock(GetReadyToRunQueue(entry), entryIndex);
+			AddToQueue(GetReadyToRunQueue(entry), entryIndex);
 		}
 
 		if (info.isLocal)
diff --git a/Engine/Source/Programs/UnrealBuildAccelerator/Common/Public/UbaScheduler.h b/Engine/Source/Programs/UnrealBuildAccelerator/Common/Public/UbaScheduler.h
--- a/Engine/Source/Programs/UnrealBuildAccelerator/Common/Public/UbaScheduler.h
+++ b/Engine/Source/Programs/UnrealBuildAccelerator/Common/Public/UbaScheduler.h
@@ -146,8 +146,7 @@ namespace uba
 
 		ProcessEntry CreateProcessEntry(const EnqueueProcessInfo& info);
 		ProcessQueue& GetReadyToRunQueue(ProcessEntry& entry);
-		void AddToQueueBackNoLock(ProcessQueue& queue, u32 entryIndex);
-		void AddToQueueFrontNoLock(ProcessQueue& queue, u32 entryIndex);
+		void AddToQueue(ProcessQueue& queue, u32 entryIndex);
 		u32 PopFromQueue(ProcessQueue& queue);
 		ProcessEntry* PeekFromQueue(ProcessQueue& queue);
 		u32 GetQueueCount(ProcessQueue& queue);
@@ -189,6 +188,7 @@ namespace uba
 			u32 dependencyCount;
 			CacheStatus cacheStatus;
 			RunStatus runStatus;
+			bool shouldSkip;
 			bool canDetour;
 			bool canExecuteRemotely;
 			bool canCrossArchitecture;
