[UBA]
* Fixing up includes for nonunity builds

diff --git a/Engine/Source/Programs/UnrealBuildAccelerator/Common/Private/UbaTraceChannel.cpp b/Engine/Source/Programs/UnrealBuildAccelerator/Common/Private/UbaTraceChannel.cpp
--- a/Engine/Source/Programs/UnrealBuildAccelerator/Common/Private/UbaTraceChannel.cpp
+++ b/Engine/Source/Programs/UnrealBuildAccelerator/Common/Private/UbaTraceChannel.cpp
@@ -2,6 +2,7 @@
 
 #include "UbaTraceChannel.h"
 #include "UbaFile.h"
+#include "UbaFileMapping.h"
 #include "UbaTrace.h"
 
 #if PLATFORM_WINDOWS
diff --git a/Engine/Source/Programs/UnrealBuildAccelerator/Common/Public/UbaDependencyCrawler.h b/Engine/Source/Programs/UnrealBuildAccelerator/Common/Public/UbaDependencyCrawler.h
--- a/Engine/Source/Programs/UnrealBuildAccelerator/Common/Public/UbaDependencyCrawler.h
+++ b/Engine/Source/Programs/UnrealBuildAccelerator/Common/Public/UbaDependencyCrawler.h
@@ -4,7 +4,8 @@
 
 #include "UbaApplicationRules.h"
 #include "UbaBottleneck.h"
-#include "UbaFileAccessor.h"
+#include "UbaHash.h"
+#include "UbaLogger.h"
 #include "UbaPathUtils.h"
 #include "UbaProcessUtils.h"
 #include "UbaWorkManager.h"
diff --git a/Engine/Source/Programs/UnrealBuildAccelerator/Common/Public/UbaFileMappingBuffer.h b/Engine/Source/Programs/UnrealBuildAccelerator/Common/Public/UbaFileMappingBuffer.h
--- a/Engine/Source/Programs/UnrealBuildAccelerator/Common/Public/UbaFileMappingBuffer.h
+++ b/Engine/Source/Programs/UnrealBuildAccelerator/Common/Public/UbaFileMappingBuffer.h
@@ -2,10 +2,12 @@
 
 #pragma once
 
-#include "UbaFileMapping.h"
+#include "UbaHash.h"
+#include "UbaSharedMemoryView.h"
 
 namespace uba
 {
+	class Logger;
 	class SharedMemoryAllocator;
 	class Trace;
 
diff --git a/Engine/Source/Programs/UnrealBuildAccelerator/Common/Public/UbaStorage.h b/Engine/Source/Programs/UnrealBuildAccelerator/Common/Public/UbaStorage.h
--- a/Engine/Source/Programs/UnrealBuildAccelerator/Common/Public/UbaStorage.h
+++ b/Engine/Source/Programs/UnrealBuildAccelerator/Common/Public/UbaStorage.h
@@ -4,11 +4,11 @@
 
 #include "UbaBottleneck.h"
 #include "UbaBufferCache.h"
+#include "UbaEnvironment.h"
 #include "UbaFile.h"
 #include "UbaFileMappingBuffer.h"
 #include "UbaLogger.h"
 #include "UbaMemory.h"
-#include "UbaPathUtils.h"
 #include "UbaSharedMemoryAllocator.h"
 #include "UbaStats.h"
 
diff --git a/Engine/Source/Programs/UnrealBuildAccelerator/Common/Public/UbaStorageUtils.h b/Engine/Source/Programs/UnrealBuildAccelerator/Common/Public/UbaStorageUtils.h
--- a/Engine/Source/Programs/UnrealBuildAccelerator/Common/Public/UbaStorageUtils.h
+++ b/Engine/Source/Programs/UnrealBuildAccelerator/Common/Public/UbaStorageUtils.h
@@ -2,17 +2,17 @@
 
 #pragma once
 
-#include "UbaBinaryReaderWriter.h"
 #include "UbaFile.h"
-#include "UbaFileMapping.h"
 #include "UbaHash.h"
-#include "UbaNetworkMessage.h"
+#include "UbaSharedMemoryView.h"
 #include <oodle2.h>
 
 namespace uba
 {
 	class FileAccessor;
+	class NetworkClient;
 	class SharedMemoryAllocator;
+	class WorkManager;
 	struct BufferCache;
 	struct StorageStats;
 
diff --git a/Engine/Source/Programs/UnrealBuildAccelerator/Common/Public/UbaTraceChannel.h b/Engine/Source/Programs/UnrealBuildAccelerator/Common/Public/UbaTraceChannel.h
--- a/Engine/Source/Programs/UnrealBuildAccelerator/Common/Public/UbaTraceChannel.h
+++ b/Engine/Source/Programs/UnrealBuildAccelerator/Common/Public/UbaTraceChannel.h
@@ -2,7 +2,8 @@
 
 #pragma once
 
-#include "UbaFileMapping.h"
+#include "UbaEnvironment.h"
+#include "UbaFileMappingHandle.h"
 #include "UbaLogger.h"
 
 #define UBA_USE_TRACE_CHANNEL_FILE !PLATFORM_WINDOWS
diff --git a/Engine/Source/Programs/UnrealBuildAccelerator/ObjTool/Private/UbaObjTool.cpp b/Engine/Source/Programs/UnrealBuildAccelerator/ObjTool/Private/UbaObjTool.cpp
--- a/Engine/Source/Programs/UnrealBuildAccelerator/ObjTool/Private/UbaObjTool.cpp
+++ b/Engine/Source/Programs/UnrealBuildAccelerator/ObjTool/Private/UbaObjTool.cpp
@@ -1,8 +1,8 @@
 // Copyright Epic Games, Inc. All Rights Reserved.
 
 #include "UbaObjectFile.h"
-#include "UbaDirectoryIterator.h"
-#include "UbaFileAccessor.h"
+#include "UbaEnvironment.h"
+#include "UbaFile.h"
 #include "UbaImportLibWriter.h"
 #include "UbaPathUtils.h"
 #include "UbaVersion.h"
