[UBA]
* Unified FileMappingHandle between windows and posix and moved it into its own file
* Started moving out platform dependent file mapping code into inl file that can be included in all targets. Still work left

Reason is because there is a bug in Wine's MapViewOfFile3 so we have to work around it by going directly on linux api when running in wine (similar to how linux sockets are directly used on wine instead of going through wine layer)

diff --git a/Engine/Source/Programs/UnrealBuildAccelerator/Common/Private/UbaEnvironment.cpp b/Engine/Source/Programs/UnrealBuildAccelerator/Common/Private/UbaEnvironment.cpp
--- a/Engine/Source/Programs/UnrealBuildAccelerator/Common/Private/UbaEnvironment.cpp
+++ b/Engine/Source/Programs/UnrealBuildAccelerator/Common/Private/UbaEnvironment.cpp
@@ -561,4 +561,40 @@ namespace uba
 #endif
 		return true;
 	}
+
+	bool g_memoryLockEnabled;
+
+	void SetMemoryWorkingSet(Logger& logger, u64 minWorkingSet)
+	{
+		if (!SetMemoryWorkingSet(minWorkingSet))
+		{
+			LogEntryType type = LogEntryType_Detail;
+			#if UBA_DEBUG
+			type = LogEntryType_Error;
+			#endif
+			logger.Logf(type, TC("Failed to set memory working set (%s)"), LastErrorToText().data);
+			return;
+		}
+
+		logger.Detail(TC("MemoryWorkingSet min set to %s"), BytesToText(minWorkingSet).str);
+		g_memoryLockEnabled = true;
+	}
+
+	void LockMemory(Logger& logger, void* dest, u64 size, const tchar* hint)
+	{
+		if (!g_memoryLockEnabled)
+			return;
+
+		#if PLATFORM_WINDOWS
+		if (!VirtualLock(dest, size))
+		{
+			g_memoryLockEnabled = false;
+			LogEntryType type = LogEntryType_Detail;
+			#if UBA_DEBUG
+			type = LogEntryType_Error;
+			#endif
+			logger.Logf(type, TC("LockMemory of %s failed for %s (%s)"), BytesToText(size).str, hint, LastErrorToText().data);
+		}
+		#endif
+	}
 }
diff --git a/Engine/Source/Programs/UnrealBuildAccelerator/Common/Private/UbaFile.cpp b/Engine/Source/Programs/UnrealBuildAccelerator/Common/Private/UbaFile.cpp
--- a/Engine/Source/Programs/UnrealBuildAccelerator/Common/Private/UbaFile.cpp
+++ b/Engine/Source/Programs/UnrealBuildAccelerator/Common/Private/UbaFile.cpp
@@ -108,12 +108,6 @@ namespace uba
 	#define MAKE_LONG_FILENAME(fileName) PREFIX_FILENAME(fileName, TCV("\\\\?\\"))
 	#define MAKE_NT_FILENAME(fileName)  PREFIX_FILENAME(fileName, TCV("\\??\\"))
 #else
-	int asFileDescriptor(FileHandle fh)
-	{
-		if (fh == InvalidFileHandle)
-			return (int)fh;
-		return (int)(fh & FileHandleFlagMask);
-	}
 	#define ENSURE_TERMINATED(fileName) \
 		StringBuffer<MaxPath> STRING_JOIN(temp, __LINE__); \
 		if (fileName.data[fileName.count] != 0) \
@@ -677,7 +671,7 @@ namespace uba
 #else
 		if (close(asFileDescriptor(h)) == 0)
 			return true;
-		UBA_ASSERTF(false, TC("CloseFile error handling not implemented while failing to close %s (%s)"), fileName, strerror(errno));
+		UBA_ASSERTF(false, TC("CloseFile error handling not implemented while failing to close %s with handle %i (%s)"), fileName, asFileDescriptor(h), strerror(errno));
 		return false;
 #endif
 	}
diff --git a/Engine/Source/Programs/UnrealBuildAccelerator/Common/Private/UbaFileAccessor.cpp b/Engine/Source/Programs/UnrealBuildAccelerator/Common/Private/UbaFileAccessor.cpp
--- a/Engine/Source/Programs/UnrealBuildAccelerator/Common/Private/UbaFileAccessor.cpp
+++ b/Engine/Source/Programs/UnrealBuildAccelerator/Common/Private/UbaFileAccessor.cpp
@@ -1,6 +1,7 @@
 // Copyright Epic Games, Inc. All Rights Reserved.
 
 #include "UbaFileAccessor.h"
+#include "UbaEvent.h"
 #include "UbaFile.h"
 #include "UbaLogger.h"
 #include "UbaStats.h"
@@ -450,7 +451,8 @@ namespace uba
 		if (!m_mappingHandle.IsValid())
 			return m_logger.Error(TC("Failed to create mapping handle for %s with size %llu (%s)"), m_fileName, m_size, LastErrorToText().data);
 #else
-		m_mappingHandle = { asFileDescriptor(m_fileHandle) };
+		m_mappingHandle.internal = ~0ull;
+		m_mappingHandle.shmFd = asFileDescriptor(m_fileHandle);
 		if (offset == m_size)
 			return true;
 #endif
diff --git a/Engine/Source/Programs/UnrealBuildAccelerator/Common/Private/UbaFileMapping.cpp b/Engine/Source/Programs/UnrealBuildAccelerator/Common/Private/UbaFileMapping.cpp
--- a/Engine/Source/Programs/UnrealBuildAccelerator/Common/Private/UbaFileMapping.cpp
+++ b/Engine/Source/Programs/UnrealBuildAccelerator/Common/Private/UbaFileMapping.cpp
@@ -2,10 +2,13 @@
 
 #include "UbaFileMapping.h"
 #include "UbaBottleneck.h"
-#include "UbaDirectoryIterator.h"
+#include "UbaDefinitions.h"
+#include "UbaFileHandle.h"
+#include "UbaLogger.h"
 #include "UbaProcessStats.h"
 
 #if !PLATFORM_WINDOWS
+#include <dirent.h>
 #include <sys/file.h>
 #endif
 
@@ -19,371 +22,23 @@ namespace uba
 	UnorderedMap<HANDLE, DebugFileMapping> g_fileMappings;
 	UnorderedMap<const void*, HANDLE> g_viewMappings;
 	#endif
+}
 
-	bool g_memoryLockEnabled;
-
-#if PLATFORM_WINDOWS
-	Bottleneck& g_createFileHandleBottleneck = *new Bottleneck(8); // Allocated and leaked just to prevent shutdown asserts in debug
-
-	HANDLE asHANDLE(FileHandle fh);
-
-	HANDLE InternalCreateFileMappingW(Logger& logger, HANDLE hFile, DWORD flProtect, DWORD dwMaximumSizeHigh, DWORD dwMaximumSizeLow, LPCWSTR lpName, const tchar* hint)
-	{
-		HANDLE res;
-		if (flProtect == PAGE_READWRITE)
-		{
-			// Experiment to try to prevent lock happening on AWS servers when lots of helpers are sending back obj files.
-			Timer timer;
-			BottleneckScope bs(g_createFileHandleBottleneck, timer);
-			res = ::CreateFileMappingW(hFile, NULL, flProtect, dwMaximumSizeHigh, dwMaximumSizeLow, lpName);
-		}
-		else
-			res = ::CreateFileMappingW(hFile, NULL, flProtect, dwMaximumSizeHigh, dwMaximumSizeLow, lpName);
-
-		#if UBA_DEBUG_FILE_MAPPING
-		if (res)
-		{
-			SCOPED_FUTEX(g_fileMappingsLock, lock);
-			UBA_FOR_ASSERT(auto insres =) g_fileMappings.try_emplace(res, hint);
-			UBA_ASSERT(insres.second);
-			//logger.Info(TC("FILEMAPPING CREAT: %s %llu (%llu)"), hint, u64(res), g_fileMappings.size());
-		}
-		#endif
-
-		return res;
-	}
-#else
-	int asFileDescriptor(FileHandle fh);
-
-	Futex g_mappingUidCounterLock;
-	Atomic<u64> g_mappingUidCounter = 0;
-#endif
-
-	FileMappingHandle CreateMemoryMappingW(Logger& logger, u32 flProtect, u64 maxSize, const tchar* name, const tchar* hint)
-	{
-		ExtendedTimerScope ts(KernelStats::GetCurrent().createFileMapping);
-#if PLATFORM_WINDOWS
-		return { 0, InternalCreateFileMappingW(logger, INVALID_HANDLE_VALUE, flProtect, (DWORD)ToHigh(maxSize), ToLow(maxSize), name, hint) };
-#else
-		UBA_ASSERT(!name);
-		UBA_ASSERT((flProtect & (~u32(PAGE_READWRITE | SEC_RESERVE))) == 0);
-
-		// Since we need to not leak file mappings we use files as a trick to know which ones are used and not
-		StringBuffer<64> lockDir;
-		lockDir.Append("/tmp/uba_shm_locks");
-
-		SCOPED_FUTEX(g_mappingUidCounterLock, lock);
-		if (!g_mappingUidCounter)
-		{
-			// Create dir
-			if (mkdir(lockDir.data, 0777) == -1)
-				if (errno != EEXIST)
-				{
-					logger.Error(TC("Failed to create %s for memory mapping (%s)"), lockDir.data, strerror(errno));
-					return {};
-				}
-
-			// Clear out all orphaned shm_open
-			TraverseDir(logger, lockDir,
-				[&](const DirectoryEntry& e)
-				{
-					u32 uid = strtoul(e.name, nullptr, 10);
-
-					StringBuffer<128> lockFile;
-					lockFile.Append(lockDir).EnsureEndsWithSlash().Append(e.name);
-					int lockFd = open(lockFile.data, O_RDWR, S_IRUSR | S_IWUSR | O_CLOEXEC);
-					if (lockFd == -1)
-					{
-						if (errno == EPERM)
-						{
-							g_mappingUidCounter = uid;
-							return;
-						}
-						logger.Warning("Failed to open %s for memory mapping (%s)", lockFile.data, strerror(errno));
-						UBA_ASSERTF(false, "Failed to open %s (%s)", lockFile.data, strerror(errno));
-						return;
-					}
-
-					if (flock(lockFd, LOCK_EX | LOCK_NB) == 0)
-					{
-						StringBuffer<64> uidName;
-						GetMappingHandleName(uidName, uid);
-						if (shm_unlink(uidName.data) == 0)
-							;// logger.Info("Removed old shared memory %s", uidName.data);
-						remove(lockFile.data);
-					}
-					else
-					{
-						g_mappingUidCounter = uid;
-					}
-					close(lockFd);
-				});
-
-			if (g_mappingUidCounter)
-				logger.Info("Starting shared memory files at %u", g_mappingUidCounter.load());
-		}
-		lock.Leave();
-
-		// Let's find a free shm name
-		StringBuffer<128> lockFile;
-		u64 uid;
-		int shmFd;
-		int lockFd;
-
-		int retryCount = 4;
-		while (true)
-		{
-			uid = ++g_mappingUidCounter;
-
-			lockFile.Clear().Append(lockDir).EnsureEndsWithSlash().AppendValue(uid);
-
-			lockFd = open(lockFile.data, O_CREAT | O_RDWR | O_NOFOLLOW | O_EXCL, S_IRUSR | S_IWUSR | O_CLOEXEC);
-			if (lockFd == -1)
-			{
-				if (errno == EEXIST)
-					continue;
-				logger.Warning(TC("Failed to open/create %s (%s)"), lockFile.data, strerror(errno));
-				UBA_ASSERTF(false, "Failed to open/create %s (%s)", lockFile.data, strerror(errno));
-				continue;
-			}
-
-			if (flock(lockFd, LOCK_EX | LOCK_NB) != 0) // Some other process is using this one
-			{
-				close(lockFd);
-				continue;
-			}
-
-			StringBuffer<64> uidName;
-			GetMappingHandleName(uidName, uid);
-
-			int oflags = O_CREAT | O_RDWR | O_NOFOLLOW | O_EXCL;
-			shmFd = shm_open(uidName.data, oflags, S_IRUSR | S_IWUSR);
-			if (shmFd != -1)
-				break;
-
-			bool retry = retryCount-- > 0;
-
-			int err = errno;
-			auto logType = retry ? LogEntryType_Warning : LogEntryType_Error;
-			logger.Logf(logType, TC("Failed to create shm %s after getting lock-file %s (%s)"), uidName.data, lockFile.data, strerror(err));
-
-			if (retry)
-				continue;
-
-			remove(lockFile.data);
-			close(lockFd);
-			SetLastError(err);
-			return {};
-		}
-
-		if (maxSize != 0)
-		{
-			if (ftruncate(shmFd, (s64)maxSize) == -1)
-			{
-				SetLastError(errno);
-				close(shmFd);
-				remove(lockFile.data);
-				close(lockFd);
-				//UBA_ASSERTF(false, "Failed to truncate file mapping '%s' to size %llu (%s)", name, maxSize, strerror(errno));
-				return {};
-			}
-		}
-		SetLastError(0);
-		return { shmFd, lockFd, uid };
-#endif
-	}
-
-	FileMappingHandle CreateFileMappingW(Logger& logger, FileHandle file, u32 protect, u64 maxSize, const tchar* hint)
-	{
-		ExtendedTimerScope ts(KernelStats::GetCurrent().createFileMapping);
-#if PLATFORM_WINDOWS
-		return { 0, InternalCreateFileMappingW(logger, asHANDLE(file), protect, (DWORD)ToHigh(maxSize), ToLow(maxSize), NULL, hint) };
-#else
-		FileMappingHandle h;
-		int fd = asFileDescriptor(file);
-		if (maxSize && (protect & (~PAGE_READONLY)) != 0)
-		{
-#if PLATFORM_MAC // For some reason lseek+write does not work on apple silicon platform
-			if (ftruncate(fd, maxSize) == -1)
-			{
-				logger.Error("ftruncate to %llu on fd %i failed for %s: %s\n", maxSize, fd, hint, strerror(errno));
-				return h;
-			}
-#else
-			if (lseek(fd, maxSize - 1, SEEK_SET) != maxSize - 1)
-			{
-				logger.Error("lseek to %llu failed for %s: %s\n", maxSize - 1, hint, strerror(errno));
-				return h;
-			}
-
-			errno = 0;
-			int res = write(fd, "", 1);
-			if (res != 1)
-			{
-				logger.Error("write one byte at %llu on fd %i (%s) failed (res: %i): %s\n", maxSize - 1, fd, hint, res, strerror(errno));
-				return h;
-			}
-#endif
-		}
-
-		h.shmFd = fd;
-		return h;
-#endif
-	}
-
-	static void InternalDiscard(Logger& logger, u8* mem, u64 size, const tchar* from, const tchar* hint)
-	{
-#if PLATFORM_WINDOWS
-		u64 stepSize = 16ull*1024*1024*1024; // Extremely large.. here to test if it is faster to do multiple calls
-		u64 left = size;
-		while (left)
-		{
-			u64 toDiscard = Min(left, stepSize);
-			DWORD res = DiscardVirtualMemory(mem, AlignUp(toDiscard, 4*1024));
-			if (res != ERROR_SUCCESS)
-				logger.Warning(TC("%s - Failed to discard memory (%s) (%u)"), from, hint, res);
-			mem += toDiscard;
-			left -= toDiscard;
-		}
-#endif
-	}
-
-	u8* MapViewOfFile(Logger& logger, FileMappingHandle fileMappingObject, u32 desiredAccess, u64 offset, u64 bytesToMap, bool allowDiscard)
-	{
-		ExtendedTimerScope ts(KernelStats::GetCurrent().mapViewOfFile);
-#if PLATFORM_WINDOWS
-		u8* res = (u8*)::MapViewOfFile(fileMappingObject.mh, desiredAccess, (DWORD)ToHigh(offset), ToLow(offset), bytesToMap);
-		if (!res)
-			return nullptr;
-		if (allowDiscard)
-			InternalDiscard(logger, res, bytesToMap, TC("MapViewOfFile"), TC(""));
-#else
-		int prot = 0;
-		if (desiredAccess & FILE_MAP_READ)
-			prot |= PROT_READ;
-		if (desiredAccess & FILE_MAP_WRITE)
-			prot |= PROT_WRITE;
-		UBA_ASSERT(fileMappingObject.IsValid());
-		int shmFd = fileMappingObject.shmFd;
-		void* rptr = mmap(NULL, bytesToMap, prot, MAP_SHARED, shmFd, s64(offset));
-		if (rptr == MAP_FAILED)
-		{
-			SetLastError(errno);
-			return nullptr;
-		}
-		u8* res = (u8*)rptr;
-		//UBA_ASSERTF(false, "Failed to map file with fd %i, desiredAccess %u offset %llu, bytesToMap %llu (%s)", fd, desiredAccess, offset, bytesToMap, strerror(errno));
-#endif
-
-		#if UBA_DEBUG_FILE_MAPPING
-		{
-			SCOPED_FUTEX(g_fileMappingsLock, lock);
-			auto findIt = g_fileMappings.find(fileMappingObject.mh);
-			UBA_ASSERTF(findIt != g_fileMappings.end(), TC("Mapping nonexisting handle: %llu"), u64(fileMappingObject.mh));
-			auto insres = g_viewMappings.try_emplace(res, fileMappingObject.mh);
-			UBA_ASSERT(insres.second);
-			++findIt->second.viewCount;
-		}
-		#endif
-
-		return res;
-	}
-
-	bool MapViewCommit(Logger& logger, void* address, u64 size, bool allowDiscard)
-	{
-#if PLATFORM_WINDOWS
-		ExtendedTimerScope ts(KernelStats::GetCurrent().virtualAlloc);
-		if (!::VirtualAlloc(address, size, MEM_COMMIT, PAGE_READWRITE))
-			return false;
-		if (allowDiscard)
-			InternalDiscard(logger, (u8*)address, size, TC("MapViewCommit"), TC(""));
-		return true;
-#else
-		return true;
-#endif
-	}
-
-	bool UnmapViewOfFile(Logger& logger, const void* address, u64 bytesToUnmap, const tchar* hint, bool allowDiscard)
-	{
-		if (!address)
-			return true;
-
-		ExtendedTimerScope ts(KernelStats::GetCurrent().unmapViewOfFile);
-
-		#if UBA_DEBUG_FILE_MAPPING
-		{
-			SCOPED_FUTEX(g_fileMappingsLock, lock);
-			auto findIt = g_viewMappings.find(address);
-			UBA_ASSERT(findIt != g_viewMappings.end());
-			auto findIt2 = g_fileMappings.find(findIt->second);
-			g_viewMappings.erase(findIt);
-			UBA_ASSERTF(findIt2 != g_fileMappings.end(), TC("Unmap nonexisting handle: %llu"), u64(findIt->second));
-			--findIt2->second.viewCount;
-		}
-		#endif
+////////////////////////////////////////////////////////////////////////////////////////////////////
 
 #if PLATFORM_WINDOWS
-		if (allowDiscard)
-			InternalDiscard(logger, (u8*)address, bytesToUnmap, TC("UnmapViewOfFile"), hint);
-		return ::UnmapViewOfFile(address);
+#include "UbaFileMappingWin.inl"
 #else
-		UBA_ASSERTF(bytesToUnmap, TC("bytesToUnmap is zero unmapping %p (%s)"), address, hint);
-		if (munmap((void*)address, bytesToUnmap) == 0)
-			return true;
-		UBA_ASSERT(false);
-		return false;
+namespace uba
+{
+#include "UbaFileMapping.inl"
+}
 #endif
-	}
 
-	void DiscardFileMapping(Logger& logger, FileMappingHandle h, u64 size, const tchar* hint)
-	{
-		#if PLATFORM_WINDOWS
-		void* mem = MapViewOfFile(logger, h, FILE_MAP_WRITE|FILE_MAP_READ, 0, size, true);
-		UnmapViewOfFile(logger, mem, size, TC(""));
-		#endif
-	}
-
-	bool CloseFileMapping(Logger& logger, FileMappingHandle h, const tchar* hint)
-	{
-		if (!h.IsValid())
-			return true;
-
-		#if UBA_DEBUG_FILE_MAPPING
-		SCOPED_FUTEX(g_fileMappingsLock, lock);
-		auto findIt = g_fileMappings.find(h.mh);
-		UBA_ASSERTF(findIt != g_fileMappings.end(), TC("Handle: %llu"), u64(h.mh));
-		UBA_ASSERTF(!findIt->second.viewCount.load(), TC("Closing file mapping with %llu open views (%s)"), findIt->second.viewCount.load(), findIt->second.hint.c_str());
-		logger.Info(TC("FILEMAPPING CLOSE: %s %llu (%llu)"), findIt->second.hint.c_str(), u64(h.mh), g_fileMappings.size() - 1);
-		g_fileMappings.erase(findIt);
-		#endif
-
-#if PLATFORM_WINDOWS
-		if (h.fh)
-			CloseHandle(h.fh);
-		return CloseHandle(h.mh);
-#else
-		if (h.uid == ~u64(0))
-			return true;
-		if (close(h.shmFd) != 0)
-			UBA_ASSERTF(false, "Failed to close file mapping (%s)", strerror(errno));
-
-		StringBuffer<64> uidName;
-		GetMappingHandleName(uidName, h.uid);
-		if (shm_unlink(uidName.data) != 0)
-		{
-			SetLastError(errno);
-			UBA_ASSERTF(false, "Failed to unlink %s (%s)", uidName.data, strerror(errno));
-			return false;
-		}
-
-		StringBuffer<128> lockFile;
-		lockFile.Append("/tmp/uba_shm_locks").EnsureEndsWithSlash().AppendValue(h.uid);
-		remove(lockFile.data);
-		close(h.lockFd);
-		return true;
-#endif
-	}
+////////////////////////////////////////////////////////////////////////////////////////////////////
 
+namespace uba
+{
 	void MapMemoryCopy(void* dest, const void* source, u64 size)
 	{
 		auto& stats = KernelStats::GetCurrent();
@@ -392,117 +47,6 @@ namespace uba
 		memcpy(dest, source, size);
 	}
 
-	void SetMemoryWorkingSet(Logger& logger, u64 minWorkingSet)
-	{
-		if (!SetMemoryWorkingSet(minWorkingSet))
-		{
-			LogEntryType type = LogEntryType_Detail;
-			#if UBA_DEBUG
-			type = LogEntryType_Error;
-			#endif
-			logger.Logf(type, TC("Failed to set memory working set (%s)"), LastErrorToText().data);
-			return;
-		}
-
-		logger.Detail(TC("MemoryWorkingSet min set to %s"), BytesToText(minWorkingSet).str);
-		g_memoryLockEnabled = true;
-	}
-
-	void LockMemory(Logger& logger, void* dest, u64 size, const tchar* hint)
-	{
-		if (!g_memoryLockEnabled)
-			return;
-
-		#if PLATFORM_WINDOWS
-		if (!VirtualLock(dest, size))
-		{
-			g_memoryLockEnabled = false;
-			LogEntryType type = LogEntryType_Detail;
-			#if UBA_DEBUG
-			type = LogEntryType_Error;
-			#endif
-			logger.Logf(type, TC("LockMemory of %s failed for %s (%s)"), BytesToText(size).str, hint, LastErrorToText().data);
-		}
-		#endif
-	}
-
 	////////////////////////////////////////////////////////////////////////////////////////////////////
 
-	FixedSizeFileMappingAllocator::FixedSizeFileMappingAllocator(Logger& logger, const tchar* name)
-	:	m_logger(logger)
-	,	m_name(name)
-	{
-	}
-
-	FixedSizeFileMappingAllocator::~FixedSizeFileMappingAllocator()
-	{
-		if (m_mappingHandle.IsValid())
-			CloseFileMapping(m_logger, m_mappingHandle, TC("FixedSizeFileMappingAllocator"));
-	}
-
-	bool FixedSizeFileMappingAllocator::Init(u64 blockSize, u64 capacity)
-	{
-		m_mappingHandle = uba::CreateMemoryMappingW(m_logger, PAGE_READWRITE|SEC_RESERVE, capacity, nullptr, TC("FixedSizeFileMappingAllocator"));
-		if (!m_mappingHandle.IsValid())
-			return m_logger.Error(TC("%s - Failed to create memory map with capacity %llu (%s)"), m_name, capacity, LastErrorToText().data);
-
-		m_blockSize = blockSize;
-		m_capacity = capacity;
-		return true;
-	}
-
-	FileMappingAllocation FixedSizeFileMappingAllocator::Alloc(const tchar* hint)
-	{
-		SCOPED_FUTEX(m_mappingLock, lock);
-
-		u64 index = m_mappingCount;
-		bool needCommit = false;
-		if (!m_availableBlocks.empty())
-		{
-			auto it = m_availableBlocks.begin();
-			index = *it;
-			m_availableBlocks.erase(it);
-		}
-		else
-		{
-			++m_mappingCount;
-			needCommit = true;
-		}
-		lock.Leave();
-
-		u64 offset = index*m_blockSize;
-		u8* data = MapViewOfFile(m_logger, m_mappingHandle, FILE_MAP_READ|FILE_MAP_WRITE, offset, m_blockSize);
-		if (!data)
-		{
-			if (m_capacity < m_mappingCount*m_blockSize)
-				m_logger.Error(TC("%s - Out of capacity (%llu) need to bump capacity for %s (%s)"), m_name, m_capacity, hint, LastErrorToText().data);
-			else
-				m_logger.Error(TC("%s - Alloc failed to map view of file for %s (%s)"), m_name, hint, LastErrorToText().data);
-			return { {}, 0, 0 };
-		}
-
-		if (needCommit)
-		{
-			if (!MapViewCommit(m_logger, data, m_blockSize, true))
-			{
-				m_logger.Error(TC("%s - Failed to allocate memory for %s (%s)"), m_name, hint, LastErrorToText().data);
-				return { {}, 0, 0 };
-			}
-		}
-		LockMemory(m_logger, data, m_blockSize, hint);
-
-		return {m_mappingHandle, offset, data};
-	}
-
-	void FixedSizeFileMappingAllocator::Free(FileMappingAllocation allocation)
-	{
-		UBA_ASSERT(allocation.handle == m_mappingHandle);
-		if (!UnmapViewOfFile(m_logger, allocation.memory, m_blockSize, m_name, true))
-			m_logger.Error(TC("%s - Failed to unmap view of file (%s)"), m_name, LastErrorToText().data);
-		u64 index = allocation.offset / m_blockSize;
-		SCOPED_FUTEX(m_mappingLock, lock);
-		m_availableBlocks.insert(index);
-	}
-
-	////////////////////////////////////////////////////////////////////////////////////////////////////
 }
diff --git a/Engine/Source/Programs/UnrealBuildAccelerator/Common/Private/UbaFixedSizeFileMappingAllocator.cpp b/Engine/Source/Programs/UnrealBuildAccelerator/Common/Private/UbaFixedSizeFileMappingAllocator.cpp
new file mode 100644
--- /dev/null
+++ b/Engine/Source/Programs/UnrealBuildAccelerator/Common/Private/UbaFixedSizeFileMappingAllocator.cpp
@@ -0,0 +1,91 @@
+// Copyright Epic Games, Inc. All Rights Reserved.
+
+#include "UbaFixedSizeFileMappingAllocator.h"
+#include "UbaEnvironment.h"
+#include "UbaFileMapping.h"
+#include "UbaLogger.h"
+#include "UbaPlatform.h"
+#include "UbaStringBuffer.h"
+
+namespace uba
+{
+	////////////////////////////////////////////////////////////////////////////////////////////////////
+
+	FixedSizeFileMappingAllocator::FixedSizeFileMappingAllocator(Logger& logger, const tchar* name)
+		:	m_logger(logger)
+		,	m_name(name)
+	{
+	}
+
+	FixedSizeFileMappingAllocator::~FixedSizeFileMappingAllocator()
+	{
+		if (m_mappingHandle.IsValid())
+			CloseFileMapping(m_logger, m_mappingHandle, TC("FixedSizeFileMappingAllocator"));
+	}
+
+	bool FixedSizeFileMappingAllocator::Init(u64 blockSize, u64 capacity)
+	{
+		m_mappingHandle = uba::CreateMemoryMappingW(m_logger, PAGE_READWRITE|SEC_RESERVE, capacity, nullptr, TC("FixedSizeFileMappingAllocator"));
+		if (!m_mappingHandle.IsValid())
+			return m_logger.Error(TC("%s - Failed to create memory map with capacity %llu (%s)"), m_name, capacity, LastErrorToText().data);
+
+		m_blockSize = blockSize;
+		m_capacity = capacity;
+		return true;
+	}
+
+	FileMappingAllocation FixedSizeFileMappingAllocator::Alloc(const tchar* hint)
+	{
+		SCOPED_FUTEX(m_mappingLock, lock);
+
+		u64 index = m_mappingCount;
+		bool needCommit = false;
+		if (!m_availableBlocks.empty())
+		{
+			auto it = m_availableBlocks.begin();
+			index = *it;
+			m_availableBlocks.erase(it);
+		}
+		else
+		{
+			++m_mappingCount;
+			needCommit = true;
+		}
+		lock.Leave();
+
+		u64 offset = index*m_blockSize;
+		u8* data = MapViewOfFile(m_logger, m_mappingHandle, FILE_MAP_READ|FILE_MAP_WRITE, offset, m_blockSize, false);
+		if (!data)
+		{
+			if (m_capacity < m_mappingCount*m_blockSize)
+				m_logger.Error(TC("%s - Out of capacity (%llu) need to bump capacity for %s (%s)"), m_name, m_capacity, hint, LastErrorToText().data);
+			else
+				m_logger.Error(TC("%s - Alloc failed to map view of file for %s (%s)"), m_name, hint, LastErrorToText().data);
+			return { {}, 0, 0 };
+		}
+
+		if (needCommit)
+		{
+			if (!MapViewCommit(m_logger, data, m_blockSize, true))
+			{
+				m_logger.Error(TC("%s - Failed to allocate memory for %s (%s)"), m_name, hint, LastErrorToText().data);
+				return { {}, 0, 0 };
+			}
+		}
+		LockMemory(m_logger, data, m_blockSize, hint);
+
+		return {m_mappingHandle, offset, data};
+	}
+
+	void FixedSizeFileMappingAllocator::Free(FileMappingAllocation allocation)
+	{
+		UBA_ASSERT(allocation.handle == m_mappingHandle);
+		if (!UnmapViewOfFile(m_logger, allocation.memory, m_blockSize, m_name, true))
+			m_logger.Error(TC("%s - Failed to unmap view of file (%s)"), m_name, LastErrorToText().data);
+		u64 index = allocation.offset / m_blockSize;
+		SCOPED_FUTEX(m_mappingLock, lock);
+		m_availableBlocks.insert(index);
+	}
+
+	////////////////////////////////////////////////////////////////////////////////////////////////////
+}
diff --git a/Engine/Source/Programs/UnrealBuildAccelerator/Common/Private/UbaProcess.cpp b/Engine/Source/Programs/UnrealBuildAccelerator/Common/Private/UbaProcess.cpp
--- a/Engine/Source/Programs/UnrealBuildAccelerator/Common/Private/UbaProcess.cpp
+++ b/Engine/Source/Programs/UnrealBuildAccelerator/Common/Private/UbaProcess.cpp
@@ -186,7 +186,7 @@ namespace uba
 		if (UseSemaphoreForSharedEvents())
 		{
 			StringBuffer<256> name;
-			name.Append(TCV("uba")).AppendValue(m_comMemory.handle.uid).Append('_').AppendValue(m_comMemory.offset);
+			name.Append(TCV("uba")).AppendValue(m_comMemory.handle.internal).Append('_').AppendValue(m_comMemory.offset);
 			u64 len = name.count;
 			name.Append("_cancel");
 			m_cancelEvent.Create(name, true);
@@ -1410,7 +1410,7 @@ namespace uba
 		#endif
 
 		#if !PLATFORM_WINDOWS
-		writer.WriteU64(process.m_comMemory.handle.uid);
+		writer.WriteU64(process.m_comMemory.handle.ToU64());
 		writer.WriteU32(process.m_comMemory.offset);
 		writer.WriteString(info.logFile);
 		#endif
@@ -2139,7 +2139,7 @@ namespace uba
 			payload.cancelEvent = m_cancelEvent.GetHandle();
 			payload.writeEvent = m_writeEvent.GetHandle();
 			payload.readEvent = m_readEvent.GetHandle();
-			payload.communicationHandle = communicationHandle.mh;
+			payload.communicationHandle = communicationHandle.ToU64();
 			payload.communicationOffset = communicationOffset;
 			payload.rulesIndex = m_startInfo.rules->index;
 			payload.version = ProcessMessageVersion;
@@ -2347,7 +2347,7 @@ namespace uba
 				processVar.Append("UBA_SESSION_PROCESS=").AppendValue(getpid());
 
 
-				comIdVar.Append("UBA_COMID=").AppendValue(communicationHandle.uid).Append('+').AppendValue(communicationOffset);
+				comIdVar.Append("UBA_COMID=").AppendValue(communicationHandle.internal).Append('+').AppendValue(communicationOffset);
 
 				#if UBA_USE_NATIVE_MAC_SEMAPHORES
 				sessionId.Append("UBA_SESSION_ID=").AppendValue(m_session.GetSessionId());
diff --git a/Engine/Source/Programs/UnrealBuildAccelerator/Common/Private/UbaSession.cpp b/Engine/Source/Programs/UnrealBuildAccelerator/Common/Private/UbaSession.cpp
--- a/Engine/Source/Programs/UnrealBuildAccelerator/Common/Private/UbaSession.cpp
+++ b/Engine/Source/Programs/UnrealBuildAccelerator/Common/Private/UbaSession.cpp
@@ -1338,10 +1338,10 @@ namespace uba
 		m_systemPath.Append(TCV("/nonexistingpath"));
 		#endif
 
-		m_fileMappingTableHandle = uba::CreateMemoryMappingW(m_logger, PAGE_READWRITE, FileMappingTableMemSize, nullptr, TC("FileMappings"));
+		m_fileMappingTableHandle = uba::CreateMemoryMappingW(m_logger, PAGE_READWRITE, FileMappingTableMaxSize, nullptr, TC("FileMappings"));
 		UBA_ASSERT(m_fileMappingTableHandle.IsValid());
-		m_fileMappingTableMem = MapViewOfFile(m_logger, m_fileMappingTableHandle, FILE_MAP_WRITE, 0, FileMappingTableMemSize, m_allowDiscardVirtualMemory);
-		LockMemory(m_logger, m_fileMappingTableMem, FileMappingTableMemSize, TC("FileMappings"));
+		m_fileMappingTableMem = MapViewOfFile(m_logger, m_fileMappingTableHandle, FILE_MAP_WRITE, 0, FileMappingTableMaxSize, m_allowDiscardVirtualMemory);
+		LockMemory(m_logger, m_fileMappingTableMem, FileMappingTableMaxSize, TC("FileMappings"));
 		UBA_ASSERT(m_fileMappingTableMem);
 		m_fileMappingTableLookup.reserve(70000);
 
@@ -1500,7 +1500,7 @@ namespace uba
 			if (i.second.canBeFreed)
 				m_sharedMemory.CloseHandle(m_logger, i.second.memoryHandle, TC("FileMappingKeptFromOutput"));
 
-		UnmapViewOfFile(m_logger, m_fileMappingTableMem, FileMappingTableMemSize, TC("FileMappingTable"), m_allowDiscardVirtualMemory);
+		UnmapViewOfFile(m_logger, m_fileMappingTableMem, FileMappingTableMaxSize, TC("FileMappingTable"), m_allowDiscardVirtualMemory);
 		CloseFileMapping(m_logger, m_fileMappingTableHandle, TC("FileMappingTable"));
 
 		//#if !UBA_DEBUG
diff --git a/Engine/Source/Programs/UnrealBuildAccelerator/Common/Private/UbaSharedMemoryAllocator.cpp b/Engine/Source/Programs/UnrealBuildAccelerator/Common/Private/UbaSharedMemoryAllocator.cpp
--- a/Engine/Source/Programs/UnrealBuildAccelerator/Common/Private/UbaSharedMemoryAllocator.cpp
+++ b/Engine/Source/Programs/UnrealBuildAccelerator/Common/Private/UbaSharedMemoryAllocator.cpp
@@ -3,6 +3,7 @@
 #include "UbaSharedMemoryAllocator.h"
 #include "UbaBinaryReaderWriter.h"
 #include "UbaFile.h"
+#include "UbaFileMapping.h"
 #include "UbaHash.h"
 #include "UbaTrace.h"
 
@@ -490,41 +491,18 @@ namespace uba
 
 		auto& kernelStats = KernelStats::GetCurrent();
 
-		u8* memory;
-		{
-			TimerScope ts(kernelStats.virtualAlloc2);
-			#if PLATFORM_WINDOWS
-			memory = (u8*)VirtualAlloc2(GetCurrentProcess(), nullptr, committed, MEM_RESERVE | MEM_RESERVE_PLACEHOLDER, PAGE_NOACCESS, nullptr, 0);
-			#else
-			memory = (u8*)mmap(NULL, committed, PROT_NONE, MAP_PRIVATE | MAP_ANONYMOUS, -1, 0);
-			#endif
-		}
-		UBA_ASSERT(memory);
+		TimerScope ts(kernelStats.virtualAlloc2);
+		u8* memory = (u8*)uba::ReserveMemory(nullptr, committed);
+		ts.Leave();
 
 		u64 pos = 0;
 		for (Slice slice : h.slices)
 		{
-			u8* targetAddress = memory + pos;
-
-			#if PLATFORM_WINDOWS
-			if (pos + slice.size < committed)
-			{
-				TimerScope ts(kernelStats.virtualFree);
-				bool res = VirtualFree(targetAddress, slice.size, MEM_RELEASE | MEM_PRESERVE_PLACEHOLDER);
-				UBA_ASSERTF(res, TC("VirtualFree failed for pos %llu with size %llu (%s)"), pos, slice.size, LastErrorToText().data); (void)res;
-			}
-			#endif
-
 			kernelStats.mapViewOfFile3.bytes += slice.size;
-			TimerScope ts(kernelStats.mapViewOfFile3);
-
-			#if PLATFORM_WINDOWS
-			void* viewA = MapViewOfFile3(m_fileMapping.mh, GetCurrentProcess(), targetAddress, slice.offset, slice.size, MEM_REPLACE_PLACEHOLDER, protect, nullptr, 0);
-			#else
-			void* viewA = mmap(targetAddress, slice.size, protect, MAP_SHARED | MAP_FIXED, m_fileMapping.shmFd, slice.offset);
-			#endif
+			TimerScope ts2(kernelStats.mapViewOfFile3);
 
-			UBA_ASSERTF(viewA == targetAddress, TC("MapViewOfFile3 failed in MapView 0x%p %llu for % s (%s)"), memory, pos, h.GetHint(), LastErrorToText().data); (void)viewA;
+			void* viewA = MapMemory(m_fileMapping, (MemoryMapType)mapType, memory, pos, committed, slice.offset, slice.size);
+			UBA_ASSERTF(viewA == memory+pos, TC("MapViewOfFile3 failed in MapView 0x%p %llu for % s (%s)"), memory, pos, h.GetHint(), LastErrorToText().data); (void)viewA;
 			pos += slice.size;
 		}
 
diff --git a/Engine/Source/Programs/UnrealBuildAccelerator/Common/Private/UbaTrace.cpp b/Engine/Source/Programs/UnrealBuildAccelerator/Common/Private/UbaTrace.cpp
--- a/Engine/Source/Programs/UnrealBuildAccelerator/Common/Private/UbaTrace.cpp
+++ b/Engine/Source/Programs/UnrealBuildAccelerator/Common/Private/UbaTrace.cpp
@@ -144,7 +144,7 @@ namespace uba
 		{
 			#if !PLATFORM_WINDOWS
 			StringBuffer<64> uidName;
-			GetMappingHandleName(uidName, m_memoryHandle.uid);
+			GetMappingHandleName(uidName, m_memoryHandle.internal);
 			namedTrace = uidName.data;
 			#endif
 
diff --git a/Engine/Source/Programs/UnrealBuildAccelerator/Common/Public/UbaEnvironment.h b/Engine/Source/Programs/UnrealBuildAccelerator/Common/Public/UbaEnvironment.h
--- a/Engine/Source/Programs/UnrealBuildAccelerator/Common/Public/UbaEnvironment.h
+++ b/Engine/Source/Programs/UnrealBuildAccelerator/Common/Public/UbaEnvironment.h
@@ -29,4 +29,7 @@ namespace uba
 	bool GetMemoryInfo(Logger& logger, u64& outAvailable, u64& outTotal, u64* outMaxPageFileSize = nullptr);
 	void GetSystemInfo(Logger& logger, StringBufferBase& out);
 	bool GetCpuTime(u64& outTotalTime, u64& outIdleTime);
+
+	void SetMemoryWorkingSet(Logger& logger, u64 minWorkingSet);
+	void LockMemory(Logger& logger, void* dest, u64 size, const tchar* hint);
 }
\ No newline at end of file
diff --git a/Engine/Source/Programs/UnrealBuildAccelerator/Common/Public/UbaFile.h b/Engine/Source/Programs/UnrealBuildAccelerator/Common/Public/UbaFile.h
--- a/Engine/Source/Programs/UnrealBuildAccelerator/Common/Public/UbaFile.h
+++ b/Engine/Source/Programs/UnrealBuildAccelerator/Common/Public/UbaFile.h
@@ -2,6 +2,7 @@
 
 #pragma once
 
+#include "UbaFileHandle.h"
 #include "UbaLogger.h"
 #include "UbaMemory.h"
 #include "UbaPlatform.h"
@@ -38,8 +39,6 @@ namespace uba
 	#endif
 	inline constexpr u32 FILE_SHARE_ALL = 0x00000007;
 
-	inline constexpr u64 FileHandleFlagMask = 0x0000'0000'ffff'ffff;
-
 	#if PLATFORM_WINDOWS
 	inline constexpr u64 OverlappedIoFlag = 0x0000'0001'0000'0000; // Only used by windows
 	#endif
diff --git a/Engine/Source/Programs/UnrealBuildAccelerator/Common/Public/UbaFileMapping.h b/Engine/Source/Programs/UnrealBuildAccelerator/Common/Public/UbaFileMapping.h
deleted file mode 100644
diff --git a/Engine/Source/Programs/UnrealBuildAccelerator/Common/Public/UbaFixedSizeFileMappingAllocator.h b/Engine/Source/Programs/UnrealBuildAccelerator/Common/Public/UbaFixedSizeFileMappingAllocator.h
new file mode 100644
--- /dev/null
+++ b/Engine/Source/Programs/UnrealBuildAccelerator/Common/Public/UbaFixedSizeFileMappingAllocator.h
@@ -0,0 +1,53 @@
+// Copyright Epic Games, Inc. All Rights Reserved.
+
+#pragma once
+
+#include "UbaFileMappingHandle.h"
+#include "UbaMemory.h"
+
+namespace uba
+{
+	class Logger;
+
+	////////////////////////////////////////////////////////////////////////////////////////////////////
+
+	struct FileMappingAllocation
+	{
+		FileMappingHandle handle;
+		u64 offset = 0;
+		u8* memory = nullptr;
+	};
+
+	////////////////////////////////////////////////////////////////////////////////////////////////////
+
+	class FixedSizeFileMappingAllocator
+	{
+	public:
+		FixedSizeFileMappingAllocator(Logger& logger, const tchar* name);
+		~FixedSizeFileMappingAllocator();
+
+		bool Init(u64 blockSize, u64 capacity);
+
+		FileMappingAllocation Alloc(const tchar* hint);
+		void Free(FileMappingAllocation allocation);
+
+		u64 GetSize();
+
+	private:
+		Logger& m_logger;
+		const tchar* m_name;
+		u64 m_capacity = 0;
+		u64 m_blockSize = 0;
+
+		Futex m_mappingLock;
+		FileMappingHandle m_mappingHandle;
+		u64 m_mappingCount = 0;
+
+		Set<u64> m_availableBlocks;
+
+		FixedSizeFileMappingAllocator(const FixedSizeFileMappingAllocator&) = delete;
+		void operator=(const FixedSizeFileMappingAllocator&) = delete;
+	};
+
+	////////////////////////////////////////////////////////////////////////////////////////////////////
+}
\ No newline at end of file
diff --git a/Engine/Source/Programs/UnrealBuildAccelerator/Common/Public/UbaSession.h b/Engine/Source/Programs/UnrealBuildAccelerator/Common/Public/UbaSession.h
--- a/Engine/Source/Programs/UnrealBuildAccelerator/Common/Public/UbaSession.h
+++ b/Engine/Source/Programs/UnrealBuildAccelerator/Common/Public/UbaSession.h
@@ -4,7 +4,7 @@
 
 #include "UbaBottleneck.h"
 #include "UbaDependencyCrawler.h"
-#include "UbaFile.h"
+#include "UbaFixedSizeFileMappingAllocator.h"
 #include "UbaDefinitions.h"
 #include "UbaDirectoryTableHolder.h"
 #include "UbaFileMappingBuffer.h"
@@ -20,6 +20,7 @@
 
 namespace uba
 {
+	class FileAccessor;
 	class Process;
 	class ProcessHandle;
 	class ProcessImpl;
diff --git a/Engine/Source/Programs/UnrealBuildAccelerator/Common/Public/UbaSharedMemoryAllocator.h b/Engine/Source/Programs/UnrealBuildAccelerator/Common/Public/UbaSharedMemoryAllocator.h
--- a/Engine/Source/Programs/UnrealBuildAccelerator/Common/Public/UbaSharedMemoryAllocator.h
+++ b/Engine/Source/Programs/UnrealBuildAccelerator/Common/Public/UbaSharedMemoryAllocator.h
@@ -3,6 +3,7 @@
 #pragma once
 
 #include "UbaDefinitions.h"
+#include "UbaFileHandle.h"
 #include "UbaHash.h"
 #include "UbaSharedMemoryView.h"
 
diff --git a/Engine/Source/Programs/UnrealBuildAccelerator/Core/Private/UbaSharedMemoryView.cpp b/Engine/Source/Programs/UnrealBuildAccelerator/Core/Private/UbaSharedMemoryView.cpp
--- a/Engine/Source/Programs/UnrealBuildAccelerator/Core/Private/UbaSharedMemoryView.cpp
+++ b/Engine/Source/Programs/UnrealBuildAccelerator/Core/Private/UbaSharedMemoryView.cpp
@@ -2,6 +2,7 @@
 
 #include "UbaSharedMemoryView.h"
 #include "UbaBinaryReaderWriter.h"
+#include "UbaFileMapping.h"
 #include "UbaHash.h"
 #include "UbaMemory.h"
 #include "UbaProcessStats.h"
@@ -114,20 +115,9 @@ namespace uba
 	bool SharedMemoryView::InternalInit(SharedMemoryAllocatorHandle handle, u64 capacity, void* baseAddress, const tchar* hint)
 	{
 		UBA_ASSERT(!m_memory);
-
 		m_capacity = AlignUp(capacity, PageSize);
-
 		TimerScope ts(KernelStats::GetCurrent().virtualAlloc2);
-
-		#if PLATFORM_WINDOWS
-		m_memory = (u8*)VirtualAlloc2(GetCurrentProcess(), baseAddress, m_capacity, MEM_RESERVE | MEM_RESERVE_PLACEHOLDER, PAGE_NOACCESS, nullptr, 0);
-		#else
-		m_memory = (u8*)mmap(NULL, m_capacity, PROT_NONE, MAP_PRIVATE | MAP_ANONYMOUS, -1, 0);
-		if (m_memory == MAP_FAILED)
-			m_memory = nullptr;
-		#endif
-
-		UBA_ASSERTF(m_memory, TC("VirtualAlloc2 failed (%s)"), LastErrorToText().data);
+		m_memory = (u8*)ReserveMemory(baseAddress, m_capacity);
 		return m_memory != nullptr;
 	}
 
@@ -173,8 +163,6 @@ namespace uba
 		UBA_ASSERTF(m_capacity != ~0ull, TC("Can't extend memory on a view directly mapped on to allocator mapping (%s)"), hint);
 		UBA_ASSERTF(AlignUp(offset, PageSize) == offset, TC("ExtendMemory got offset %llu parameter which is not aligned (%s)"), offset, hint);
 
-		u8* targetAddress = m_memory + m_mappedSize;
-
 		u64 left = m_capacity - m_mappedSize;
 		UBA_ASSERTF(left >= size, TC("Want to extend memory more than capacity (Capacity: %llu, Committed: %llu, Requested: %llu)"), m_capacity, m_mappedSize, size);
 		if (left < size)
@@ -182,39 +170,9 @@ namespace uba
 
 		auto& kernelStats = KernelStats::GetCurrent();
 
-		#if PLATFORM_WINDOWS
-		u32 protect = PAGE_READWRITE;
-		if (mapType == SharedMemoryMapType_ReadOnly)
-			protect = PAGE_READONLY;
-		else if (mapType == SharedMemoryMapType_CopyOnWrite)
-			protect = PAGE_WRITECOPY;
-
-		size = AlignUp(size, 4*1024); // This is needed for VirtualFree to align properly
-
-		if (u64 toRelease = m_capacity - m_mappedSize - size)
-		{
-			TimerScope ts(kernelStats.virtualFree);
-			bool res = VirtualFree(targetAddress + size, toRelease, MEM_RELEASE | MEM_PRESERVE_PLACEHOLDER);
-			UBA_ASSERTF(res, TC("VirtualFree for placeholder failed releasing %llu at %llu (Capacity: %llu MappedSize: %llu) - %s (%s)"), toRelease, size, m_capacity, m_mappedSize, LastErrorToText().data, hint);
-			if (!res)
-				return false;
-		}
-
-		kernelStats.mapViewOfFile3.bytes += size;
-		TimerScope ts(kernelStats.mapViewOfFile3);
-		void* viewA = MapViewOfFile3(handle.mh, GetCurrentProcess(), targetAddress, offset, size, MEM_REPLACE_PLACEHOLDER, protect, nullptr, 0);
-		#else
-		UBA_ASSERT(mapType != SharedMemoryMapType_CopyOnWrite);
-		u32 protect = PROT_READ;
-		if (mapType == SharedMemoryMapType_ReadWrite)
-			protect |= PROT_WRITE;
 		TimerScope ts(kernelStats.mapViewOfFile3);
-		void* viewA = mmap(targetAddress, size, protect, MAP_SHARED | MAP_FIXED, handle.shmFd, offset);
-		if (viewA == MAP_FAILED)
-			viewA = nullptr;
-		#endif
-
-		UBA_ASSERTF(viewA == targetAddress, TC("MapViewOfFile3 failed extending memory with %llu bytes at offset %llu - %s (%s)"), size, offset, LastErrorToText().data, hint);
+		void* viewA = MapMemory(handle, (MemoryMapType)mapType, m_memory, m_mappedSize, m_capacity, offset, size);
+		UBA_ASSERTF(viewA == m_memory + m_mappedSize, TC("MapViewOfFile3 failed extending memory with %llu bytes at offset %llu - %s (%s)"), size, offset, LastErrorToText().data, hint);
 		if (!viewA)
 			return false;
 
diff --git a/Engine/Source/Programs/UnrealBuildAccelerator/Core/Public/UbaDetoursPayload.h b/Engine/Source/Programs/UnrealBuildAccelerator/Core/Public/UbaDetoursPayload.h
--- a/Engine/Source/Programs/UnrealBuildAccelerator/Core/Public/UbaDetoursPayload.h
+++ b/Engine/Source/Programs/UnrealBuildAccelerator/Core/Public/UbaDetoursPayload.h
@@ -15,7 +15,7 @@ namespace uba
 		HANDLE cancelEvent = 0;
 		HANDLE writeEvent = 0;
 		HANDLE readEvent = 0;
-		HANDLE communicationHandle = 0;
+		u64 communicationHandle = 0;
 		u64 communicationOffset = 0;
 		u32 version = 0;
 		u32 rulesIndex = 0;
diff --git a/Engine/Source/Programs/UnrealBuildAccelerator/Core/Public/UbaFileHandle.h b/Engine/Source/Programs/UnrealBuildAccelerator/Core/Public/UbaFileHandle.h
new file mode 100644
--- /dev/null
+++ b/Engine/Source/Programs/UnrealBuildAccelerator/Core/Public/UbaFileHandle.h
@@ -0,0 +1,21 @@
+// Copyright Epic Games, Inc. All Rights Reserved.
+
+#pragma once
+
+#include "UbaBase.h"
+
+namespace uba
+{
+	////////////////////////////////////////////////////////////////////////////////////////////////////
+
+	enum FileHandle : u64 {};
+	inline constexpr FileHandle InvalidFileHandle = (FileHandle)-1;
+
+	inline constexpr u64 FileHandleFlagMask = 0x0000'0000'ffff'ffff;
+
+	#if !PLATFORM_WINDOWS
+	inline int asFileDescriptor(FileHandle fh) { return fh == InvalidFileHandle ? (int)fh : (int)(fh & FileHandleFlagMask); }
+	#endif
+
+	////////////////////////////////////////////////////////////////////////////////////////////////////
+}
diff --git a/Engine/Source/Programs/UnrealBuildAccelerator/Core/Public/UbaFileMapping.h b/Engine/Source/Programs/UnrealBuildAccelerator/Core/Public/UbaFileMapping.h
new file mode 100644
--- /dev/null
+++ b/Engine/Source/Programs/UnrealBuildAccelerator/Core/Public/UbaFileMapping.h
@@ -0,0 +1,45 @@
+// Copyright Epic Games, Inc. All Rights Reserved.
+
+#pragma once
+
+#include "UbaFileMappingHandle.h"
+
+namespace uba
+{
+	enum FileHandle : u64;
+	class Logger;
+
+	////////////////////////////////////////////////////////////////////////////////////////////////////
+
+	#if !PLATFORM_WINDOWS
+	inline constexpr u32 FILE_MAP_WRITE = 0x0002;
+	inline constexpr u32 FILE_MAP_READ = 0x0004;
+	inline constexpr u32 FILE_MAP_ALL_ACCESS = FILE_MAP_WRITE | FILE_MAP_READ;
+	inline constexpr u32 SEC_RESERVE = 0x04000000;
+	inline constexpr u32 PAGE_READONLY = PROT_READ;
+	inline constexpr u32 PAGE_READWRITE = PROT_READ | PROT_WRITE;
+
+	#if PLATFORM_MAC
+	inline constexpr u32 SHM_MAX_FILENAME = PSHMNAMLEN;
+	#else
+	inline constexpr u32 SHM_MAX_FILENAME = 38;
+	#endif
+	#endif
+
+	enum MemoryMapType { MemoryMapType_ReadOnly, MemoryMapType_ReadWrite, MemoryMapType_CopyOnWrite };
+
+	FileMappingHandle CreateMemoryMappingW(Logger& logger, u32 flProtect, u64 maxSize, const tchar* name, const tchar* hint);
+	FileMappingHandle CreateFileMappingW(Logger& logger, FileHandle file, u32 flProtect, u64 maxSize, const tchar* hint);
+	u8* MapViewOfFile(Logger& logger, FileMappingHandle fileMappingObject, u32 dwDesiredAccess, u64 offset, u64 dwNumberOfBytesToMap, bool allowDiscard = false);
+	bool MapViewCommit(Logger& logger, void* address, u64 size, bool allowDiscard = false);
+	bool UnmapViewOfFile(Logger& logger, const void* lpBaseAddress, u64 bytesToUnmap, const tchar* hint, bool allowDiscard = false);
+	bool CloseFileMapping(Logger& logger, FileMappingHandle h, const tchar* hint);
+	void* ReserveMemory(void* baseAddress, u64 capacity);
+	void* MapMemory(FileMappingHandle handle, MemoryMapType mapType, void* targetAddress, u64 targetOffset, u64 targetCapacity, u64 handleOffset, u64 size);
+	void MapMemoryCopy(void* dest, const void* source, u64 size);
+
+	FileMappingHandle DuplicateMappingFromHost(FileMappingHandle handle, const char* hint);
+	u8* GetMemoryFromMapping(FileMappingHandle handle, u64 size, u64 offset, bool writable, const char* hint);
+
+	////////////////////////////////////////////////////////////////////////////////////////////////////
+}
diff --git a/Engine/Source/Programs/UnrealBuildAccelerator/Core/Public/UbaFileMapping.inl b/Engine/Source/Programs/UnrealBuildAccelerator/Core/Public/UbaFileMapping.inl
new file mode 100644
--- /dev/null
+++ b/Engine/Source/Programs/UnrealBuildAccelerator/Core/Public/UbaFileMapping.inl
@@ -0,0 +1,500 @@
+// Copyright Epic Games, Inc. All Rights Reserved.
+
+#pragma once
+
+// No includes in this file!
+
+////////////////////////////////////////////////////////////////////////////////////////////////////
+
+#if PLATFORM_WINDOWS
+HANDLE asHANDLE(FileMappingHandle h) { return (HANDLE)(uintptr_t)u32(h.internal); }
+HANDLE asHANDLE(FileHandle fh) { return (HANDLE)(fh == InvalidFileHandle ? InvalidFileHandle : (fh & FileHandleFlagMask)); }
+Bottleneck& g_createFileHandleBottleneck = *new Bottleneck(8); // Allocated and leaked just to prevent shutdown asserts in debug
+
+HANDLE InternalCreateFileMappingW(Logger& logger, HANDLE hFile, DWORD flProtect, DWORD dwMaximumSizeHigh, DWORD dwMaximumSizeLow, LPCWSTR lpName, const tchar* hint)
+{
+	HANDLE res;
+	if (flProtect == PAGE_READWRITE)
+	{
+		// Experiment to try to prevent lock happening on AWS servers when lots of helpers are sending back obj files.
+		Timer timer;
+		BottleneckScope bs(g_createFileHandleBottleneck, timer);
+		res = ::CreateFileMappingW(hFile, NULL, flProtect, dwMaximumSizeHigh, dwMaximumSizeLow, lpName);
+	}
+	else
+		res = ::CreateFileMappingW(hFile, NULL, flProtect, dwMaximumSizeHigh, dwMaximumSizeLow, lpName);
+
+#if UBA_DEBUG_FILE_MAPPING
+	if (res)
+	{
+		SCOPED_FUTEX(g_fileMappingsLock, lock);
+		UBA_FOR_ASSERT(auto insres =) g_fileMappings.try_emplace(res, hint);
+		UBA_ASSERT(insres.second);
+		//logger.Info(TC("FILEMAPPING CREAT: %s %llu (%llu)"), hint, u64(res), g_fileMappings.size());
+	}
+#endif
+
+	return res;
+}
+static void InternalDiscard(Logger& logger, u8* mem, u64 size, const tchar* from, const tchar* hint)
+{
+#if PLATFORM_WINDOWS
+	u64 stepSize = 16ull*1024*1024*1024; // Extremely large.. here to test if it is faster to do multiple calls
+	u64 left = size;
+	while (left)
+	{
+		u64 toDiscard = Min(left, stepSize);
+		DWORD res = DiscardVirtualMemory(mem, AlignUp(toDiscard, 4*1024));
+		if (res != ERROR_SUCCESS)
+			logger.Warning(TC("%s - Failed to discard memory (%s) (%u)"), from, hint, res);
+		mem += toDiscard;
+		left -= toDiscard;
+	}
+#endif
+}
+#else
+Atomic<u64> g_mappingUidCounter = 0;
+static pthread_mutex_t g_mappingUidCounterLock = PTHREAD_MUTEX_INITIALIZER;
+void AppendPrefixAndValue(char* out, const char* prefix, u64 value)
+{
+	strcpy(out, prefix);
+	out += strlen(prefix);
+	constexpr char hexChars[] = "0123456789abcdef";
+	for (int i=0;i!=8;++i)
+	{
+		*out++ = hexChars[(value >> 4) & 0xf];
+		*out++ = hexChars[value & 0xf];
+		value = value >> 8;
+		if (!value)
+			break;
+	}
+	*out = 0;
+}
+void GetMappingHandleName(char* out, u64 uid)
+{
+	#if PLATFORM_MAC
+	if (!IsRunningDarling())
+	{
+		memcpy(out, "/tmp", 4);
+		out += 4;
+	}
+	#endif
+	AppendPrefixAndValue(out, "/uba_", uid);
+}
+bool GarbageCollectFileMappings(Logger& logger)
+{
+	// Since we need to not leak file mappings we use files as a trick to know which ones are used and not
+	char filePath[128];
+	strcpy(filePath, "/tmp/uba_shm_locks");
+
+	// Create dir
+	if (mkdir(filePath, 0777) == -1)
+		if (errno != EEXIST)
+		{
+			logger.Error(TC("Failed to create %s for memory mapping (%s)"), filePath, strerror(errno));
+			return false;
+		}
+
+	// Clear out all orphaned shm_open
+	if (DIR* dir = opendir(filePath))
+	{
+		auto pathLen = strlen(filePath);
+		filePath[pathLen++] = '/';
+		while (true)
+		{
+			errno = 0;
+			struct dirent* pDirent = readdir(dir);
+			if (!pDirent)
+			{
+				if (errno == 0)
+					break;
+
+				// This is actually an error.. should return false?
+				break;
+			}
+
+			char* fileName = pDirent->d_name;
+			if (fileName[0] == '.' && (fileName[1] == 0 || (fileName[1] == '.' && fileName[2] == 0))) 
+				continue;
+
+			u32 uid = strtoul(fileName, nullptr, 10);
+
+			strcpy(filePath+pathLen, fileName);
+			int lockFd = open(filePath, O_RDWR, S_IRUSR | S_IWUSR | O_CLOEXEC);
+			if (lockFd == -1)
+			{
+				if (errno == EPERM)
+				{
+					g_mappingUidCounter = uid;
+					return true;
+				}
+				logger.Warning("Failed to open %s for memory mapping (%s)", filePath, strerror(errno));
+				UBA_ASSERTF(false, "Failed to open %s (%s)", filePath, strerror(errno));
+				return false;
+			}
+
+			if (flock(lockFd, LOCK_EX | LOCK_NB) == 0)
+			{
+				char uidName[64];
+				GetMappingHandleName(uidName, uid);
+				if (shm_unlink(uidName) == 0)
+					;// logger.Info("Removed old shared memory %s", uidName.data);
+				remove(filePath);
+			}
+			else
+			{
+				g_mappingUidCounter = uid;
+			}
+			close(lockFd);
+		}
+		closedir(dir);
+	}
+	if (g_mappingUidCounter)
+		logger.Info("Starting shared memory files at %u", g_mappingUidCounter.load());
+	return true;
+}
+#endif
+
+FileMappingHandle CreateMemoryMappingW(Logger& logger, u32 flProtect, u64 maxSize, const tchar* name, const tchar* hint)
+{
+#if PLATFORM_WINDOWS
+	ExtendedTimerScope ts(KernelStats::GetCurrent().createFileMapping);
+	return u64(InternalCreateFileMappingW(logger, INVALID_HANDLE_VALUE, flProtect, (DWORD)ToHigh(maxSize), ToLow(maxSize), name, hint));
+#else
+	UBA_ASSERT(!name);
+	UBA_ASSERT((flProtect & (~u32(PAGE_READWRITE | SEC_RESERVE))) == 0);
+
+	pthread_mutex_lock(&g_mappingUidCounterLock);
+	if (!g_mappingUidCounter)
+		GarbageCollectFileMappings(logger);
+	pthread_mutex_unlock(&g_mappingUidCounterLock);
+
+	// Let's find a free shm name
+	u64 uid;
+	int shmFd;
+	int lockFd;
+	char lockFilePath[128];
+
+	int retryCount = 4;
+	while (true)
+	{
+		uid = ++g_mappingUidCounter;
+
+		AppendPrefixAndValue(lockFilePath, "/tmp/uba_shm_locks/", uid);
+
+		lockFd = open(lockFilePath, O_CREAT | O_RDWR | O_NOFOLLOW | O_EXCL, S_IRUSR | S_IWUSR | O_CLOEXEC);
+		if (lockFd == -1)
+		{
+			if (errno == EEXIST)
+				continue;
+			logger.Warning(TC("Failed to open/create %s (%s)"), lockFilePath, strerror(errno));
+			UBA_ASSERTF(false, "Failed to open/create %s (%s)", lockFilePath, strerror(errno));
+			continue;
+		}
+
+		if (flock(lockFd, LOCK_EX | LOCK_NB) != 0) // Some other process is using this one
+		{
+			close(lockFd);
+			continue;
+		}
+
+		char uidName[64];
+		GetMappingHandleName(uidName, uid);
+
+		int oflags = O_CREAT | O_RDWR | O_NOFOLLOW | O_EXCL;
+		shmFd = shm_open(uidName, oflags, S_IRUSR | S_IWUSR);
+		if (shmFd != -1)
+			break;
+
+		bool retry = retryCount-- > 0;
+
+		int err = errno;
+		auto logType = retry ? LogEntryType_Warning : LogEntryType_Error;
+		logger.Logf(logType, TC("Failed to create shm %s after getting lock-file %s (%s)"), uidName, lockFilePath, strerror(err));
+
+		if (retry)
+			continue;
+
+		remove(lockFilePath);
+		close(lockFd);
+		SetLastError(err);
+		return {};
+	}
+
+	if (maxSize != 0)
+	{
+		if (ftruncate(shmFd, (s64)maxSize) == -1)
+		{
+			SetLastError(errno);
+			close(shmFd);
+			remove(lockFilePath);
+			close(lockFd);
+			//UBA_ASSERTF(false, "Failed to truncate file mapping '%s' to size %llu (%s)", name, maxSize, strerror(errno));
+			return {};
+		}
+	}
+	SetLastError(0);
+	FileMappingHandle h(uid);
+	h.shmFd = shmFd;
+	h.lockFd = lockFd;
+	return h;
+#endif
+}
+
+FileMappingHandle CreateFileMappingW(Logger& logger, FileHandle file, u32 protect, u64 maxSize, const tchar* hint)
+{
+#if PLATFORM_WINDOWS
+	ExtendedTimerScope ts(KernelStats::GetCurrent().createFileMapping);
+	return u64(InternalCreateFileMappingW(logger, asHANDLE(file), protect, (DWORD)ToHigh(maxSize), ToLow(maxSize), NULL, hint));
+#else
+	FileMappingHandle h;
+	int fd = asFileDescriptor(file);
+	if (maxSize && (protect & (~PAGE_READONLY)) != 0)
+	{
+#if PLATFORM_MAC // For some reason lseek+write does not work on apple silicon platform
+		if (ftruncate(fd, maxSize) == -1)
+		{
+			logger.Error("ftruncate to %llu on fd %i failed for %s: %s\n", maxSize, fd, hint, strerror(errno));
+			return h;
+		}
+#else
+		if (lseek(fd, maxSize - 1, SEEK_SET) != maxSize - 1)
+		{
+			logger.Error("lseek to %llu failed for %s: %s\n", maxSize - 1, hint, strerror(errno));
+			return h;
+		}
+
+		errno = 0;
+		int res = write(fd, "", 1);
+		if (res != 1)
+		{
+			logger.Error("write one byte at %llu on fd %i (%s) failed (res: %i): %s\n", maxSize - 1, fd, hint, res, strerror(errno));
+			return h;
+		}
+#endif
+	}
+
+	h.shmFd = fd;
+	h.internal = ~0ull;
+	return h;
+#endif
+}
+
+u8* MapViewOfFile(Logger& logger, FileMappingHandle fileMappingObject, u32 desiredAccess, u64 offset, u64 bytesToMap, bool allowDiscard)
+{
+#if PLATFORM_WINDOWS
+	ExtendedTimerScope ts(KernelStats::GetCurrent().mapViewOfFile);
+	u8* res = (u8*)::MapViewOfFile(asHANDLE(fileMappingObject), desiredAccess, (DWORD)ToHigh(offset), ToLow(offset), bytesToMap);
+	if (!res)
+		return nullptr;
+	if (allowDiscard)
+		InternalDiscard(logger, res, bytesToMap, TC("MapViewOfFile"), TC(""));
+#else
+	int prot = 0;
+	if (desiredAccess & FILE_MAP_READ)
+		prot |= PROT_READ;
+	if (desiredAccess & FILE_MAP_WRITE)
+		prot |= PROT_WRITE;
+	UBA_ASSERT(fileMappingObject.IsValid());
+	int shmFd = fileMappingObject.shmFd;
+	void* rptr = mmap(NULL, bytesToMap, prot, MAP_SHARED, shmFd, s64(offset));
+	if (rptr == MAP_FAILED)
+	{
+		SetLastError(errno);
+		return nullptr;
+	}
+	u8* res = (u8*)rptr;
+	//UBA_ASSERTF(false, "Failed to map file with fd %i, desiredAccess %u offset %llu, bytesToMap %llu (%s)", fd, desiredAccess, offset, bytesToMap, strerror(errno));
+#endif
+
+#if UBA_DEBUG_FILE_MAPPING
+	{
+		SCOPED_FUTEX(g_fileMappingsLock, lock);
+		auto findIt = g_fileMappings.find(fileMappingObject.mh);
+		UBA_ASSERTF(findIt != g_fileMappings.end(), TC("Mapping nonexisting handle: %llu"), u64(fileMappingObject.mh));
+		auto insres = g_viewMappings.try_emplace(res, fileMappingObject.mh);
+		UBA_ASSERT(insres.second);
+		++findIt->second.viewCount;
+	}
+#endif
+
+	return res;
+}
+
+bool MapViewCommit(Logger& logger, void* address, u64 size, bool allowDiscard)
+{
+#if PLATFORM_WINDOWS
+	ExtendedTimerScope ts(KernelStats::GetCurrent().virtualAlloc);
+	if (!::VirtualAlloc(address, size, MEM_COMMIT, PAGE_READWRITE))
+		return false;
+	if (allowDiscard)
+		InternalDiscard(logger, (u8*)address, size, TC("MapViewCommit"), TC(""));
+	return true;
+#else
+	return true;
+#endif
+}
+
+bool UnmapViewOfFile(Logger& logger, const void* address, u64 bytesToUnmap, const tchar* hint, bool allowDiscard)
+{
+	if (!address)
+		return true;
+
+#if UBA_DEBUG_FILE_MAPPING
+	{
+		SCOPED_FUTEX(g_fileMappingsLock, lock);
+		auto findIt = g_viewMappings.find(address);
+		UBA_ASSERT(findIt != g_viewMappings.end());
+		auto findIt2 = g_fileMappings.find(findIt->second);
+		g_viewMappings.erase(findIt);
+		UBA_ASSERTF(findIt2 != g_fileMappings.end(), TC("Unmap nonexisting handle: %llu"), u64(findIt->second));
+		--findIt2->second.viewCount;
+	}
+#endif
+
+#if PLATFORM_WINDOWS
+	ExtendedTimerScope ts(KernelStats::GetCurrent().unmapViewOfFile);
+	if (allowDiscard)
+		InternalDiscard(logger, (u8*)address, bytesToUnmap, TC("UnmapViewOfFile"), hint);
+	return ::UnmapViewOfFile(address);
+#else
+	UBA_ASSERTF(bytesToUnmap, TC("bytesToUnmap is zero unmapping %p (%s)"), address, hint);
+	if (munmap((void*)address, bytesToUnmap) == 0)
+		return true;
+	UBA_ASSERT(false);
+	return false;
+#endif
+}
+
+bool CloseFileMapping(Logger& logger, FileMappingHandle h, const tchar* hint)
+{
+	if (!h.IsValid())
+		return true;
+
+#if UBA_DEBUG_FILE_MAPPING
+	SCOPED_FUTEX(g_fileMappingsLock, lock);
+	auto findIt = g_fileMappings.find(h.mh);
+	UBA_ASSERTF(findIt != g_fileMappings.end(), TC("Handle: %llu"), u64(h.mh));
+	UBA_ASSERTF(!findIt->second.viewCount.load(), TC("Closing file mapping with %llu open views (%s)"), findIt->second.viewCount.load(), findIt->second.hint.c_str());
+	logger.Info(TC("FILEMAPPING CLOSE: %s %llu (%llu)"), findIt->second.hint.c_str(), u64(h.mh), g_fileMappings.size() - 1);
+	g_fileMappings.erase(findIt);
+#endif
+
+#if PLATFORM_WINDOWS
+	return ::CloseHandle(asHANDLE(h));
+#else
+	if (h.shmFd != -1 && h.internal != ~0ull)
+		if (close(h.shmFd) != 0)
+			UBA_ASSERTF(false, "Failed to close file mapping (%s)", strerror(errno));
+
+	if (h.internal != 0 && h.internal != ~0ull)
+	{
+		char uidName[128];
+		GetMappingHandleName(uidName, h.internal);
+		if (shm_unlink(uidName) != 0)
+		{
+			SetLastError(errno);
+			UBA_ASSERTF(false, "Failed to unlink %s (%s)", uidName, strerror(errno));
+			return false;
+		}
+
+		char lockFile[128];
+		AppendPrefixAndValue(lockFile, "/tmp/uba_shm_locks/", h.internal);
+		remove(lockFile);
+	}
+	if (h.lockFd != -1)
+		close(h.lockFd);
+	return true;
+#endif
+}
+
+void* ReserveMemory(void* baseAddress, u64 capacity)
+{
+#if PLATFORM_WINDOWS
+	void* memory = VirtualAlloc2(GetCurrentProcess(), baseAddress, capacity, MEM_RESERVE | MEM_RESERVE_PLACEHOLDER, PAGE_NOACCESS, nullptr, 0);
+#else
+	void* memory = mmap(NULL, capacity, PROT_NONE, MAP_PRIVATE | MAP_ANONYMOUS, -1, 0);
+	if (memory == MAP_FAILED)
+		memory = nullptr;
+#endif
+	UBA_ASSERTF(memory, TC("VirtualAlloc2 failed (%s)"), LastErrorToText().data);
+	return memory;
+}
+
+void* MapMemory(FileMappingHandle handle, MemoryMapType mapType, void* targetAddress, u64 targetOffset, u64 targetCapacity, u64 handleOffset, u64 size)
+{
+#if PLATFORM_WINDOWS
+	u32 protect = PAGE_READWRITE;
+	if (mapType == MemoryMapType_ReadOnly)
+		protect = PAGE_READONLY;
+	else if (mapType == MemoryMapType_CopyOnWrite)
+		protect = PAGE_WRITECOPY;
+
+	size = AlignUp(size, 4*1024); // This is needed for VirtualFree to align properly
+
+	if (u64 toRelease = targetCapacity - targetOffset - size)
+	{
+		//TimerScope ts(kernelStats.virtualFree);
+		bool res = VirtualFree((u8*)targetAddress + targetOffset + size, toRelease, MEM_RELEASE | MEM_PRESERVE_PLACEHOLDER);
+		//UBA_ASSERTF(res, TC("VirtualFree for placeholder failed releasing %llu at %llu (Capacity: %llu MappedSize: %llu) - %s (%s)"), toRelease, size, m_capacity, m_mappedSize, LastErrorToText().data, hint);
+		if (!res)
+			return nullptr;//false;
+	}
+
+	//kernelStats.mapViewOfFile3.bytes += size;
+	//TimerScope ts(kernelStats.mapViewOfFile3);
+	void* viewA = MapViewOfFile3(asHANDLE(handle), GetCurrentProcess(), (u8*)targetAddress+targetOffset, handleOffset, size, MEM_REPLACE_PLACEHOLDER, protect, nullptr, 0);
+#else
+	UBA_ASSERT(mapType != MemoryMapType_CopyOnWrite);
+	u32 protect = PROT_READ;
+	if (mapType == MemoryMapType_ReadWrite)
+		protect |= PROT_WRITE;
+	//TimerScope ts(kernelStats.mapViewOfFile3);
+	void* viewA = mmap((u8*)targetAddress+targetOffset, size, protect, MAP_SHARED | MAP_FIXED, handle.shmFd, handleOffset);
+	if (viewA == MAP_FAILED)
+		viewA = nullptr;
+#endif
+	return viewA;
+}
+
+FileMappingHandle DuplicateMappingFromHost(FileMappingHandle handle, const char* hint)
+{
+#if PLATFORM_WINDOWS
+	HANDLE realHandle;
+	if (!DuplicateHandle(g_hostProcess, asHANDLE(handle), GetCurrentProcess(), &realHandle, 0, FALSE, DUPLICATE_SAME_ACCESS))
+	{
+		UBA_ASSERTF(false, L"Failed to duplicate %S handle (%u)", hint, GetLastError());
+		return {};
+	}
+	return FileMappingHandle(u64(realHandle));
+#else
+	return handle;
+#endif
+}
+
+u8* GetMemoryFromMapping(FileMappingHandle handle, u64 size, u64 offset, bool writable, const char* hint)
+{
+#if PLATFORM_WINDOWS
+	DWORD dwDesiredAccess = FILE_MAP_READ;
+	if (writable)
+		dwDesiredAccess |= FILE_MAP_WRITE;
+	LARGE_INTEGER li;
+	li.QuadPart = (LONGLONG)offset;
+	//TimerScope ts(g_kernelStats.mapViewOfFile);
+	u8* mem = (u8*)::MapViewOfFile(asHANDLE(handle), dwDesiredAccess, li.HighPart, li.LowPart, size);
+	UBA_ASSERT(mem);
+	::CloseHandle(asHANDLE(handle));
+	return mem;
+#else
+	char uidName[128];
+	GetMappingHandleName(uidName, handle.internal);
+	int shmFd = shm_open(uidName, O_RDONLY, S_IRUSR | S_IWUSR);
+	UBA_ASSERTF(shmFd != -1, "shm_open failed for %s and handle %s (%s)", hint, uidName, strerror(errno));
+	u8* mem = (u8*)mmap(NULL, size, PROT_READ, MAP_SHARED, shmFd, 0);
+	UBA_ASSERTF(mem != MAP_FAILED, "mmap failed (%s)", strerror(errno));
+	close(shmFd);
+	return mem;
+#endif
+}
+
+////////////////////////////////////////////////////////////////////////////////////////////////////
diff --git a/Engine/Source/Programs/UnrealBuildAccelerator/Core/Public/UbaFileMappingHandle.h b/Engine/Source/Programs/UnrealBuildAccelerator/Core/Public/UbaFileMappingHandle.h
new file mode 100644
--- /dev/null
+++ b/Engine/Source/Programs/UnrealBuildAccelerator/Core/Public/UbaFileMappingHandle.h
@@ -0,0 +1,27 @@
+// Copyright Epic Games, Inc. All Rights Reserved.
+
+#pragma once
+
+#include "UbaBase.h"
+
+namespace uba
+{
+	////////////////////////////////////////////////////////////////////////////////////////////////////
+
+	struct FileMappingHandle
+	{
+		FileMappingHandle() = default;
+		bool operator==(const FileMappingHandle& o) const { return internal == o.internal; }
+		bool IsValid() const { return internal != 0; }
+		u64 ToU64() const { return internal; }
+		static FileMappingHandle FromU64(u64 v) { return { v }; }
+
+		FileMappingHandle(u64 v) { internal = v; }
+		u64 internal = 0; // HANDLE on windows, uid on posix
+
+		int shmFd = -1; // Unused on windows
+		int lockFd = -1; // Unused on windows
+	};
+
+	////////////////////////////////////////////////////////////////////////////////////////////////////
+}
diff --git a/Engine/Source/Programs/UnrealBuildAccelerator/Core/Public/UbaFileMappingWin.inl b/Engine/Source/Programs/UnrealBuildAccelerator/Core/Public/UbaFileMappingWin.inl
new file mode 100644
--- /dev/null
+++ b/Engine/Source/Programs/UnrealBuildAccelerator/Core/Public/UbaFileMappingWin.inl
@@ -0,0 +1,67 @@
+// Copyright Epic Games, Inc. All Rights Reserved.
+
+#pragma once
+
+////////////////////////////////////////////////////////////////////////////////////////////////////
+
+namespace uba
+{
+	HANDLE g_hostProcess;
+
+	namespace internal
+	{
+		#include "UbaFileMapping.inl"
+	}
+
+	FileMappingHandle CreateMemoryMappingW(Logger& logger, u32 flProtect, u64 maxSize, const tchar* name, const tchar* hint)
+	{
+		return internal::CreateMemoryMappingW(logger, flProtect, maxSize, name, hint);
+	}
+
+	FileMappingHandle CreateFileMappingW(Logger& logger, FileHandle file, u32 flProtect, u64 maxSize, const tchar* hint)
+	{
+		return internal::CreateFileMappingW(logger, file, flProtect, maxSize, hint);
+	}
+
+	u8* MapViewOfFile(Logger& logger, FileMappingHandle fileMappingObject, u32 dwDesiredAccess, u64 offset, u64 dwNumberOfBytesToMap, bool allowDiscard)
+	{
+		return internal::MapViewOfFile(logger, fileMappingObject, dwDesiredAccess, offset, dwNumberOfBytesToMap, allowDiscard);
+	}
+
+	bool MapViewCommit(Logger& logger, void* address, u64 size, bool allowDiscard)
+	{
+		return internal::MapViewCommit(logger, address, size, allowDiscard);
+	}
+
+	bool UnmapViewOfFile(Logger& logger, const void* lpBaseAddress, u64 bytesToUnmap, const tchar* hint, bool allowDiscard)
+	{
+		return internal::UnmapViewOfFile(logger, lpBaseAddress, bytesToUnmap, hint, allowDiscard);
+	}
+
+	bool CloseFileMapping(Logger& logger, FileMappingHandle h, const tchar* hint)
+	{
+		return internal::CloseFileMapping(logger, h, hint);
+	}
+
+	void* ReserveMemory(void* baseAddress, u64 capacity)
+	{
+		return internal::ReserveMemory(baseAddress, capacity);
+	}
+
+	void* MapMemory(FileMappingHandle handle, MemoryMapType mapType, void* targetAddress, u64 targetOffset, u64 targetCapacity, u64 handleOffset, u64 size)
+	{
+		return internal::MapMemory(handle, mapType, targetAddress, targetOffset, targetCapacity, handleOffset, size);
+	}
+
+	FileMappingHandle DuplicateMappingFromHost(FileMappingHandle handle, const char* hint)
+	{
+		return internal::DuplicateMappingFromHost(handle, hint);
+	}
+
+	u8* GetMemoryFromMapping(FileMappingHandle handle, u64 size, u64 offset, bool writable, const char* hint)
+	{
+		return internal::GetMemoryFromMapping(handle, size, offset, writable, hint);
+	}
+}
+
+////////////////////////////////////////////////////////////////////////////////////////////////////
diff --git a/Engine/Source/Programs/UnrealBuildAccelerator/Core/Public/UbaPlatform.h b/Engine/Source/Programs/UnrealBuildAccelerator/Core/Public/UbaPlatform.h
--- a/Engine/Source/Programs/UnrealBuildAccelerator/Core/Public/UbaPlatform.h
+++ b/Engine/Source/Programs/UnrealBuildAccelerator/Core/Public/UbaPlatform.h
@@ -130,9 +130,6 @@ namespace uba
 	bool GetOsVersion(StringBufferBase& outPretty, u32& outValue);
 	u32 GetCurrentProcessId();
 
-	enum FileHandle : u64 {};
-	inline constexpr FileHandle InvalidFileHandle = (FileHandle)-1;
-
 	void AddExceptionHandler();
 	void InitMemory();
 	bool SetMemoryWorkingSet(u64 minWorkingSet);
@@ -164,15 +161,6 @@ namespace uba
 	#define TStrcat_s wcscat_s
 	#define TFputs fputws
 	#define UBA_NOINLINE __declspec(noinline)
-	struct FileMappingHandle
-	{
-		bool operator==(const FileMappingHandle& o) const { return mh == o.mh; }
-		bool IsValid() const { return mh != 0; }
-		u64 ToU64() const { return u64(mh) | (u64(fh) << 32); }
-		static FileMappingHandle FromU64(u64 v) { return { (HANDLE)(v >> 32), (HANDLE)(v & ~0u) }; }
-		HANDLE fh = 0;
-		HANDLE mh = 0;
-	};
 #else
 	inline constexpr tchar PathSeparator = '/';
 	inline constexpr tchar NonPathSeparator = '\\';
@@ -206,19 +194,5 @@ namespace uba
 	#define UBA_NOINLINE __attribute__((noinline))
     #define ERROR_SUCCESS 0
 	#define DUPLICATE_SAME_ACCESS 0
-	inline constexpr u32 PAGE_READONLY = PROT_READ;
-
-	struct FileMappingHandle
-	{
-		FileMappingHandle(int shmFd_ = -1, int lockFd_ = -1, u64 uid_ = ~u64(0)) : shmFd(shmFd_), lockFd(lockFd_), uid(uid_) {};
-		FileMappingHandle(const FileMappingHandle& o) { shmFd = o.shmFd; uid = o.uid; lockFd = o.lockFd; }
-		bool operator==(const FileMappingHandle& o) const { return shmFd == o.shmFd; }
-		bool IsValid() const { return shmFd != -1; }
-		u64 ToU64() const { return uid; }
-		static FileMappingHandle FromU64(u64 v) { return FileMappingHandle(-1, -1, v); }
-		int shmFd;
-		int lockFd;
-		u64 uid;
-	};
 #endif // PLATFORM_WINDOWS
 }
diff --git a/Engine/Source/Programs/UnrealBuildAccelerator/Core/Public/UbaProtocol.h b/Engine/Source/Programs/UnrealBuildAccelerator/Core/Public/UbaProtocol.h
--- a/Engine/Source/Programs/UnrealBuildAccelerator/Core/Public/UbaProtocol.h
+++ b/Engine/Source/Programs/UnrealBuildAccelerator/Core/Public/UbaProtocol.h
@@ -55,7 +55,7 @@ namespace uba
 
 	inline constexpr u32 CommunicationMemSize = IsWindows ? 64*1024 : 64*1024*2; // Macos expands some commandlines to be crazy long
 
-	inline constexpr u32 FileMappingTableMemSize = 16 * 1024 * 1024;
+	inline constexpr u32 FileMappingTableMaxSize = 16 * 1024 * 1024;
 	inline constexpr u32 DirTableMaxSize = 128 * 1024 * 1024;
 	inline constexpr u32 OverlayTableMaxSize = 16 * 1024 * 1024; // Seen ~600kb max size for ltcg/thinlto linking of our biggest binaries (one temp file per .o file)
 
diff --git a/Engine/Source/Programs/UnrealBuildAccelerator/Core/Public/UbaSharedMemoryView.h b/Engine/Source/Programs/UnrealBuildAccelerator/Core/Public/UbaSharedMemoryView.h
--- a/Engine/Source/Programs/UnrealBuildAccelerator/Core/Public/UbaSharedMemoryView.h
+++ b/Engine/Source/Programs/UnrealBuildAccelerator/Core/Public/UbaSharedMemoryView.h
@@ -2,8 +2,8 @@
 
 #pragma once
 
+#include "UbaFileMappingHandle.h"
 #include "UbaMemory.h"
-#include "UbaPlatform.h"
 
 namespace uba
 {
diff --git a/Engine/Source/Programs/UnrealBuildAccelerator/Detours/Private/PosixOS/UbaDetoursFunctionsPosix.cpp b/Engine/Source/Programs/UnrealBuildAccelerator/Detours/Private/PosixOS/UbaDetoursFunctionsPosix.cpp
--- a/Engine/Source/Programs/UnrealBuildAccelerator/Detours/Private/PosixOS/UbaDetoursFunctionsPosix.cpp
+++ b/Engine/Source/Programs/UnrealBuildAccelerator/Detours/Private/PosixOS/UbaDetoursFunctionsPosix.cpp
@@ -6,6 +6,7 @@
 #include "UbaBinaryReaderWriter.h"
 #include "UbaDefinitions.h"
 #include "UbaDirectoryTable.h"
+#include "UbaFileMapping.h"
 #include "UbaPlatform.h"
 #include "UbaProcessStats.h"
 #include "UbaProcessUtils.h"
@@ -16,6 +17,7 @@
 #include <dirent.h> 
 #include <glob.h>
 #include <dlfcn.h>
+#include <sys/file.h>
 #include <sys/stat.h>
 
 // undefs for mimalloc
@@ -32,6 +34,10 @@
 
 namespace uba
 {
+	enum LogEntryType : u8 { LogEntryType_Error = 0, LogEntryType_Warning = 1, };
+	class Logger { public: void Logf(...) {} void Info(...) {} void Warning(...) {} void Error(...) {} };
+	#include "UbaFileMapping.inl"
+
 	constexpr bool g_logToScreen = false;
 
 	bool g_isDetouring;
@@ -2890,13 +2896,13 @@ namespace uba
 		g_isInitialized = true;
 
 		bool trackInputs = false;
-		u64 directoryTableHandle;
+		FileMappingHandle directoryTableHandle;
 		DirTableSize directoryTableSize;
 		u32 directoryTableCount;
-		u64 mappedFileTableHandle;
+		FileMappingHandle mappedFileTableHandle;
 		u32 mappedFileTableSize;
 		u32 mappedFileTableCount;
-		u64 overlayTableHandle;
+		FileMappingHandle overlayTableHandle;
 
 		{
 			RPC_MESSAGE(Init, init);
@@ -2911,13 +2917,13 @@ namespace uba
 
 			g_sharedMemoryAllocator = SharedMemoryAllocatorHandle::FromU64(reader.ReadU64());
 			reader.ReadU64(); // permanent files memory handle
-			directoryTableHandle = reader.ReadU64();
+			directoryTableHandle = FileMappingHandle::FromU64(reader.ReadU64());
 			directoryTableSize = FromU64(reader.ReadU64());
 			directoryTableCount = reader.ReadU32();
-			mappedFileTableHandle = reader.ReadU64();
+			mappedFileTableHandle = FileMappingHandle::FromU64(reader.ReadU64());
 			mappedFileTableSize = reader.ReadU32();
 			mappedFileTableCount = reader.ReadU32();
-			overlayTableHandle = reader.ReadU64();
+			overlayTableHandle = FileMappingHandle::FromU64(reader.ReadU64());
 
 			if (u16 vfsSize = reader.ReadU16())
 			{
@@ -2943,29 +2949,17 @@ namespace uba
 
 		LogHeader();
 
-		StringBuffer<128> mappedFileTableUidName;
-		GetMappingHandleName(mappedFileTableUidName, mappedFileTableHandle);
-		int mappedFileTableFd = shm_open(mappedFileTableUidName.data, O_RDONLY, S_IRUSR | S_IWUSR);
-		UBA_ASSERTF(mappedFileTableFd != -1, "Failed to open mapped file table memory mapping %s (%s)", mappedFileTableFd, mappedFileTableUidName.data, strerror(errno));
-		u8* mappedFileTableMem = (u8*)mmap(NULL, FileMappingTableMemSize, PROT_READ, MAP_SHARED, mappedFileTableFd, 0);
-		UBA_ASSERTF(mappedFileTableMem != MAP_FAILED, "mmap failed (%s)", strerror(errno));
+		TimerScope ts1(g_stats.fileTable);
+		u8* mappedFileTableMem = GetMemoryFromMapping(mappedFileTableHandle, FileMappingTableMaxSize,0, false, "FileTable");
 		g_mappedFileTable.Init(mappedFileTableMem, mappedFileTableCount, mappedFileTableSize);
+		ts1.Leave();
 
-		StringBuffer<128> dirTableUidName;
-		GetMappingHandleName(dirTableUidName, directoryTableHandle);
-		int dirTableFd = shm_open(dirTableUidName.data, O_RDONLY, S_IRUSR | S_IWUSR);
-		UBA_ASSERTF(dirTableFd != -1, "shm_open failed (%s)", strerror(errno));
-		u8* dirTableMem = (u8*)mmap(NULL, DirTableMaxSize, PROT_READ, MAP_SHARED, dirTableFd, 0);
-		UBA_ASSERTF(dirTableMem != MAP_FAILED, "mmap for dirtable mem failed (%s)", strerror(errno));
+		TimerScope ts2(g_stats.dirTable);
+		u8* dirTableMem = GetMemoryFromMapping(directoryTableHandle, DirTableMaxSize,0, false, "DirTable");
+		u8* overlayTableMem = GetMemoryFromMapping(overlayTableHandle, OverlayTableMaxSize,0, false, "OverlayTable");
 		g_directoryTable.Init(dirTableMem, directoryTableCount, directoryTableSize.main);
-
-		StringBuffer<128> overlayTableUidName;
-		GetMappingHandleName(overlayTableUidName, overlayTableHandle);
-		int overlayTableFd = shm_open(overlayTableUidName.data, O_RDONLY, S_IRUSR | S_IWUSR);
-		UBA_ASSERTF(overlayTableFd != -1, "shm_open failed (%s)", strerror(errno));
-		u8* overlayTableMem = (u8*)mmap(NULL, OverlayTableMaxSize, PROT_READ, MAP_SHARED, overlayTableFd, 0);
-		UBA_ASSERTF(overlayTableMem != MAP_FAILED, "mmap for dirtable mem failed (%s)", strerror(errno));
 		g_directoryTable.InitOverlay(overlayTableMem , directoryTableSize.overlay);
+		ts2.Leave();
 
 		if (trackInputs)
 			InitTrackInput();
diff --git a/Engine/Source/Programs/UnrealBuildAccelerator/Detours/Private/UbaDetoursShared.h b/Engine/Source/Programs/UnrealBuildAccelerator/Detours/Private/UbaDetoursShared.h
--- a/Engine/Source/Programs/UnrealBuildAccelerator/Detours/Private/UbaDetoursShared.h
+++ b/Engine/Source/Programs/UnrealBuildAccelerator/Detours/Private/UbaDetoursShared.h
@@ -3,6 +3,7 @@
 #pragma once
 
 #include "UbaApplicationRules.h"
+#include "UbaFileHandle.h"
 #include "UbaStringBuffer.h"
 #include "UbaPlatform.h"
 #include "UbaProcessStats.h"
diff --git a/Engine/Source/Programs/UnrealBuildAccelerator/Detours/Private/Windows/UbaDetoursFunctionsWin.cpp b/Engine/Source/Programs/UnrealBuildAccelerator/Detours/Private/Windows/UbaDetoursFunctionsWin.cpp
--- a/Engine/Source/Programs/UnrealBuildAccelerator/Detours/Private/Windows/UbaDetoursFunctionsWin.cpp
+++ b/Engine/Source/Programs/UnrealBuildAccelerator/Detours/Private/Windows/UbaDetoursFunctionsWin.cpp
@@ -52,6 +52,7 @@ DETOURED_FUNCTIONS_MEMORY_NON_WINE
 #include "UbaApplicationRules.h"
 #include "UbaDetoursShared.h"
 #include "UbaDetoursUtilsWin.h"
+#include "UbaFileMapping.h"
 
 #include "Shlwapi.h"
 #include <detours/detours.h>
@@ -1082,15 +1083,15 @@ void Init(const DetoursPayload& payload, u64 startTime)
 	StringBuffer<512> workingDirBuffer;
 
 	bool trackInputs = false;
-	HANDLE sharedMemoryAllocatorHandle;
+	SharedMemoryAllocatorHandle sharedMemoryAllocatorHandle;
 	SharedMemoryHandle permanentFilesHandle;
-	HANDLE directoryTableHandle;
+	FileMappingHandle directoryTableHandle;
 	DirTableSize directoryTableSize;
 	u32 directoryTableCount;
-	HANDLE mappedFileTableHandle;
+	FileMappingHandle mappedFileTableHandle;
 	u32 mappedFileTableSize;
 	u32 mappedFileTableCount;
-	HANDLE overlayFilesHandle;
+	FileMappingHandle overlayFilesHandle;
 
 	{
 		RPC_MESSAGE(Init, init)
@@ -1103,15 +1104,15 @@ void Init(const DetoursPayload& payload, u64 startTime)
 		reader.ReadString(applicationBuffer);
 		reader.ReadString(workingDirBuffer);
 
-		sharedMemoryAllocatorHandle = (HANDLE)reader.ReadU64();
+		sharedMemoryAllocatorHandle = FileMappingHandle::FromU64(reader.ReadU64());
 		permanentFilesHandle = SharedMemoryHandle::FromU64(reader.ReadU64());
-		directoryTableHandle = FileMappingHandle::FromU64(reader.ReadU64()).mh;
+		directoryTableHandle = FileMappingHandle::FromU64(reader.ReadU64());
 		directoryTableSize = FromU64(reader.ReadU64());
 		directoryTableCount = reader.ReadU32();
-		mappedFileTableHandle = FileMappingHandle::FromU64(reader.ReadU64()).mh;
+		mappedFileTableHandle = FileMappingHandle::FromU64(reader.ReadU64());
 		mappedFileTableSize = reader.ReadU32();
 		mappedFileTableCount = reader.ReadU32();
-		overlayFilesHandle = FileMappingHandle::FromU64(reader.ReadU64()).mh;
+		overlayFilesHandle = FileMappingHandle::FromU64(reader.ReadU64());
 
 		if (u16 vfsSize = reader.ReadU16())
 		{
@@ -1150,45 +1151,19 @@ void Init(const DetoursPayload& payload, u64 startTime)
 	}
 	#endif
 
-	if (!True_DuplicateHandle(g_hostProcess, sharedMemoryAllocatorHandle, GetCurrentProcess(), &g_sharedMemoryAllocator.mh, 0, FALSE, DUPLICATE_SAME_ACCESS))
-		UBA_ASSERTF(false, L"Failed to duplicate shared memory handle (%u)", GetLastError());
+	g_sharedMemoryAllocator = DuplicateMappingFromHost(sharedMemoryAllocatorHandle, "SharedMem");
 
-	if (!True_DuplicateHandle(g_hostProcess, mappedFileTableHandle, GetCurrentProcess(), &mappedFileTableHandle, 0, FALSE, DUPLICATE_SAME_ACCESS))
-		UBA_ASSERTF(false, L"Failed to duplicate filetable handle (%u)", GetLastError());
+	TimerScope ts1(g_stats.fileTable);
+	u8* mappedFileTableMem = GetMemoryFromMapping(DuplicateMappingFromHost(mappedFileTableHandle, "FileTable"), FileMappingTableMaxSize, 0, false, "FileTable");
+	g_mappedFileTable.Init(mappedFileTableMem, mappedFileTableCount, mappedFileTableSize);
+	ts1.Leave();
 
-	u8* mappedFileTableMem;
-	{
-		TimerScope ts(g_kernelStats.mapViewOfFile);
-		mappedFileTableMem = (u8*)True_MapViewOfFile(mappedFileTableHandle, FILE_MAP_READ, 0, 0, 0);
-		UBA_ASSERT(mappedFileTableMem);
-	}
-	{
-		TimerScope ts2(g_stats.fileTable);
-		g_mappedFileTable.Init(mappedFileTableMem, mappedFileTableCount, mappedFileTableSize);
-	}
-
-	if (!True_DuplicateHandle(g_hostProcess, directoryTableHandle, GetCurrentProcess(), &directoryTableHandle, 0, FALSE, DUPLICATE_SAME_ACCESS))
-		UBA_ASSERTF(false, L"Failed to duplicate directorytable handle (%u)", GetLastError());
-	if (!True_DuplicateHandle(g_hostProcess, overlayFilesHandle, GetCurrentProcess(), &overlayFilesHandle, 0, FALSE, DUPLICATE_SAME_ACCESS))
-		UBA_ASSERTF(false, L"Failed to duplicate directorytable handle (%u)", GetLastError());
-
-	u8* directoryTableMem;
-	{
-		TimerScope ts(g_kernelStats.mapViewOfFile);
-		directoryTableMem = (u8*)True_MapViewOfFile(directoryTableHandle, FILE_MAP_READ, 0, 0, 0);
-		UBA_ASSERT(directoryTableMem);
-	}
-	u8* overlayTableMem;
-	{
-		TimerScope ts(g_kernelStats.mapViewOfFile);
-		overlayTableMem = (u8*)True_MapViewOfFile(overlayFilesHandle, FILE_MAP_READ, 0, 0, 0);
-		UBA_ASSERT(overlayTableMem);
-	}
-	{
-		TimerScope ts2(g_stats.dirTable);
-		g_directoryTable.Init(directoryTableMem, directoryTableCount, directoryTableSize.main);
-		g_directoryTable.InitOverlay(overlayTableMem, directoryTableSize.overlay);
-	}
+	TimerScope ts2(g_stats.dirTable);
+	u8* directoryTableMem = GetMemoryFromMapping(DuplicateMappingFromHost(directoryTableHandle, "DirTable"), DirTableMaxSize, 0, false, "DirTable");
+	u8* overlayTableMem = GetMemoryFromMapping(DuplicateMappingFromHost(overlayFilesHandle, "OverlayTable"), OverlayTableMaxSize, 0, false, "OverlayTable");
+	g_directoryTable.Init(directoryTableMem, directoryTableCount, directoryTableSize.main);
+	g_directoryTable.InitOverlay(overlayTableMem, directoryTableSize.overlay);
+	ts2.Leave();
 
 	if (permanentFilesHandle.IsValid())
 	{
diff --git a/Engine/Source/Programs/UnrealBuildAccelerator/Detours/Private/Windows/UbaDetoursMainWin.cpp b/Engine/Source/Programs/UnrealBuildAccelerator/Detours/Private/Windows/UbaDetoursMainWin.cpp
--- a/Engine/Source/Programs/UnrealBuildAccelerator/Detours/Private/Windows/UbaDetoursMainWin.cpp
+++ b/Engine/Source/Programs/UnrealBuildAccelerator/Detours/Private/Windows/UbaDetoursMainWin.cpp
@@ -12,12 +12,21 @@
 #include "UbaDetoursPayload.h"
 #include "UbaDetoursShared.h"
 #include "UbaDetoursFunctionsWin.h"
+#include "UbaFileMapping.h"
 #include <detours/detours.h>
 #include <stdio.h>
 
 namespace uba
 {
-	HANDLE g_hostProcess;
+	enum LogEntryType : u8 { LogEntryType_Error = 0, LogEntryType_Warning = 1, };
+	class Logger { public: void Logf(...) {} void Info(...) {} void Warning(...) {} void Error(...) {} };
+	struct Bottleneck { Bottleneck(...) {} };
+	struct BottleneckScope { BottleneckScope(Bottleneck&, Timer&) {} };
+}
+#include "UbaFileMappingWin.inl"
+
+namespace uba
+{
 	HANDLE g_writeEvent;
 	HANDLE g_readEvent;
 	HANDLE g_cancelEvent;
@@ -48,13 +57,8 @@ namespace uba
 		if (!DuplicateHandle(g_hostProcess, payload.writeEvent, GetCurrentProcess(), &g_readEvent, SYNCHRONIZE, false, 0))
 			return -3;
 
-		if (!DuplicateHandle(g_hostProcess, payload.communicationHandle, GetCurrentProcess(), &g_messageMappingHandle, FILE_MAP_READ | FILE_MAP_WRITE, false, 0))
-		{
-			UBA_ASSERT(!WasCancelled());
-			return -4;
-		}
-
-		u8* mem = (u8*)::MapViewOfFile(g_messageMappingHandle, FILE_MAP_READ | FILE_MAP_WRITE, ToHigh(payload.communicationOffset), ToLow(payload.communicationOffset), CommunicationMemSize);
+		FileMappingHandle communicationHandle = DuplicateMappingFromHost(FileMappingHandle::FromU64(payload.communicationHandle), "Communication");
+		u8* mem = GetMemoryFromMapping(communicationHandle, CommunicationMemSize, payload.communicationOffset, true, "Communication");
 		if (!mem)
 		{
 			return -5;
@@ -72,13 +76,13 @@ namespace uba
 
 	int Disconnect()
 	{
-		UnmapViewOfFile(g_writeMessageMappingMem);
-		CloseHandle(g_messageMappingHandle);
-		CloseHandle(g_cancelEvent);
-		CloseHandle(g_readEvent);
-		CloseHandle(g_writeEvent);
+		::UnmapViewOfFile(g_writeMessageMappingMem);
+		::CloseHandle(g_messageMappingHandle);
+		::CloseHandle(g_cancelEvent);
+		::CloseHandle(g_readEvent);
+		::CloseHandle(g_writeEvent);
 
-		CloseHandle(g_hostProcess);
+		::CloseHandle(g_hostProcess);
 		return 0;
 	}
 
diff --git a/Engine/Source/Programs/UnrealBuildAccelerator/Visualizer/Private/UbaTraceReader.cpp b/Engine/Source/Programs/UnrealBuildAccelerator/Visualizer/Private/UbaTraceReader.cpp
--- a/Engine/Source/Programs/UnrealBuildAccelerator/Visualizer/Private/UbaTraceReader.cpp
+++ b/Engine/Source/Programs/UnrealBuildAccelerator/Visualizer/Private/UbaTraceReader.cpp
@@ -277,11 +277,12 @@ namespace uba
 		{
 #if PLATFORM_WINDOWS
 			u64 traceCapacity = 0;
-			m_memoryHandle.mh = ::OpenFileMappingW(PAGE_READWRITE, false, namedTrace);
+			m_memoryHandle.internal = u64(::OpenFileMappingW(PAGE_READWRITE, false, namedTrace));
 #else
 			u64 traceCapacity = 128 * 1024 * 1024;
 			int oflags = O_RDONLY | O_NOFOLLOW;
 			m_memoryHandle.shmFd = shm_open(namedTrace, oflags, S_IRUSR | S_IWUSR);
+			m_memoryHandle.internal = ~0ull;
 
 			// We don't want to unlink
 			//if (m_memoryHandle.shmFd != -1)
diff --git a/Engine/Source/Programs/UnrealBuildAccelerator/Wine/UbaWineLinux.cpp b/Engine/Source/Programs/UnrealBuildAccelerator/Wine/UbaWineLinux.cpp
--- a/Engine/Source/Programs/UnrealBuildAccelerator/Wine/UbaWineLinux.cpp
+++ b/Engine/Source/Programs/UnrealBuildAccelerator/Wine/UbaWineLinux.cpp
@@ -3,9 +3,17 @@
 
 #include <netinet/tcp.h>
 #include <netinet/in.h>
+#include <dirent.h>
 #include <errno.h>
+#include <pthread.h>
 #include <stdio.h>
 #include <stdlib.h>
+#include <string.h>
+#include <sys/file.h>
+#include <sys/mman.h>
+#include <sys/stat.h>
+#include <atomic>
+
 //#include "../Common/Public/UbaLinuxNetworkWrappers.h"
 
 int UnixGetsockopt(int fd, int level, int optname, char* optval, int* optlen)
@@ -39,8 +47,22 @@ int UnixGetTcpInfo(int fd, void* buf, int* len)
 #include "../Core/Public/UbaBase.h"
 #include "../Common/Public/UbaSocket.h"
 
+#define UBA_ASSERT(...)
+#define UBA_ASSERTF(...) 
+
+#include "../Core/Public/UbaFileHandle.h"
+#include "../Core/Public/UbaFileMappingHandle.h"
+#include "../Core/Public/UbaFileMapping.h"
+
 namespace uba
 {
+	struct Logger { void Logf(...) {} void Info(...) {} void Warning(...) {} void Error(...) {} };
+	template<typename Type> using Atomic = std::atomic<Type>;
+	void SetLastError(int err) {}
+	enum LogEntryType : u8 { LogEntryType_Error = 0, LogEntryType_Warning = 1, };
+
+	#include "../Core/Public/UbaFileMapping.inl"
+
 	static constexpr u32 KeepAliveProbeCount = 10;
 	#include "../Common/Public/UbaSocket.inl"
 }
diff --git a/Engine/Source/Programs/UnrealBuildAccelerator/Wine/UbaWineWindows.cpp b/Engine/Source/Programs/UnrealBuildAccelerator/Wine/UbaWineWindows.cpp
--- a/Engine/Source/Programs/UnrealBuildAccelerator/Wine/UbaWineWindows.cpp
+++ b/Engine/Source/Programs/UnrealBuildAccelerator/Wine/UbaWineWindows.cpp
@@ -1,8 +1,8 @@
 // Copyright Epic Games, Inc. All Rights Reserved.
 
-#define PLATFORM_WINDOWS 0
+#define PLATFORM_WINDOWS 1
 #define PLATFORM_MAC 0
-#define PLATFORM_LINUX 1
+#define PLATFORM_LINUX 0
 
 #include <winsock2.h>
 #include <stdint.h>
@@ -11,6 +11,7 @@
 #include <stdio.h>
 #include "../Core/Public/UbaBase.h"
 #include "../Common/Public/UbaSocket.h"
+#include "../Core/Public/UbaFileMapping.h"
 
 // NOTE, README
 // Must use "WINAPI" on function and add Windows.def
