[UBA]
* Fixed so asHANDLE and asFileDescriptor is only in one place. Also fixed casing

diff --git a/Engine/Source/Programs/UnrealBuildAccelerator/Common/Private/UbaFile.cpp b/Engine/Source/Programs/UnrealBuildAccelerator/Common/Private/UbaFile.cpp
--- a/Engine/Source/Programs/UnrealBuildAccelerator/Common/Private/UbaFile.cpp
+++ b/Engine/Source/Programs/UnrealBuildAccelerator/Common/Private/UbaFile.cpp
@@ -72,11 +72,6 @@ namespace uba
 		}
 	}
 
-	HANDLE asHANDLE(FileHandle fh)
-	{
-		return (HANDLE)(fh == InvalidFileHandle ? InvalidFileHandle : (fh & FileHandleFlagMask));
-	}
-
 	template<typename T>
 	void PrefixFileName(T& fileName, StringView prefix, StringBufferBase& temp)
 	{
@@ -131,11 +126,11 @@ namespace uba
 			u32 toRead = u32(Min(readLeft, u64(~u32(0)) - 1));
 #if PLATFORM_WINDOWS
 			DWORD wasRead = 0;
-			if (!::ReadFile(asHANDLE(fileHandle), buffer, toRead, &wasRead, NULL))
+			if (!::ReadFile(AsHANDLE(fileHandle), buffer, toRead, &wasRead, NULL))
 				if (GetLastError() != ERROR_IO_PENDING)
-					return logger.Error(TC("ERROR reading %llu bytes from file %s with handle %llu (error: %s)"), toRead, fileName, asHANDLE(fileHandle), LastErrorToText().data);
+					return logger.Error(TC("ERROR reading %llu bytes from file %s with handle %llu (error: %s)"), toRead, fileName, AsHANDLE(fileHandle), LastErrorToText().data);
 #else
-			ssize_t wasRead = read(asFileDescriptor(fileHandle), buffer, toRead);
+			ssize_t wasRead = read(AsFileDescriptor(fileHandle), buffer, toRead);
 			if (wasRead == -1)
 			{
 				UBA_ASSERTF(false, TC("ERROR ReadFile error handling not implemented (Trying to read %u bytes from fd %llu)"), toRead, fileHandle);
@@ -168,10 +163,10 @@ namespace uba
 		{
 #if PLATFORM_WINDOWS
 			DWORD written = 0;
-			if (!::WriteFile(asHANDLE(fileHandle), buffer, (DWORD)writeLeft, &written, NULL))
+			if (!::WriteFile(AsHANDLE(fileHandle), buffer, (DWORD)writeLeft, &written, NULL))
 				return logger.Error(TC("ERROR WriteFile failed writing to file %s (%s)"), fileName, LastErrorToText().data);
 #else
-			auto written = write(asFileDescriptor(fileHandle), buffer, writeLeft);
+			auto written = write(AsFileDescriptor(fileHandle), buffer, writeLeft);
 			if (written == -1)
 				return false;
 #endif
@@ -228,7 +223,7 @@ namespace uba
 			return GetFileBasicInformation(out, logger, fileName, errorOnFail);
 		IO_STATUS_BLOCK b;
 		FILE_NETWORK_OPEN_INFORMATION info;
-		NTSTATUS res = NtQueryInformationFile(asHANDLE(hFile), &b, &info, sizeof(info), FileNetworkOpenInformation);
+		NTSTATUS res = NtQueryInformationFile(AsHANDLE(hFile), &b, &info, sizeof(info), FileNetworkOpenInformation);
 		if (res != STATUS_SUCCESS)
 			return errorOnFail ? logger.Error(TC("GetFileBasicInformationByHandle (NtQueryInformationFile) failed on %s (0x%x)"), fileName, res) : false;
 		out.attributes = info.FileAttributes;
@@ -281,7 +276,7 @@ namespace uba
 		ExtendedTimerScope ts(KernelStats::GetCurrent().getFileInfo);
 #if PLATFORM_WINDOWS
 		BY_HANDLE_FILE_INFORMATION info;
-		if (!::GetFileInformationByHandle(asHANDLE(hFile), &info))
+		if (!::GetFileInformationByHandle(AsHANDLE(hFile), &info))
 			return logger.Error(TC("GetFileInformationByHandle failed on %s (%s)"), fileName, LastErrorToText().data);
 		out.attributes = info.dwFileAttributes;
 		out.volumeSerialNumber = info.dwVolumeSerialNumber;
@@ -291,7 +286,7 @@ namespace uba
 		return true;
 #else
 		struct stat attr;
-		int res = fstat(asFileDescriptor(hFile), &attr);
+		int res = fstat(AsFileDescriptor(hFile), &attr);
 		if (res != 0)
 			return logger.Error(TC("GetFileInformationByHandle (fstat) failed on %s (%s)"), fileName, strerror(errno));
 		out.lastWriteTime = FromTimeSpec(attr.st_mtimespec);
@@ -393,11 +388,11 @@ namespace uba
 	bool SetFilePointer(Logger& logger, const tchar* fileName, FileHandle handle, u64 position)
 	{
 #if PLATFORM_WINDOWS
-		if (!::SetFilePointerEx(asHANDLE(handle), ToLargeInteger(position), NULL, FILE_BEGIN))
+		if (!::SetFilePointerEx(AsHANDLE(handle), ToLargeInteger(position), NULL, FILE_BEGIN))
 			return logger.Error(TC("SetFilePointerEx failed on %s (%s)"), fileName, LastErrorToText().data);
 		return true;
 #else
-		if (lseek(asFileDescriptor(handle), position, SEEK_SET) != position)
+		if (lseek(AsFileDescriptor(handle), position, SEEK_SET) != position)
 			return logger.Error("lseek to %llu failed for %s: %s", position, fileName, strerror(errno));
 		return true;
 #endif
@@ -409,11 +404,11 @@ namespace uba
 #if PLATFORM_WINDOWS
 		FILE_END_OF_FILE_INFO info;
 		info.EndOfFile = ToLargeInteger(size);
-		if (!::SetFileInformationByHandle(asHANDLE(handle), FileEndOfFileInfo, &info, sizeof(info)))
+		if (!::SetFileInformationByHandle(AsHANDLE(handle), FileEndOfFileInfo, &info, sizeof(info)))
 			return logger.Error(TC("SetFileInformationByHandle failed on %s (%s)"), fileName, LastErrorToText().data);
 		return true;
 #else
-		if (ftruncate(asFileDescriptor(handle), size) == -1)
+		if (ftruncate(AsFileDescriptor(handle), size) == -1)
 			return logger.Error(TC("ftruncate failed on %s (%s)"), fileName, LastErrorToText().data);
 		return false;
 #endif
@@ -422,7 +417,7 @@ namespace uba
 	bool MakeSparse(Logger& logger, const tchar* fileName, FileHandle h, u64 size)
 	{
 #if PLATFORM_WINDOWS
-		HANDLE handle = asHANDLE(h);
+		HANDLE handle = AsHANDLE(h);
 		DWORD dwTemp;
 		if (!::DeviceIoControl(handle, FSCTL_SET_SPARSE, NULL, 0, NULL, 0, &dwTemp, NULL))
 			return logger.Error(TC("Failed to make file %s sparse (%s)"), fileName, LastErrorToText().data);
@@ -667,11 +662,11 @@ namespace uba
 			return true;
 
 #if PLATFORM_WINDOWS
-		return ::CloseHandle(asHANDLE(h));
+		return ::CloseHandle(AsHANDLE(h));
 #else
-		if (close(asFileDescriptor(h)) == 0)
+		if (close(AsFileDescriptor(h)) == 0)
 			return true;
-		UBA_ASSERTF(false, TC("CloseFile error handling not implemented while failing to close %s with handle %i (%s)"), fileName, asFileDescriptor(h), strerror(errno));
+		UBA_ASSERTF(false, TC("CloseFile error handling not implemented while failing to close %s with handle %i (%s)"), fileName, AsFileDescriptor(h), strerror(errno));
 		return false;
 #endif
 	}
@@ -832,12 +827,12 @@ namespace uba
 		ExtendedTimerScope ts(KernelStats::GetCurrent().getFileTime);
 #if PLATFORM_WINDOWS
 		FILETIME lastWriteTime;
-		auto res = ::GetFileTime(asHANDLE(fileHandle), NULL, NULL, &lastWriteTime);
+		auto res = ::GetFileTime(AsHANDLE(fileHandle), NULL, NULL, &lastWriteTime);
 		outTime = (u64&)lastWriteTime;
 		return res;
 #else
 		struct stat attr;
-		int res = fstat(asFileDescriptor(fileHandle), &attr);
+		int res = fstat(AsFileDescriptor(fileHandle), &attr);
 		if (res != 0)
 		{
 			UBA_ASSERTF(errno == ENOENT, TC("GetFileLastWriteTime (fstat) error handling not implemented: %s"), strerror(errno));
@@ -852,11 +847,11 @@ namespace uba
 	bool SetFileLastWriteTime(FileHandle fileHandle, u64 writeTime)
 	{
 #if PLATFORM_WINDOWS
-		return SetFileTime(asHANDLE(fileHandle), (FILETIME*)&writeTime, NULL, (FILETIME*)&writeTime);
+		return SetFileTime(AsHANDLE(fileHandle), (FILETIME*)&writeTime, NULL, (FILETIME*)&writeTime);
 #else
 		struct timespec times[2];
 		times[0] = times[1] = ToTimeSpec(writeTime);
-		if (futimens(asFileDescriptor(fileHandle), times) == 0)
+		if (futimens(AsFileDescriptor(fileHandle), times) == 0)
 			return true;
 		UBA_ASSERTF(false, TC("SetFileLastWriteTime (futimens) error handling not implemented: %s"), strerror(errno));
 		return false;
@@ -886,13 +881,13 @@ namespace uba
 		ExtendedTimerScope ts(KernelStats::GetCurrent().getFileInfo);
 #if PLATFORM_WINDOWS
 		LARGE_INTEGER lpFileSize;
-		if (!::GetFileSizeEx(asHANDLE(hFile), &lpFileSize))
+		if (!::GetFileSizeEx(AsHANDLE(hFile), &lpFileSize))
 			return false;
 		outFileSize = u64(lpFileSize.QuadPart);
 		return true;
 #else
 		struct stat attr;
-		int res = fstat(asFileDescriptor(hFile), &attr);
+		int res = fstat(AsFileDescriptor(hFile), &attr);
 		if (res == 0)
 		{
 			outFileSize = attr.st_size;
@@ -1201,7 +1196,7 @@ namespace uba
 	{
 		#if PLATFORM_WINDOWS
 		for (auto& volume : volumes)
-			CloseHandle(asHANDLE(volume.handle));
+			CloseHandle(AsHANDLE(volume.handle));
 		#endif
 	}
 
@@ -1299,9 +1294,9 @@ namespace uba
 		#if PLATFORM_WINDOWS
 		DISK_PERFORMANCE perf = {};
 		DWORD bytesReturned;
-		if (!::DeviceIoControl(asHANDLE(handle), IOCTL_DISK_PERFORMANCE, NULL, 0, &perf, sizeof(perf), &bytesReturned, NULL))
+		if (!::DeviceIoControl(AsHANDLE(handle), IOCTL_DISK_PERFORMANCE, NULL, 0, &perf, sizeof(perf), &bytesReturned, NULL))
 		{
-			::CloseHandle(asHANDLE(handle));
+			::CloseHandle(AsHANDLE(handle));
 			handle = InvalidFileHandle;
 			return false;
 		}
diff --git a/Engine/Source/Programs/UnrealBuildAccelerator/Common/Private/UbaFileAccessor.cpp b/Engine/Source/Programs/UnrealBuildAccelerator/Common/Private/UbaFileAccessor.cpp
--- a/Engine/Source/Programs/UnrealBuildAccelerator/Common/Private/UbaFileAccessor.cpp
+++ b/Engine/Source/Programs/UnrealBuildAccelerator/Common/Private/UbaFileAccessor.cpp
@@ -20,9 +20,7 @@ namespace uba
 
 	#if PLATFORM_WINDOWS
 	void GetProcessHoldingFile(StringBufferBase& out, const tchar* fileName);
-	HANDLE asHANDLE(FileHandle fh);
 	#else
-	int asFileDescriptor(FileHandle fh);
 	Atomic<u32> g_tempFileCounter = 0;
 	#endif
 		
@@ -32,7 +30,7 @@ namespace uba
 		ExtendedTimerScope ts(KernelStats::GetCurrent().setFileInfo);
 		FILE_DISPOSITION_INFO info;
 		info.DeleteFile = value;
-		if (!::SetFileInformationByHandle(asHANDLE(handle), FileDispositionInfo, &info, sizeof(info)))
+		if (!::SetFileInformationByHandle(AsHANDLE(handle), FileDispositionInfo, &info, sizeof(info)))
 			return logger.Error(TC("SetFileInformationByHandle (FileDispositionInfo) failed on %llu %s (%s)"), uintptr_t(handle), fileName, LastErrorToText().data);
 		return true;
 	}
@@ -40,7 +38,7 @@ namespace uba
 	u64 GetFilePointer(Logger& logger, FileHandle& handle)
 	{
 		LARGE_INTEGER pos;
-		if (SetFilePointerEx(asHANDLE(handle), ToLargeInteger(0), &pos, FILE_CURRENT))
+		if (SetFilePointerEx(AsHANDLE(handle), ToLargeInteger(0), &pos, FILE_CURRENT))
 			return pos.QuadPart;
 		logger.Error(TC("SetFilePointerEx failed (%s)"), LastErrorToText().data);
 		return ~0ull;
@@ -151,7 +149,7 @@ namespace uba
 
 			#if UBA_USE_SPARSE_FILE
 			{
-				HANDLE handle = asHANDLE(m_fileHandle);
+				HANDLE handle = AsHANDLE(m_fileHandle);
 				DWORD dwTemp;
 				if (!::DeviceIoControl(handle, FSCTL_SET_SPARSE, NULL, 0, NULL, 0, &dwTemp, NULL))
 				{
@@ -167,7 +165,7 @@ namespace uba
 			{
 				FILE_ALLOCATION_INFO info;
 				info.AllocationSize = ToLargeInteger(allocSize);
-				if (!::SetFileInformationByHandle(asHANDLE(m_fileHandle), FileAllocationInfo, &info, sizeof(info)))
+				if (!::SetFileInformationByHandle(AsHANDLE(m_fileHandle), FileAllocationInfo, &info, sizeof(info)))
 					if (!IsRunningWine())
 						return m_logger.Error(TC("SetFileInformationByHandle (FileDispositionInfo) failed on %llu %s (%s)"), uintptr_t(m_fileHandle), realFileName, LastErrorToText().data);
 			}
@@ -176,9 +174,9 @@ namespace uba
 #elif PLATFORM_LINUX
 		if (isMemoryMap && fileSize != 0) // We want to make sure we allocate memory on the file system to catch out-of-diskspace
 		{
-			if (fallocate(asFileDescriptor(m_fileHandle), 0, 0, fileSize) == -1)
+			if (fallocate(AsFileDescriptor(m_fileHandle), 0, 0, fileSize) == -1)
 				if (errno != ENOSYS && errno != EOPNOTSUPP)
-					return m_logger.Error(TC("fallocate failed on %i %s (%s)"), asFileDescriptor(m_fileHandle), realFileName, strerror(errno));
+					return m_logger.Error(TC("fallocate failed on %i %s (%s)"), AsFileDescriptor(m_fileHandle), realFileName, strerror(errno));
 		}
 #endif
 		m_isWrite = true;
@@ -302,7 +300,7 @@ namespace uba
 					ol[0].hEvent = ev[0].GetHandle();
 					ol[0].Offset = ToLow(lastOffset);
 					ol[0].OffsetHigh = ToHigh(lastOffset);
-					if (!::WriteFile(asHANDLE(m_fileHandle), experiment, WriteUnit, NULL, ol))
+					if (!::WriteFile(AsHANDLE(m_fileHandle), experiment, WriteUnit, NULL, ol))
 					{
 						u32 lastError = GetLastError();
 						if (lastError != ERROR_IO_PENDING)
@@ -360,7 +358,7 @@ namespace uba
 				ol[index].OffsetHigh = ToHigh(writeOffset);
 				ev[index].Reset();
 
-				if (!::WriteFile(asHANDLE(m_fileHandle), pos, u32(toActuallyWrite), NULL, ol + index))
+				if (!::WriteFile(AsHANDLE(m_fileHandle), pos, u32(toActuallyWrite), NULL, ol + index))
 				{
 					u32 lastError = GetLastError();
 					if (lastError != ERROR_IO_PENDING)
@@ -390,7 +388,7 @@ namespace uba
 
 #if PLATFORM_WINDOWS
 			DWORD written;
-			if (!::WriteFile(asHANDLE(m_fileHandle), pos, toActuallyWrite, &written, NULL))
+			if (!::WriteFile(AsHANDLE(m_fileHandle), pos, toActuallyWrite, &written, NULL))
 			{
 				DWORD lastError = GetLastError();
 				if ((lastError == ERROR_WORKING_SET_QUOTA || lastError == ERROR_PAGEFILE_QUOTA) && retryIndex < 5)
@@ -415,7 +413,7 @@ namespace uba
 			if (written > toWrite)
 				written = toWrite;
 #else
-			ssize_t written = write(asFileDescriptor(m_fileHandle), pos, toActuallyWrite);
+			ssize_t written = write(AsFileDescriptor(m_fileHandle), pos, toActuallyWrite);
 			if (written == -1)
 				return m_logger.Error(TC("ERROR writing file %s writing %u bytes (%llu bytes written out of %llu) (%s)"), m_fileName, toActuallyWrite, (dataLen - writeLeft), dataLen, strerror(errno));
 #endif
@@ -452,7 +450,7 @@ namespace uba
 			return m_logger.Error(TC("Failed to create mapping handle for %s with size %llu (%s)"), m_fileName, m_size, LastErrorToText().data);
 #else
 		m_mappingHandle.internal = ~0ull;
-		m_mappingHandle.shmFd = asFileDescriptor(m_fileHandle);
+		m_mappingHandle.shmFd = AsFileDescriptor(m_fileHandle);
 		if (offset == m_size)
 			return true;
 #endif
@@ -552,7 +550,7 @@ namespace uba
 					if (targetFd == -1)
 						return m_logger.Error(TC("Failed to create file %s for move from temporary file %s (%s)"), m_fileName, realFileName, strerror(errno));
 						
-					int sourceFd = asFileDescriptor(m_fileHandle);
+					int sourceFd = AsFileDescriptor(m_fileHandle);
 					#if PLATFORM_MAC
 					if (fcopyfile(sourceFd, targetFd, 0, COPYFILE_ALL) == -1)
 						return m_logger.Error(TC("Failed to do fcopyfile from temporary %s to file %s (%s)"), realFileName, m_fileName, strerror(errno));
diff --git a/Engine/Source/Programs/UnrealBuildAccelerator/Core/Public/UbaFileHandle.h b/Engine/Source/Programs/UnrealBuildAccelerator/Core/Public/UbaFileHandle.h
--- a/Engine/Source/Programs/UnrealBuildAccelerator/Core/Public/UbaFileHandle.h
+++ b/Engine/Source/Programs/UnrealBuildAccelerator/Core/Public/UbaFileHandle.h
@@ -13,8 +13,12 @@ namespace uba
 
 	inline constexpr u64 FileHandleFlagMask = 0x0000'0000'ffff'ffff;
 
-	#if !PLATFORM_WINDOWS
-	inline int asFileDescriptor(FileHandle fh) { return fh == InvalidFileHandle ? (int)fh : (int)(fh & FileHandleFlagMask); }
+	////////////////////////////////////////////////////////////////////////////////////////////////////
+
+	#if PLATFORM_WINDOWS
+	inline HANDLE AsHANDLE(FileHandle fh) { return (HANDLE)(fh == InvalidFileHandle ? InvalidFileHandle : (fh & FileHandleFlagMask)); }
+	#else
+	inline int AsFileDescriptor(FileHandle fh) { return fh == InvalidFileHandle ? (int)fh : (int)(fh & FileHandleFlagMask); }
 	#endif
 
 	////////////////////////////////////////////////////////////////////////////////////////////////////
diff --git a/Engine/Source/Programs/UnrealBuildAccelerator/Core/Public/UbaFileMapping.inl b/Engine/Source/Programs/UnrealBuildAccelerator/Core/Public/UbaFileMapping.inl
--- a/Engine/Source/Programs/UnrealBuildAccelerator/Core/Public/UbaFileMapping.inl
+++ b/Engine/Source/Programs/UnrealBuildAccelerator/Core/Public/UbaFileMapping.inl
@@ -7,8 +7,6 @@
 ////////////////////////////////////////////////////////////////////////////////////////////////////
 
 #if PLATFORM_WINDOWS
-HANDLE asHANDLE(FileMappingHandle h) { return (HANDLE)(uintptr_t)u32(h.internal); }
-HANDLE asHANDLE(FileHandle fh) { return (HANDLE)(fh == InvalidFileHandle ? InvalidFileHandle : (fh & FileHandleFlagMask)); }
 Bottleneck& g_createFileHandleBottleneck = *new Bottleneck(8); // Allocated and leaked just to prevent shutdown asserts in debug
 
 HANDLE InternalCreateFileMappingW(Logger& logger, HANDLE hFile, DWORD flProtect, DWORD dwMaximumSizeHigh, DWORD dwMaximumSizeLow, LPCWSTR lpName, const tchar* hint)
@@ -245,10 +243,10 @@ FileMappingHandle CreateFileMappingW(Logger& logger, FileHandle file, u32 protec
 {
 #if PLATFORM_WINDOWS
 	ExtendedTimerScope ts(KernelStats::GetCurrent().createFileMapping);
-	return u64(InternalCreateFileMappingW(logger, asHANDLE(file), protect, (DWORD)ToHigh(maxSize), ToLow(maxSize), NULL, hint));
+	return u64(InternalCreateFileMappingW(logger, AsHANDLE(file), protect, (DWORD)ToHigh(maxSize), ToLow(maxSize), NULL, hint));
 #else
 	FileMappingHandle h;
-	int fd = asFileDescriptor(file);
+	int fd = AsFileDescriptor(file);
 	if (maxSize && (protect & (~PAGE_READONLY)) != 0)
 	{
 #if PLATFORM_MAC // For some reason lseek+write does not work on apple silicon platform
@@ -284,7 +282,7 @@ u8* MapViewOfFile(Logger& logger, FileMappingHandle fileMappingObject, u32 desir
 {
 #if PLATFORM_WINDOWS
 	ExtendedTimerScope ts(KernelStats::GetCurrent().mapViewOfFile);
-	u8* res = (u8*)::MapViewOfFile(asHANDLE(fileMappingObject), desiredAccess, (DWORD)ToHigh(offset), ToLow(offset), bytesToMap);
+	u8* res = (u8*)::MapViewOfFile(AsHANDLE(fileMappingObject), desiredAccess, (DWORD)ToHigh(offset), ToLow(offset), bytesToMap);
 	if (!res)
 		return nullptr;
 	if (allowDiscard)
@@ -381,7 +379,7 @@ bool CloseFileMapping(Logger& logger, FileMappingHandle h, const tchar* hint)
 #endif
 
 #if PLATFORM_WINDOWS
-	return ::CloseHandle(asHANDLE(h));
+	return ::CloseHandle(AsHANDLE(h));
 #else
 	if (h.shmFd != -1 && h.internal != ~0ull)
 		if (close(h.shmFd) != 0)
@@ -443,7 +441,7 @@ void* MapMemory(FileMappingHandle handle, MemoryMapType mapType, void* targetAdd
 
 	//kernelStats.mapViewOfFile3.bytes += size;
 	//TimerScope ts(kernelStats.mapViewOfFile3);
-	void* viewA = MapViewOfFile3(asHANDLE(handle), GetCurrentProcess(), (u8*)targetAddress+targetOffset, handleOffset, size, MEM_REPLACE_PLACEHOLDER, protect, nullptr, 0);
+	void* viewA = MapViewOfFile3(AsHANDLE(handle), GetCurrentProcess(), (u8*)targetAddress+targetOffset, handleOffset, size, MEM_REPLACE_PLACEHOLDER, protect, nullptr, 0);
 #else
 	UBA_ASSERT(mapType != MemoryMapType_CopyOnWrite);
 	u32 protect = PROT_READ;
@@ -461,7 +459,7 @@ FileMappingHandle DuplicateMappingFromHost(FileMappingHandle handle, const char*
 {
 #if PLATFORM_WINDOWS
 	HANDLE realHandle;
-	if (!DuplicateHandle(g_hostProcess, asHANDLE(handle), GetCurrentProcess(), &realHandle, 0, FALSE, DUPLICATE_SAME_ACCESS))
+	if (!DuplicateHandle(g_hostProcess, AsHANDLE(handle), GetCurrentProcess(), &realHandle, 0, FALSE, DUPLICATE_SAME_ACCESS))
 	{
 		UBA_ASSERTF(false, L"Failed to duplicate %S handle (%u)", hint, GetLastError());
 		return {};
@@ -481,9 +479,9 @@ u8* GetMemoryFromMapping(FileMappingHandle handle, u64 size, u64 offset, bool wr
 	LARGE_INTEGER li;
 	li.QuadPart = (LONGLONG)offset;
 	//TimerScope ts(g_kernelStats.mapViewOfFile);
-	u8* mem = (u8*)::MapViewOfFile(asHANDLE(handle), dwDesiredAccess, li.HighPart, li.LowPart, size);
+	u8* mem = (u8*)::MapViewOfFile(AsHANDLE(handle), dwDesiredAccess, li.HighPart, li.LowPart, size);
 	UBA_ASSERT(mem);
-	::CloseHandle(asHANDLE(handle));
+	::CloseHandle(AsHANDLE(handle));
 	return mem;
 #else
 	char uidName[128];
diff --git a/Engine/Source/Programs/UnrealBuildAccelerator/Core/Public/UbaFileMappingHandle.h b/Engine/Source/Programs/UnrealBuildAccelerator/Core/Public/UbaFileMappingHandle.h
--- a/Engine/Source/Programs/UnrealBuildAccelerator/Core/Public/UbaFileMappingHandle.h
+++ b/Engine/Source/Programs/UnrealBuildAccelerator/Core/Public/UbaFileMappingHandle.h
@@ -24,4 +24,10 @@ namespace uba
 	};
 
 	////////////////////////////////////////////////////////////////////////////////////////////////////
+
+	#if PLATFORM_WINDOWS
+	inline HANDLE AsHANDLE(FileMappingHandle h) { return (HANDLE)(uintptr_t)u32(h.internal); }
+	#endif
+
+	////////////////////////////////////////////////////////////////////////////////////////////////////
 }
