[UBA]
* Added forward declaration of HANDLE on windows

diff --git a/Engine/Source/Programs/UnrealBuildAccelerator/Core/Public/UbaFileHandle.h b/Engine/Source/Programs/UnrealBuildAccelerator/Core/Public/UbaFileHandle.h
--- a/Engine/Source/Programs/UnrealBuildAccelerator/Core/Public/UbaFileHandle.h
+++ b/Engine/Source/Programs/UnrealBuildAccelerator/Core/Public/UbaFileHandle.h
@@ -16,6 +16,7 @@ namespace uba
 	////////////////////////////////////////////////////////////////////////////////////////////////////
 
 	#if PLATFORM_WINDOWS
+	typedef void *HANDLE;
 	inline HANDLE AsHANDLE(FileHandle fh) { return (HANDLE)(fh == InvalidFileHandle ? InvalidFileHandle : (fh & FileHandleFlagMask)); }
 	#else
 	inline int AsFileDescriptor(FileHandle fh) { return fh == InvalidFileHandle ? (int)fh : (int)(fh & FileHandleFlagMask); }
