[UBA]
* Added ReportBadCas message when client finds bad cas. Server will only print out a warning right now. TODO to implement this fully

diff --git a/Engine/Source/Programs/UnrealBuildAccelerator/Common/Private/UbaStorageServer.cpp b/Engine/Source/Programs/UnrealBuildAccelerator/Common/Private/UbaStorageServer.cpp
--- a/Engine/Source/Programs/UnrealBuildAccelerator/Common/Private/UbaStorageServer.cpp
+++ b/Engine/Source/Programs/UnrealBuildAccelerator/Common/Private/UbaStorageServer.cpp
@@ -1289,6 +1289,15 @@ namespace uba
 				GetProxy(writer, connectionInfo.GetId(), false);
 				return true;
 			}
+			case StorageMessageType_ReportBadCas:
+			{
+				CasKey casKey = reader.ReadCasKey();
+				TString file = reader.ReadString();
+
+				// TODO: Implement this, should drop the file
+				m_logger.Warning(TC("Client reported bad cas %s (%s). Todo: Validate cas and delete it if it is bad"), CasKeyString(casKey).str, file.c_str());
+				return true;
+			}
 		}
 		UBA_ASSERT(false);
 		return false;
diff --git a/Engine/Source/Programs/UnrealBuildAccelerator/Common/Private/UbaStorageUtils.cpp b/Engine/Source/Programs/UnrealBuildAccelerator/Common/Private/UbaStorageUtils.cpp
--- a/Engine/Source/Programs/UnrealBuildAccelerator/Common/Private/UbaStorageUtils.cpp
+++ b/Engine/Source/Programs/UnrealBuildAccelerator/Common/Private/UbaStorageUtils.cpp
@@ -2,9 +2,12 @@
 
 #include "UbaStorageUtils.h"
 #include "UbaCompressedFileHeader.h"
+#include "UbaFileMapping.h"
+#include "UbaNetworkMessage.h"
 #include "UbaStorage.h"
 #include "UbaFileAccessor.h"
 #include "UbaStats.h"
+#include "UbaWorkManager.h"
 #include <oodle2.h>
 
 namespace uba
@@ -633,6 +636,17 @@ namespace uba
 
 		u64 destOffset = 0;
 
+		auto ReportBadCas = [&]()
+			{
+				StackBinaryWriter<1024> writer;
+				NetworkMessage msg(client, StorageServiceId, StorageMessageType_ReportBadCas, writer);
+				writer.WriteCasKey(casKey);
+				writer.WriteString(destination);
+				StackBinaryReader<8> reader;
+				msg.Send(reader);
+			};
+
+
 		auto WriteDestination = [&](const void* source, u64 sourceSize)
 			{
 				if (fileMappingMem)
@@ -743,7 +757,10 @@ namespace uba
 						TimerScope ts2(m_stats.decompressRecv);
 						OO_SINTa decompLen = OodleLZ_Decompress(readBuffer, int(compressedSize), destMem, int(decompressedSize), OodleLZ_FuzzSafe_Yes, OodleLZ_CheckCRC_No, OodleLZ_Verbosity_None, NULL, 0, NULL, NULL, decoderMem, decoderMemSize);
 						if (decompLen != decompressedSize)
+						{
+							ReportBadCas();
 							return logger.Error(TC("Expected %u but got %i when decompressing %u bytes for file %s"), decompressedSize, int(decompLen), compressedSize, destination);
+						}
 
 						readBuffer += compressedSize;
 						destMem += decompressedSize;
@@ -819,8 +836,11 @@ namespace uba
 					if (sendSegmentMessage)
 					{
 						if (fetchId == u16(~0))
+						{
+							ReportBadCas();
 							return logger.Error(TC("Cas content error (2). Server believes %s was only one segment but client sees more. UncompressedSize: %llu LeftUncompressed: %llu Size: %llu Left to read: %llu ResponseSize: %u. (%s)"), destination, actualSize, leftUncompressed, fileSize, left, responseSize, CasKeyString(casKey).str);
-						
+						}
+
 						u64 capacity = maxReadSize - u64(readPosition - readBuffer);
 						u64 writeCapacity = capacity;
 						u8* writeDest = readPosition;
@@ -854,7 +874,10 @@ namespace uba
 					if (isFirstInBlock)
 					{
 						if (readPosition - readBuffer < sizeof(u32) * 2)
+						{
+							ReportBadCas();
 							return logger.Error(TC("Received less than minimum amount of data. Most likely corrupt cas file %s (Available: %u UncompressedSize: %llu LeftUncompressed: %llu)"), CasKeyString(casKey).str, u32(readPosition - readBuffer), actualSize, leftUncompressed);
+						}
 						isFirstInBlock = false;
 						u32* blockSize = (u32*)readBuffer;
 						compressedSize = blockSize[0];
@@ -909,7 +932,10 @@ namespace uba
 						TimerScope ts2(m_stats.decompressRecv);
 						OO_SINTa decompLen = OodleLZ_Decompress(readBuffer, int(compressedSize), decompressBuffer, int(decompressedSize));
 						if (decompLen != decompressedSize)
+						{
+							ReportBadCas();
 							return logger.Error(TC("Expected %u but got %i when decompressing %u bytes for file %s"), decompressedSize, int(decompLen), compressedSize, destination);
+						}
 					}
 
 					lastDecompressSize = decompressedSize;
diff --git a/Engine/Source/Programs/UnrealBuildAccelerator/Common/Public/UbaNetwork.h b/Engine/Source/Programs/UnrealBuildAccelerator/Common/Public/UbaNetwork.h
--- a/Engine/Source/Programs/UnrealBuildAccelerator/Common/Public/UbaNetwork.h
+++ b/Engine/Source/Programs/UnrealBuildAccelerator/Common/Public/UbaNetwork.h
@@ -42,6 +42,7 @@ namespace uba
 		UBA_STORAGE_MESSAGE(ProxyFetchBegin) \
 		UBA_STORAGE_MESSAGE(ProxyFetchEnd) \
 		UBA_STORAGE_MESSAGE(ReportBadProxy) \
+		UBA_STORAGE_MESSAGE(ReportBadCas) \
 
 	#define UBA_SESSION_MESSAGES \
 		UBA_SESSION_MESSAGE(Connect) \
