[UBA]
* Renamed file mapping functions to not collide with windows file mapping functions

diff --git a/Engine/Source/Programs/UnrealBuildAccelerator/Common/Private/UbaCacheClient.cpp b/Engine/Source/Programs/UnrealBuildAccelerator/Common/Private/UbaCacheClient.cpp
--- a/Engine/Source/Programs/UnrealBuildAccelerator/Common/Private/UbaCacheClient.cpp
+++ b/Engine/Source/Programs/UnrealBuildAccelerator/Common/Private/UbaCacheClient.cpp
@@ -1094,7 +1094,7 @@ namespace uba
 			{
 				if (!destFile.CreateMemoryWrite(false, DefaultAttributes(), localBlock.writtenSize))
 					return logger.Error(TC("Failed to create file for cache output %s for %s"), path.data, info.GetDescription());
-				MapMemoryCopy(destFile.GetData(), localBlock.memory, localBlock.writtenSize);
+				FileMapping_CopyMem(destFile.GetData(), localBlock.memory, localBlock.writtenSize);
 			}
 			else
 			{
diff --git a/Engine/Source/Programs/UnrealBuildAccelerator/Common/Private/UbaDirectoryTableHolder.cpp b/Engine/Source/Programs/UnrealBuildAccelerator/Common/Private/UbaDirectoryTableHolder.cpp
--- a/Engine/Source/Programs/UnrealBuildAccelerator/Common/Private/UbaDirectoryTableHolder.cpp
+++ b/Engine/Source/Programs/UnrealBuildAccelerator/Common/Private/UbaDirectoryTableHolder.cpp
@@ -2,6 +2,7 @@
 
 #include "UbaDirectoryTableHolder.h"
 #include "UbaDirectoryIterator.h"
+#include "UbaEnvironment.h"
 #include "UbaFileMapping.h"
 #include "UbaProtocol.h"
 
@@ -93,9 +94,9 @@ namespace uba
 	,	m_directoryTableMemory(TC("SessionDirectoryTable"))
 	,	m_directoryTable(m_directoryTableMemory)
 	{
-		m_directoryTableHandle = uba::CreateMemoryMappingW(m_logger, PAGE_READWRITE|SEC_RESERVE, DirTableMaxSize, nullptr, TC("DirMappings"));
+		m_directoryTableHandle = FileMapping_Create(m_logger, PAGE_READWRITE|SEC_RESERVE, DirTableMaxSize, nullptr, TC("DirMappings"));
 		UBA_ASSERT(m_directoryTableHandle.IsValid());
-		m_directoryTableMem = MapViewOfFile(m_logger, m_directoryTableHandle, FILE_MAP_WRITE, 0, DirTableMaxSize);
+		m_directoryTableMem = FileMapping_Map(m_logger, m_directoryTableHandle, FILE_MAP_WRITE, 0, DirTableMaxSize);
 		UBA_ASSERT(m_directoryTableMem);
 
 		m_directoryTable.m_memory = m_directoryTableMem;
@@ -104,8 +105,8 @@ namespace uba
 
 	DirectoryTableHolder::~DirectoryTableHolder()
 	{
-		UnmapViewOfFile(m_logger, m_directoryTableMem, DirTableMaxSize, TC("DirectoryTable"), false); // Can't discard since only parts might be mapped
-		CloseFileMapping(m_logger, m_directoryTableHandle, TC("DirectoryTable"));
+		FileMapping_Unmap(m_logger, m_directoryTableMem, DirTableMaxSize, TC("DirectoryTable"), false); // Can't discard since only parts might be mapped
+		FileMapping_Close(m_logger, m_directoryTableHandle, TC("DirectoryTable"));
 	}
 
 	void DirectoryTableHolder::SetTreatTempDirAsEmpty()
@@ -204,9 +205,9 @@ namespace uba
 	{
 		UBA_ASSERT(!overlay.handle.IsValid());
 		SCOPED_WRITE_LOCK(overlay.lock, lock);
-		overlay.handle = uba::CreateMemoryMappingW(m_logger, PAGE_READWRITE|SEC_RESERVE, OverlayTableMaxSize, nullptr, TC("Overlay"));
+		overlay.handle = FileMapping_Create(m_logger, PAGE_READWRITE|SEC_RESERVE, OverlayTableMaxSize, nullptr, TC("Overlay"));
 		UBA_ASSERT(overlay.handle.IsValid());
-		overlay.memory = MapViewOfFile(m_logger, overlay.handle, FILE_MAP_WRITE, 0, OverlayTableMaxSize);
+		overlay.memory = FileMapping_Map(m_logger, overlay.handle, FILE_MAP_WRITE, 0, OverlayTableMaxSize);
 		UBA_ASSERT(overlay.memory);
 	}
 
@@ -235,7 +236,7 @@ namespace uba
 
 		u64 toCommit = newSize - overlay.committed;
 		u8* address = overlay.memory + overlay.committed;
-		if (!MapViewCommit(m_logger, address, toCommit, false))
+		if (!FileMapping_Commit(m_logger, address, toCommit, false))
 			m_logger.Error(TC("Failed to commit memory for directory table (Committed: %llu, ToCommit: %llu) (%s)"), m_directoryTableMemCommitted, toCommit, LastErrorToText().data);
 		overlay.committed += u32(toCommit);
 		return true;
@@ -246,9 +247,9 @@ namespace uba
 		if (!overlay.handle.IsValid())
 			return;
 		SCOPED_WRITE_LOCK(overlay.lock, lock);
-		UnmapViewOfFile(m_logger, overlay.memory, OverlayTableMaxSize, TC("Overlay"), false); // Can't discard since only parts might be mapped
+		FileMapping_Unmap(m_logger, overlay.memory, OverlayTableMaxSize, TC("Overlay"), false); // Can't discard since only parts might be mapped
 		overlay.memory = nullptr;
-		CloseFileMapping(m_logger, overlay.handle, TC("Overlay"));
+		FileMapping_Close(m_logger, overlay.handle, TC("Overlay"));
 		overlay.handle = {};
 		overlay.committed = 0;
 		overlay.size = 0;
@@ -990,7 +991,7 @@ namespace uba
 
 		u64 toCommit = newSize - m_directoryTableMemCommitted;
 		u8* address = m_directoryTableMem + m_directoryTableMemCommitted;
-		if (!MapViewCommit(m_logger, address, toCommit, false))
+		if (!FileMapping_Commit(m_logger, address, toCommit, false))
 			m_logger.Error(TC("Failed to commit memory for directory table (Committed: %llu, ToCommit: %llu) (%s)"), m_directoryTableMemCommitted, toCommit, LastErrorToText().data);
 		m_directoryTableMemCommitted += toCommit;
 		LockMemory(m_logger, address, toCommit, TC("DirectoryTable"));
diff --git a/Engine/Source/Programs/UnrealBuildAccelerator/Common/Private/UbaFileAccessor.cpp b/Engine/Source/Programs/UnrealBuildAccelerator/Common/Private/UbaFileAccessor.cpp
--- a/Engine/Source/Programs/UnrealBuildAccelerator/Common/Private/UbaFileAccessor.cpp
+++ b/Engine/Source/Programs/UnrealBuildAccelerator/Common/Private/UbaFileAccessor.cpp
@@ -206,11 +206,11 @@ namespace uba
 			return true;
 		}
 
-		m_mappingHandle = uba::CreateFileMappingW(m_logger, m_fileHandle, PAGE_READWRITE, size, realFileName);
+		m_mappingHandle = FileMapping_CreateFromFile(m_logger, m_fileHandle, PAGE_READWRITE, size, realFileName);
 		if (!m_mappingHandle.IsValid())
 			return m_logger.Error(TC("Failed to create memory map %s with size %llu (%s)"), realFileName, size, LastErrorToText().data);
 
-		m_data = MapViewOfFile(m_logger, m_mappingHandle, FILE_MAP_WRITE, 0, size);
+		m_data = FileMapping_Map(m_logger, m_mappingHandle, FILE_MAP_WRITE, 0, size);
 		if (!m_data)
 			return m_logger.Error(TC("Failed to map view of file %s with size %llu, for write (%s)"), realFileName, size, LastErrorToText().data);
 
@@ -442,9 +442,9 @@ namespace uba
 		m_size = info.size;
 #if PLATFORM_WINDOWS
 		if (m_size)
-			m_mappingHandle = uba::CreateFileMappingW(m_logger, m_fileHandle, PAGE_READONLY, m_size, m_fileName);
+			m_mappingHandle = FileMapping_CreateFromFile(m_logger, m_fileHandle, PAGE_READONLY, m_size, m_fileName);
 		else
-			m_mappingHandle = uba::CreateFileMappingW(m_logger, InvalidFileHandle, PAGE_READONLY, 1, m_fileName);
+			m_mappingHandle = FileMapping_CreateFromFile(m_logger, InvalidFileHandle, PAGE_READONLY, 1, m_fileName);
 
 		if (!m_mappingHandle.IsValid())
 			return m_logger.Error(TC("Failed to create mapping handle for %s with size %llu (%s)"), m_fileName, m_size, LastErrorToText().data);
@@ -454,7 +454,7 @@ namespace uba
 		if (offset == m_size)
 			return true;
 #endif
-		m_data = MapViewOfFile(m_logger, m_mappingHandle, FILE_MAP_READ, offset, m_size);
+		m_data = FileMapping_Map(m_logger, m_mappingHandle, FILE_MAP_READ, offset, m_size);
 		if (!m_data)
 			return m_logger.Error(TC("%s - MapViewOfFile failed (%s)"), m_fileName, LastErrorToText().data);
 
@@ -470,14 +470,14 @@ namespace uba
 	{
 		if (m_data)
 		{
-			if (m_size && !UnmapViewOfFile(m_logger, m_data, m_size, m_fileName))
+			if (m_size && !FileMapping_Unmap(m_logger, m_data, m_size, m_fileName))
 				return m_logger.Error(TC("Failed to unmap memory for %s (%s)"), m_fileName, LastErrorToText().data);
 			m_data = nullptr;
 		}
 
 		if (m_mappingHandle.IsValid())
 		{
-			if (!CloseFileMapping(m_logger, m_mappingHandle, m_fileName))
+			if (!FileMapping_Close(m_logger, m_mappingHandle, m_fileName))
 				return m_logger.Error(TC("Failed to close file mapping for %s (%s)"), m_fileName, LastErrorToText().data);
 			m_mappingHandle = {};
 		}
diff --git a/Engine/Source/Programs/UnrealBuildAccelerator/Common/Private/UbaFileMapping.cpp b/Engine/Source/Programs/UnrealBuildAccelerator/Common/Private/UbaFileMapping.cpp
--- a/Engine/Source/Programs/UnrealBuildAccelerator/Common/Private/UbaFileMapping.cpp
+++ b/Engine/Source/Programs/UnrealBuildAccelerator/Common/Private/UbaFileMapping.cpp
@@ -39,7 +39,7 @@ namespace uba
 
 namespace uba
 {
-	void MapMemoryCopy(void* dest, const void* source, u64 size)
+	void FileMapping_CopyMem(void* dest, const void* source, u64 size)
 	{
 		auto& stats = KernelStats::GetCurrent();
 		ExtendedTimerScope ts(stats.memoryCopy);
diff --git a/Engine/Source/Programs/UnrealBuildAccelerator/Common/Private/UbaFixedSizeFileMappingAllocator.cpp b/Engine/Source/Programs/UnrealBuildAccelerator/Common/Private/UbaFixedSizeFileMappingAllocator.cpp
--- a/Engine/Source/Programs/UnrealBuildAccelerator/Common/Private/UbaFixedSizeFileMappingAllocator.cpp
+++ b/Engine/Source/Programs/UnrealBuildAccelerator/Common/Private/UbaFixedSizeFileMappingAllocator.cpp
@@ -20,12 +20,12 @@ namespace uba
 	FixedSizeFileMappingAllocator::~FixedSizeFileMappingAllocator()
 	{
 		if (m_mappingHandle.IsValid())
-			CloseFileMapping(m_logger, m_mappingHandle, TC("FixedSizeFileMappingAllocator"));
+			FileMapping_Close(m_logger, m_mappingHandle, TC("FixedSizeFileMappingAllocator"));
 	}
 
 	bool FixedSizeFileMappingAllocator::Init(u64 blockSize, u64 capacity)
 	{
-		m_mappingHandle = uba::CreateMemoryMappingW(m_logger, PAGE_READWRITE|SEC_RESERVE, capacity, nullptr, TC("FixedSizeFileMappingAllocator"));
+		m_mappingHandle = FileMapping_Create(m_logger, PAGE_READWRITE|SEC_RESERVE, capacity, nullptr, TC("FixedSizeFileMappingAllocator"));
 		if (!m_mappingHandle.IsValid())
 			return m_logger.Error(TC("%s - Failed to create memory map with capacity %llu (%s)"), m_name, capacity, LastErrorToText().data);
 
@@ -54,7 +54,7 @@ namespace uba
 		lock.Leave();
 
 		u64 offset = index*m_blockSize;
-		u8* data = MapViewOfFile(m_logger, m_mappingHandle, FILE_MAP_READ|FILE_MAP_WRITE, offset, m_blockSize, false);
+		u8* data = FileMapping_Map(m_logger, m_mappingHandle, FILE_MAP_READ|FILE_MAP_WRITE, offset, m_blockSize, false);
 		if (!data)
 		{
 			if (m_capacity < m_mappingCount*m_blockSize)
@@ -66,7 +66,7 @@ namespace uba
 
 		if (needCommit)
 		{
-			if (!MapViewCommit(m_logger, data, m_blockSize, true))
+			if (!FileMapping_Commit(m_logger, data, m_blockSize, true))
 			{
 				m_logger.Error(TC("%s - Failed to allocate memory for %s (%s)"), m_name, hint, LastErrorToText().data);
 				return { {}, 0, 0 };
@@ -80,7 +80,7 @@ namespace uba
 	void FixedSizeFileMappingAllocator::Free(FileMappingAllocation allocation)
 	{
 		UBA_ASSERT(allocation.handle == m_mappingHandle);
-		if (!UnmapViewOfFile(m_logger, allocation.memory, m_blockSize, m_name, true))
+		if (!FileMapping_Unmap(m_logger, allocation.memory, m_blockSize, m_name, true))
 			m_logger.Error(TC("%s - Failed to unmap view of file (%s)"), m_name, LastErrorToText().data);
 		u64 index = allocation.offset / m_blockSize;
 		SCOPED_FUTEX(m_mappingLock, lock);
diff --git a/Engine/Source/Programs/UnrealBuildAccelerator/Common/Private/UbaSession.cpp b/Engine/Source/Programs/UnrealBuildAccelerator/Common/Private/UbaSession.cpp
--- a/Engine/Source/Programs/UnrealBuildAccelerator/Common/Private/UbaSession.cpp
+++ b/Engine/Source/Programs/UnrealBuildAccelerator/Common/Private/UbaSession.cpp
@@ -530,7 +530,7 @@ namespace uba
 		{
 			if (!destinationFile.CreateMemoryWrite(false, attributes, fileSize, m_tempPath.data))
 				return false;
-			MapMemoryCopy(destinationFile.GetData(), fileMem, fileSize);
+			FileMapping_CopyMem(destinationFile.GetData(), fileMem, fileSize);
 		}
 		else
 		{
@@ -770,7 +770,7 @@ namespace uba
 			//FileHandle tempFileHandle;
 			//if (!OpenFileSequentialRead(m_logger, tempFileName.data, tempFileHandle))
 			//	return false;
-			FileMappingHandle mappingHandle = CreateFileMappingW(m_logger, tempFileHandle, PAGE_READONLY, contentSize, tempFileName.data);
+			FileMappingHandle mappingHandle = FileMapping_CreateFromFile(m_logger, tempFileHandle, PAGE_READONLY, contentSize, tempFileName.data);
 			if (!mappingHandle.IsValid())
 				return false;
 			//uba::CloseFile(tempFileName.data, tempFileHandle); // it seems like this can't be closed because then the mapping might go away
@@ -1338,9 +1338,9 @@ namespace uba
 		m_systemPath.Append(TCV("/nonexistingpath"));
 		#endif
 
-		m_fileMappingTableHandle = uba::CreateMemoryMappingW(m_logger, PAGE_READWRITE, FileMappingTableMaxSize, nullptr, TC("FileMappings"));
+		m_fileMappingTableHandle = FileMapping_Create(m_logger, PAGE_READWRITE, FileMappingTableMaxSize, nullptr, TC("FileMappings"));
 		UBA_ASSERT(m_fileMappingTableHandle.IsValid());
-		m_fileMappingTableMem = MapViewOfFile(m_logger, m_fileMappingTableHandle, FILE_MAP_WRITE, 0, FileMappingTableMaxSize, m_allowDiscardVirtualMemory);
+		m_fileMappingTableMem = FileMapping_Map(m_logger, m_fileMappingTableHandle, FILE_MAP_WRITE, 0, FileMappingTableMaxSize, m_allowDiscardVirtualMemory);
 		LockMemory(m_logger, m_fileMappingTableMem, FileMappingTableMaxSize, TC("FileMappings"));
 		UBA_ASSERT(m_fileMappingTableMem);
 		m_fileMappingTableLookup.reserve(70000);
@@ -1491,7 +1491,7 @@ namespace uba
 		{
 			VirtualSourceFile& file = kv.second;
 			if (file.mappingHandle.IsValid())
-				CloseFileMapping(m_logger, file.mappingHandle, TC("VirtualFile"));
+				FileMapping_Close(m_logger, file.mappingHandle, TC("VirtualFile"));
 			if (file.memoryHandle.IsValid())
 				m_sharedMemory.CloseHandle(m_logger, file.memoryHandle, TC("VirtualFile"));
 		}
@@ -1500,8 +1500,8 @@ namespace uba
 			if (i.second.canBeFreed)
 				m_sharedMemory.CloseHandle(m_logger, i.second.memoryHandle, TC("FileMappingKeptFromOutput"));
 
-		UnmapViewOfFile(m_logger, m_fileMappingTableMem, FileMappingTableMaxSize, TC("FileMappingTable"), m_allowDiscardVirtualMemory);
-		CloseFileMapping(m_logger, m_fileMappingTableHandle, TC("FileMappingTable"));
+		FileMapping_Unmap(m_logger, m_fileMappingTableMem, FileMappingTableMaxSize, TC("FileMappingTable"), m_allowDiscardVirtualMemory);
+		FileMapping_Close(m_logger, m_fileMappingTableHandle, TC("FileMappingTable"));
 
 		//#if !UBA_DEBUG
 		//u32 count;
@@ -3269,7 +3269,7 @@ namespace uba
 			u64 fileSize = 0;
 			if (!GetFileSizeEx(fileSize, fileHandle))
 				return m_logger.Error(TC("[RegisterVirtualFileInternal] GetFileSizeEx for %s failed (%s)"), sourceFile, LastErrorToText().data);
-			virtualFile.mappingHandle = CreateFileMappingW(m_logger, fileHandle, PAGE_READONLY, fileSize, sourceFile);
+			virtualFile.mappingHandle = FileMapping_CreateFromFile(m_logger, fileHandle, PAGE_READONLY, fileSize, sourceFile);
 			virtualFile.size = fileSize;
 			insres.first->second = virtualFile;
 		}
@@ -3304,7 +3304,7 @@ namespace uba
 		u8* mem2 = m_sharedMemory.MapView(mapping, filePath.data, SharedMemoryMapType_ReadWrite);
 		if (!mem2)
 			return false;
-		MapMemoryCopy(mem2, memory, memorySize);
+		FileMapping_CopyMem(mem2, memory, memorySize);
 		m_sharedMemory.UnmapView(mapping, mem2);
 
 		virtualFile.memoryHandle = mapping;
@@ -3371,7 +3371,7 @@ namespace uba
 
 	bool Session::GetOutputFileDataInternal(void* outData, const StringKey& fileNameKey, StringView filePath, bool deleteInternalMapping)
 	{
-		return GetFileMemory([&](const void* fileMem, u64 fileSize) { MapMemoryCopy(outData, fileMem, fileSize); return true; }, fileNameKey, filePath, deleteInternalMapping);
+		return GetFileMemory([&](const void* fileMem, u64 fileSize) { FileMapping_CopyMem(outData, fileMem, fileSize); return true; }, fileNameKey, filePath, deleteInternalMapping);
 	}
 
 	bool Session::WriteOutputFileInternal(const StringKey& fileNameKey, StringView filePath, bool deleteInternalMapping)
@@ -3524,7 +3524,7 @@ namespace uba
 			if (!mem2)
 				return false;
 
-			MapMemoryCopy(mem2, memoryBlock.memory, memoryBlock.writtenSize);
+			FileMapping_CopyMem(mem2, memoryBlock.memory, memoryBlock.writtenSize);
 			m_sharedMemory.UnmapView(symHandle, mem2);
 
 			lastWriteTime = GetSystemTimeAsFileTime();
diff --git a/Engine/Source/Programs/UnrealBuildAccelerator/Common/Private/UbaSharedMemoryAllocator.cpp b/Engine/Source/Programs/UnrealBuildAccelerator/Common/Private/UbaSharedMemoryAllocator.cpp
--- a/Engine/Source/Programs/UnrealBuildAccelerator/Common/Private/UbaSharedMemoryAllocator.cpp
+++ b/Engine/Source/Programs/UnrealBuildAccelerator/Common/Private/UbaSharedMemoryAllocator.cpp
@@ -70,22 +70,22 @@ namespace uba
 			if (tempFileSparse)
 				if (!MakeSparse(m_logger, tempFile, m_tempFile, m_capacity))
 					return false;
-			m_fileMapping = CreateFileMappingW(m_logger, m_tempFile, PAGE_READWRITE, m_capacity, TC("SharedMemoryAllocator"));
+			m_fileMapping = FileMapping_CreateFromFile(m_logger, m_tempFile, PAGE_READWRITE, m_capacity, TC("SharedMemoryAllocator"));
 		}
 		else
 		#endif
 		{
-			m_fileMapping = CreateMemoryMappingW(m_logger, PAGE_READWRITE|SEC_RESERVE, m_capacity, nullptr, TC("SharedMemoryAllocator"));
+			m_fileMapping = FileMapping_Create(m_logger, PAGE_READWRITE|SEC_RESERVE, m_capacity, nullptr, TC("SharedMemoryAllocator"));
 		}
 		if (!m_fileMapping.IsValid())
 			return false;
 
-		m_memory = MapViewOfFile(m_logger, m_fileMapping, FILE_MAP_READ|FILE_MAP_WRITE, 0, m_capacity);
+		m_memory = FileMapping_Map(m_logger, m_fileMapping, FILE_MAP_READ|FILE_MAP_WRITE, 0, m_capacity);
 		UBA_ASSERTF(m_memory, TC("MapViewOfFile failed trying to map %llu bytes (%s)"), m_capacity, LastErrorToText().data);
 
 		if (!m_memory)
 		{
-			CloseFileMapping(m_logger, m_fileMapping, TC("SharedMemoryAllocator"));
+			FileMapping_Close(m_logger, m_fileMapping, TC("SharedMemoryAllocator"));
 			return false;
 		}
 
@@ -96,11 +96,11 @@ namespace uba
 	void SharedMemoryAllocator::Deinit(bool keepFileMapping)
 	{
 		if (m_memory)
-			UnmapViewOfFile(m_logger, m_memory, m_capacity, TC("SharedMemoryAllocator"));
+			FileMapping_Unmap(m_logger, m_memory, m_capacity, TC("SharedMemoryAllocator"));
 		m_memory = nullptr;
 
 		if (m_fileMapping.IsValid() && !keepFileMapping)
-			CloseFileMapping(m_logger, m_fileMapping, TC("SharedMemoryAllocator"));
+			FileMapping_Close(m_logger, m_fileMapping, TC("SharedMemoryAllocator"));
 		m_fileMapping = {};
 
 		if (m_tempFile != InvalidFileHandle && !keepFileMapping)
@@ -294,7 +294,7 @@ namespace uba
 		UBA_ASSERT(h.slices.size() == 1);
 		auto& slice = h.slices.back();
 		if (m_tempFile == InvalidFileHandle)
-			MapViewCommit(m_logger, m_memory + slice.offset + offset, size);
+			FileMapping_Commit(m_logger, m_memory + slice.offset + offset, size);
 
 		SCOPED_FUTEX(m_lock, lock);
 		m_committed += size;
@@ -404,7 +404,7 @@ namespace uba
 			// TODO: Right now we leak this mapping. Not good..
 			// .. but only used in very special situations
 			FileMappingHandle file = FileMappingHandle::FromU64(handle.internalHandle & ~0x8000'0000'0000'0000);
-			u8* mem = MapViewOfFile(m_logger, file, FILE_MAP_READ, offset, size);
+			u8* mem = FileMapping_Map(m_logger, file, FILE_MAP_READ, offset, size);
 			UBA_ASSERT(mem);
 			view.ReferenceExternalMemory(mem, size);
 			return true;
@@ -492,7 +492,7 @@ namespace uba
 		auto& kernelStats = KernelStats::GetCurrent();
 
 		TimerScope ts(kernelStats.virtualAlloc2);
-		u8* memory = (u8*)uba::ReserveMemory(nullptr, committed);
+		u8* memory = (u8*)FileMapping_ReservePlaceholder(nullptr, committed);
 		ts.Leave();
 
 		u64 pos = 0;
@@ -501,7 +501,7 @@ namespace uba
 			kernelStats.mapViewOfFile3.bytes += slice.size;
 			TimerScope ts2(kernelStats.mapViewOfFile3);
 
-			void* viewA = MapMemory(m_fileMapping, (MemoryMapType)mapType, memory, pos, committed, slice.offset, slice.size);
+			void* viewA = FileMapping_MapPlaceholder(m_fileMapping, (MemoryMapType)mapType, memory, pos, committed, slice.offset, slice.size);
 			UBA_ASSERTF(viewA == memory+pos, TC("MapViewOfFile3 failed in MapView 0x%p %llu for % s (%s)"), memory, pos, h.GetHint(), LastErrorToText().data); (void)viewA;
 			pos += slice.size;
 		}
@@ -585,7 +585,7 @@ namespace uba
 	void SharedMemoryAllocator::UnmapExternalMapping(SharedMemoryHandle handle, u8* memory, u64 size)
 	{
 		UBA_ASSERT((handle.internalHandle & 0x8000'0000'0000'0000) != 0);
-		bool res = UnmapViewOfFile(m_logger, memory, size, TC("ExternalMapping"));
+		bool res = FileMapping_Unmap(m_logger, memory, size, TC("ExternalMapping"));
 		UBA_ASSERT(res); (void)res;
 	}
 
@@ -698,7 +698,7 @@ namespace uba
 			UBA_SHIPPING_ASSERTF(m_reserved + toCommit <= m_capacity, TC("Not enough space in capacity for commit (Capacity: %llu, Reserved: %llu Requested: %llu"), m_capacity, m_reserved, toCommit);
 			if (m_tempFile == InvalidFileHandle)
 			{
-				bool res = MapViewCommit(m_logger, m_memory + m_reserved, toCommit);
+				bool res = FileMapping_Commit(m_logger, m_memory + m_reserved, toCommit);
 				UBA_SHIPPING_ASSERTF(res, TC("VirtualAlloc failed commit %llu bytes at address 0x%llx + %llu (%s)"), toCommit, u64(m_memory), m_reserved, LastErrorToText().data);
 			}
 			u64 offset = m_reserved;
diff --git a/Engine/Source/Programs/UnrealBuildAccelerator/Common/Private/UbaStorage.cpp b/Engine/Source/Programs/UnrealBuildAccelerator/Common/Private/UbaStorage.cpp
--- a/Engine/Source/Programs/UnrealBuildAccelerator/Common/Private/UbaStorage.cpp
+++ b/Engine/Source/Programs/UnrealBuildAccelerator/Common/Private/UbaStorage.cpp
@@ -2622,7 +2622,7 @@ namespace uba
 					if (mappedView.memory)
 					{
 						TimerScope cts(stats.memoryCopy);
-						MapMemoryCopy(writePos, mappedView.memory, compressedFileSize);
+						FileMapping_CopyMem(writePos, mappedView.memory, compressedFileSize);
 					}
 					else
 					{
@@ -2668,7 +2668,7 @@ namespace uba
 					return false;
 				u8* writePos = destinationFile.GetData();
 				TimerScope cts(stats.memoryCopy);
-				MapMemoryCopy(writePos, mappedView.memory, mappedView.size);
+				FileMapping_CopyMem(writePos, mappedView.memory, mappedView.size);
 				return done(destinationFile, mappedView.size);
 			}
 
@@ -2900,14 +2900,14 @@ namespace uba
 
 		if (fileSize > BufferSlotSize) // Note that when filesize is larger than BufferSlotSize the hash becomes a hash of hashes
 		{
-			FileMappingHandle fileMapping = uba::CreateFileMappingW(m_logger, fileHandle, PAGE_READONLY, fileSize, fileName);
+			FileMappingHandle fileMapping = FileMapping_CreateFromFile(m_logger, fileHandle, PAGE_READONLY, fileSize, fileName);
 			if (!fileMapping.IsValid())
 			{
 				m_logger.Error(TC("Failed to create file mapping for %s (%s)"), fileName, LastErrorToText().data);
 				return CasKeyZero;
 			}
-			auto fmg = MakeGuard([&]() { CloseFileMapping(m_logger, fileMapping, fileName); });
-			u8* fileData = MapViewOfFile(m_logger, fileMapping, FILE_MAP_READ, 0, fileSize);
+			auto fmg = MakeGuard([&]() { FileMapping_Close(m_logger, fileMapping, fileName); });
+			u8* fileData = FileMapping_Map(m_logger, fileMapping, FILE_MAP_READ, 0, fileSize);
 			if (!fileData)
 			{
 				m_logger.Error(TC("Failed to map view of file mapping for %s (%s)"), fileName, LastErrorToText().data);
@@ -2915,9 +2915,9 @@ namespace uba
 			}
 			auto udg = MakeGuard([&]() { 
 					if (m_asyncUnmapViewOfFile)
-						m_workManager.AddWork([this, fileData, fileSize, fn = TString(fileName)](const WorkContext&) { UnmapViewOfFile(m_logger, fileData, fileSize, fn.c_str()); }, 1, TC("UnmapFile"));
+						m_workManager.AddWork([this, fileData, fileSize, fn = TString(fileName)](const WorkContext&) { FileMapping_Unmap(m_logger, fileData, fileSize, fn.c_str()); }, 1, TC("UnmapFile"));
 					else
-						UnmapViewOfFile(m_logger, fileData, fileSize, fileName);
+						FileMapping_Unmap(m_logger, fileData, fileSize, fileName);
 				});
 
 			return uba::CalculateCasKey(fileData, fileSize, storeCompressed, &m_workManager, fileName);
diff --git a/Engine/Source/Programs/UnrealBuildAccelerator/Common/Private/UbaStorageClient.cpp b/Engine/Source/Programs/UnrealBuildAccelerator/Common/Private/UbaStorageClient.cpp
--- a/Engine/Source/Programs/UnrealBuildAccelerator/Common/Private/UbaStorageClient.cpp
+++ b/Engine/Source/Programs/UnrealBuildAccelerator/Common/Private/UbaStorageClient.cpp
@@ -645,7 +645,7 @@ namespace uba
 					}
 					else
 					{
-						MapMemoryCopy(writePos, readBuffer, responseSize);
+						FileMapping_CopyMem(writePos, readBuffer, responseSize);
 						writePos += responseSize;
 					}
 
@@ -795,7 +795,7 @@ namespace uba
 							}
 							else
 							{
-								MapMemoryCopy(writePos, decompressBuffer, uncompressedSize);
+								FileMapping_CopyMem(writePos, decompressBuffer, uncompressedSize);
 								writePos += uncompressedSize;
 							}
 
diff --git a/Engine/Source/Programs/UnrealBuildAccelerator/Common/Private/UbaStorageServer.cpp b/Engine/Source/Programs/UnrealBuildAccelerator/Common/Private/UbaStorageServer.cpp
--- a/Engine/Source/Programs/UnrealBuildAccelerator/Common/Private/UbaStorageServer.cpp
+++ b/Engine/Source/Programs/UnrealBuildAccelerator/Common/Private/UbaStorageServer.cpp
@@ -177,9 +177,9 @@ namespace uba
 			++server.m_activeUnmap;
 			server.GetServer().AddWork([&server, mb = memoryBegin, mps = mappedView.size, fmp = fileMappingHandle, rfh = readFileHandle, reason](const WorkContext& context)
 				{
-					if (!UnmapViewOfFile(server.m_logger, mb, mps, reason))
+					if (!FileMapping_Unmap(server.m_logger, mb, mps, reason))
 						server.m_logger.Error(TC("%s - Failed to unmap memory at 0x%llx with size %llu (%s)"), reason, mb, mps, LastErrorToText().data);
-					if (!CloseFileMapping(server.m_logger, fmp, reason))
+					if (!FileMapping_Close(server.m_logger, fmp, reason))
 						server.m_logger.Error(TC("%s - Failed to close file mapping %llu (%s)"), reason, fmp.ToU64(), LastErrorToText().data);
 					if (!CloseFile(nullptr, rfh))
 						server.m_logger.Error(TC("%s - Failed to close file (%s)"), reason, LastErrorToText().data);
@@ -728,10 +728,10 @@ namespace uba
 						{
 							todoRemoveMe(12);
 							if (mappedView.memory)
-								if (!UnmapViewOfFile(m_logger, mappedView.memory, fileSize, TC("FetchBegin")))
+								if (!FileMapping_Unmap(m_logger, mappedView.memory, fileSize, TC("FetchBegin")))
 									m_logger.Error(TC("Failed to unmap memory %s at 0x%llx with size %llu (%s)"), hint.data, mappedView.memory, fileSize, LastErrorToText().data);
 							if (fileMappingHandle.IsValid())
-								if (!CloseFileMapping(m_logger, fileMappingHandle, hint.data))
+								if (!FileMapping_Close(m_logger, fileMappingHandle, hint.data))
 									m_logger.Error(TC("Failed to close file mapping for %s with handle %llu (%s)"), hint.data, fileMappingHandle.ToU64(), LastErrorToText().data);
 						}
 						else
@@ -768,11 +768,11 @@ namespace uba
 					if (fileSize > BufferSlotSize)
 					{
 						ownsMapping = true;
-						fileMappingHandle = CreateFileMappingW(m_logger, readFileHandle, PAGE_READONLY, fileSize, casFile.data);
+						fileMappingHandle = FileMapping_CreateFromFile(m_logger, readFileHandle, PAGE_READONLY, fileSize, casFile.data);
 						if (!fileMappingHandle.IsValid())
 							return m_logger.Error(TC("Failed to create file mapping of %s (%s)"), casFile.data, LastErrorToText().data);
 						u64 offset = memoryPos - memoryBegin;
-						mappedView.memory = MapViewOfFile(m_logger, fileMappingHandle, FILE_MAP_READ, 0, fileSize);
+						mappedView.memory = FileMapping_Map(m_logger, fileMappingHandle, FILE_MAP_READ, 0, fileSize);
 						if (!mappedView.memory)
 							return m_logger.Error(TC("Failed to map memory of %s with size %llu (%s)"), casFile.data, fileSize, LastErrorToText().data);
 						memoryBegin = mappedView.memory;
diff --git a/Engine/Source/Programs/UnrealBuildAccelerator/Common/Private/UbaStorageUtils.cpp b/Engine/Source/Programs/UnrealBuildAccelerator/Common/Private/UbaStorageUtils.cpp
--- a/Engine/Source/Programs/UnrealBuildAccelerator/Common/Private/UbaStorageUtils.cpp
+++ b/Engine/Source/Programs/UnrealBuildAccelerator/Common/Private/UbaStorageUtils.cpp
@@ -652,7 +652,7 @@ namespace uba
 				if (fileMappingMem)
 				{
 					TimerScope ts(m_stats.memoryCopy);
-					MapMemoryCopy(fileMappingMem + destOffset, source, sourceSize);
+					FileMapping_CopyMem(fileMappingMem + destOffset, source, sourceSize);
 					destOffset += sourceSize;
 				}
 				else if (destinationMem)
@@ -1271,21 +1271,21 @@ namespace uba
 		{
 			if (!readMem)
 			{
-				FileMappingHandle fileMapping = uba::CreateFileMappingW(m_logger, readHandle, PAGE_READONLY, fileSize, from);
+				FileMappingHandle fileMapping = FileMapping_CreateFromFile(m_logger, readHandle, PAGE_READONLY, fileSize, from);
 				if (!fileMapping.IsValid())
 					return m_logger.Error(TC("Failed to create file mapping for %s (%s)"), from, LastErrorToText().data);
 
-				auto fmg = MakeGuard([&]() { CloseFileMapping(m_logger, fileMapping, from); });
-				u8* uncompressedData = MapViewOfFile(m_logger, fileMapping, FILE_MAP_READ, 0, fileSize);
+				auto fmg = MakeGuard([&]() { FileMapping_Close(m_logger, fileMapping, from); });
+				u8* uncompressedData = FileMapping_Map(m_logger, fileMapping, FILE_MAP_READ, 0, fileSize);
 				if (!uncompressedData)
 					return m_logger.Error(TC("Failed to map view of file mapping for %s (%s)"), from, LastErrorToText().data);
 
 				auto udg = MakeGuard([&]()
 					{
 						if (m_asyncUnmapViewOfFile)
-							m_workManager.AddWork([uncompressedData, fileSize, f = TString(from), loggerPtr = &m_logger](const WorkContext&) { UnmapViewOfFile(*loggerPtr, uncompressedData, fileSize, f.c_str()); }, 1, TC("UnmapFile"));
+							m_workManager.AddWork([uncompressedData, fileSize, f = TString(from), loggerPtr = &m_logger](const WorkContext&) { FileMapping_Unmap(*loggerPtr, uncompressedData, fileSize, f.c_str()); }, 1, TC("UnmapFile"));
 						else
-							UnmapViewOfFile(m_logger, uncompressedData, fileSize, from);
+							FileMapping_Unmap(m_logger, uncompressedData, fileSize, from);
 					});
 
 				if (!CompressFromMem(destination, workCount, uncompressedData, fileSize, maxUncompressedBlock, totalWritten))
@@ -1599,14 +1599,14 @@ namespace uba
 			u64 compressedSize;
 			if (!uba::GetFileSizeEx(compressedSize, fileHandle))
 				return m_logger.Error(TC("GetFileSize failed for %s (%s)"), fileName, LastErrorToText().data);
-			FileMappingHandle fileMapping = uba::CreateFileMappingW(m_logger, fileHandle, PAGE_READONLY, compressedSize, fileName);
+			FileMappingHandle fileMapping = FileMapping_CreateFromFile(m_logger, fileHandle, PAGE_READONLY, compressedSize, fileName);
 			if (!fileMapping.IsValid())
 				return m_logger.Error(TC("Failed to create file mapping for %s (%s)"), fileName, LastErrorToText().data);
-			auto fmg = MakeGuard([&]() { CloseFileMapping(m_logger, fileMapping, fileName); });
-			u8* fileData = MapViewOfFile(m_logger, fileMapping, FILE_MAP_READ, 0, compressedSize);
+			auto fmg = MakeGuard([&]() { FileMapping_Close(m_logger, fileMapping, fileName); });
+			u8* fileData = FileMapping_Map(m_logger, fileMapping, FILE_MAP_READ, 0, compressedSize);
 			if (!fileData)
 				return m_logger.Error(TC("Failed to map view of file mapping for %s (%s)"), fileName, LastErrorToText().data);
-			auto udg = MakeGuard([&]() { UnmapViewOfFile(m_logger, fileData, compressedSize, fileName); });
+			auto udg = MakeGuard([&]() { FileMapping_Unmap(m_logger, fileData, compressedSize, fileName); });
 			
 			u8* readPos = fileData + 8 + fileStartOffset;
 			if (!DecompressMemoryToMemory(readPos, compressedSize, dest, decompressedSize, fileName, writeHint))
diff --git a/Engine/Source/Programs/UnrealBuildAccelerator/Common/Private/UbaTrace.cpp b/Engine/Source/Programs/UnrealBuildAccelerator/Common/Private/UbaTrace.cpp
--- a/Engine/Source/Programs/UnrealBuildAccelerator/Common/Private/UbaTrace.cpp
+++ b/Engine/Source/Programs/UnrealBuildAccelerator/Common/Private/UbaTrace.cpp
@@ -107,18 +107,18 @@ namespace uba
 		#if !PLATFORM_WINDOWS
 		mappingName = nullptr;
 		#endif
-		m_memoryHandle = uba::CreateMemoryMappingW(m_logger, PAGE_READWRITE|SEC_RESERVE, m_memoryCapacity, mappingName, TC("Trace"));
+		m_memoryHandle = FileMapping_Create(m_logger, PAGE_READWRITE|SEC_RESERVE, m_memoryCapacity, mappingName, TC("Trace"));
 		if (!m_memoryHandle.IsValid())
 			return false;
 		if (GetLastError() != ERROR_ALREADY_EXISTS)
 		{
-			m_memoryBegin = MapViewOfFile(m_logger, m_memoryHandle, FILE_MAP_WRITE, 0, m_memoryCapacity);
+			m_memoryBegin = FileMapping_Map(m_logger, m_memoryHandle, FILE_MAP_WRITE, 0, m_memoryCapacity);
 		}
 
 		if (!m_memoryBegin)
 		{
 			m_logger.Warning(TC("Failed to map view of trace mapping '%s' (%s)"), namedTrace, LastErrorToText().data);
-			CloseFileMapping(m_logger, m_memoryHandle, TC("Trace"));
+			FileMapping_Close(m_logger, m_memoryHandle, TC("Trace"));
 			m_memoryHandle = {};
 			return false;
 		}
@@ -308,12 +308,12 @@ namespace uba
 	{
 		if (m_memoryBegin)
 		{
-			UnmapViewOfFile(m_logger, m_memoryBegin, m_memoryCapacity, TC("Trace"));
+			FileMapping_Unmap(m_logger, m_memoryBegin, m_memoryCapacity, TC("Trace"));
 			m_memoryBegin = nullptr;
 		}
 		if (m_memoryHandle.IsValid())
 		{
-			CloseFileMapping(m_logger, m_memoryHandle, TC("Trace"));
+			FileMapping_Close(m_logger, m_memoryHandle, TC("Trace"));
 			m_memoryHandle = {};
 		}
 	}
@@ -329,7 +329,7 @@ namespace uba
 
 		void* address = m_memoryBegin + m_memoryCommitted;
 		u64 toCommit = committedMemoryNeeded - m_memoryCommitted;
-		if (MapViewCommit(m_logger, address, toCommit))
+		if (FileMapping_Commit(m_logger, address, toCommit))
 		{
 			LockMemory(m_logger, address, toCommit, TC("Trace"));
 			m_memoryCommitted = committedMemoryNeeded;
diff --git a/Engine/Source/Programs/UnrealBuildAccelerator/Common/Private/UbaTraceChannel.cpp b/Engine/Source/Programs/UnrealBuildAccelerator/Common/Private/UbaTraceChannel.cpp
--- a/Engine/Source/Programs/UnrealBuildAccelerator/Common/Private/UbaTraceChannel.cpp
+++ b/Engine/Source/Programs/UnrealBuildAccelerator/Common/Private/UbaTraceChannel.cpp
@@ -37,7 +37,7 @@ namespace uba
 	#else
 		StringBuffer<245> channelMutex;
 		channelMutex.Append(TCV("Uba")).Append(channelName).Append(TCV("Channel"));
-		m_memHandle = uba::CreateMemoryMappingW(m_logger, PAGE_READWRITE, 256, channelName, TC("TraceChannel"));
+		m_memHandle = FileMapping_Create(m_logger, PAGE_READWRITE, 256, channelName, TC("TraceChannel"));
 		if (!m_memHandle.IsValid())
 		{
 			m_logger.Error(TC("Failed to create file mapping %s for trace channel (%s)"), channelName, LastErrorToText().data);
@@ -45,9 +45,9 @@ namespace uba
 		}
 		bool isCreator = GetLastError() != ERROR_ALREADY_EXISTS;
 
-		auto mhg = MakeGuard([&]() { CloseFileMapping(m_logger, m_memHandle, TC("TraceChannel")); m_memHandle = {}; });
+		auto mhg = MakeGuard([&]() { FileMapping_Close(m_logger, m_memHandle, TC("TraceChannel")); m_memHandle = {}; });
 
-		m_mem = MapViewOfFile(m_logger, m_memHandle, FILE_MAP_WRITE, 0, 256);
+		m_mem = FileMapping_Map(m_logger, m_memHandle, FILE_MAP_WRITE, 0, 256);
 		if (!m_mem)
 		{
 			m_logger.Error(TC("Failed to map file mapping for uba trace channel"));
@@ -57,7 +57,7 @@ namespace uba
 		if (isCreator)
 			*(tchar*)m_mem = 0;
 
-		auto mg = MakeGuard([&]() { UnmapViewOfFile(m_logger, m_mem, 256, channelMutex.data); m_mem = nullptr; });
+		auto mg = MakeGuard([&]() { FileMapping_Unmap(m_logger, m_mem, 256, channelMutex.data); m_mem = nullptr; });
 
 		channelMutex.Append(channelName).Append(TCV("Mutex"));
 		m_mutex = CreateMutexW(false, channelMutex.data);
@@ -75,9 +75,9 @@ namespace uba
 		#if UBA_USE_TRACE_CHANNEL_FILE
 		#else
 		if (m_mem)
-			UnmapViewOfFile(m_logger, m_mem, 256, TC("TraceChannel"));
+			FileMapping_Unmap(m_logger, m_mem, 256, TC("TraceChannel"));
 		if (m_memHandle.IsValid())
-			CloseFileMapping(m_logger, m_memHandle, TC("TraceChannel"));
+			FileMapping_Close(m_logger, m_memHandle, TC("TraceChannel"));
 		if (m_mutex != InvalidMutexHandle)
 			CloseMutex(m_mutex);
 		m_mem = nullptr;
diff --git a/Engine/Source/Programs/UnrealBuildAccelerator/Common/Public/UbaStorageProxy.h b/Engine/Source/Programs/UnrealBuildAccelerator/Common/Public/UbaStorageProxy.h
--- a/Engine/Source/Programs/UnrealBuildAccelerator/Common/Public/UbaStorageProxy.h
+++ b/Engine/Source/Programs/UnrealBuildAccelerator/Common/Public/UbaStorageProxy.h
@@ -2,7 +2,7 @@
 
 #pragma once
 
-#include "UbaFileMapping.h"
+#include "UbaHash.h"
 #include "UbaLogger.h"
 #include "UbaNetwork.h"
 
diff --git a/Engine/Source/Programs/UnrealBuildAccelerator/Core/Private/UbaSharedMemoryView.cpp b/Engine/Source/Programs/UnrealBuildAccelerator/Core/Private/UbaSharedMemoryView.cpp
--- a/Engine/Source/Programs/UnrealBuildAccelerator/Core/Private/UbaSharedMemoryView.cpp
+++ b/Engine/Source/Programs/UnrealBuildAccelerator/Core/Private/UbaSharedMemoryView.cpp
@@ -117,7 +117,7 @@ namespace uba
 		UBA_ASSERT(!m_memory);
 		m_capacity = AlignUp(capacity, PageSize);
 		TimerScope ts(KernelStats::GetCurrent().virtualAlloc2);
-		m_memory = (u8*)ReserveMemory(baseAddress, m_capacity);
+		m_memory = (u8*)FileMapping_ReservePlaceholder(baseAddress, m_capacity);
 		return m_memory != nullptr;
 	}
 
@@ -171,7 +171,7 @@ namespace uba
 		auto& kernelStats = KernelStats::GetCurrent();
 
 		TimerScope ts(kernelStats.mapViewOfFile3);
-		void* viewA = MapMemory(handle, (MemoryMapType)mapType, m_memory, m_mappedSize, m_capacity, offset, size);
+		void* viewA = FileMapping_MapPlaceholder(handle, (MemoryMapType)mapType, m_memory, m_mappedSize, m_capacity, offset, size);
 		UBA_ASSERTF(viewA == m_memory + m_mappedSize, TC("MapViewOfFile3 failed extending memory with %llu bytes at offset %llu - %s (%s)"), size, offset, LastErrorToText().data, hint);
 		if (!viewA)
 			return false;
diff --git a/Engine/Source/Programs/UnrealBuildAccelerator/Core/Public/UbaFileMapping.h b/Engine/Source/Programs/UnrealBuildAccelerator/Core/Public/UbaFileMapping.h
--- a/Engine/Source/Programs/UnrealBuildAccelerator/Core/Public/UbaFileMapping.h
+++ b/Engine/Source/Programs/UnrealBuildAccelerator/Core/Public/UbaFileMapping.h
@@ -28,18 +28,18 @@ namespace uba
 
 	enum MemoryMapType { MemoryMapType_ReadOnly, MemoryMapType_ReadWrite, MemoryMapType_CopyOnWrite };
 
-	FileMappingHandle CreateMemoryMappingW(Logger& logger, u32 flProtect, u64 maxSize, const tchar* name, const tchar* hint);
-	FileMappingHandle CreateFileMappingW(Logger& logger, FileHandle file, u32 flProtect, u64 maxSize, const tchar* hint);
-	u8* MapViewOfFile(Logger& logger, FileMappingHandle fileMappingObject, u32 dwDesiredAccess, u64 offset, u64 dwNumberOfBytesToMap, bool allowDiscard = false);
-	bool MapViewCommit(Logger& logger, void* address, u64 size, bool allowDiscard = false);
-	bool UnmapViewOfFile(Logger& logger, const void* lpBaseAddress, u64 bytesToUnmap, const tchar* hint, bool allowDiscard = false);
-	bool CloseFileMapping(Logger& logger, FileMappingHandle h, const tchar* hint);
-	void* ReserveMemory(void* baseAddress, u64 capacity);
-	void* MapMemory(FileMappingHandle handle, MemoryMapType mapType, void* targetAddress, u64 targetOffset, u64 targetCapacity, u64 handleOffset, u64 size);
-	void MapMemoryCopy(void* dest, const void* source, u64 size);
-
-	FileMappingHandle DuplicateMappingFromHost(FileMappingHandle handle, const char* hint);
-	u8* GetMemoryFromMapping(FileMappingHandle handle, u64 size, u64 offset, bool writable, const char* hint);
+	FileMappingHandle FileMapping_Create(Logger& logger, u32 flProtect, u64 maxSize, const tchar* name, const tchar* hint);
+	FileMappingHandle FileMapping_CreateFromFile(Logger& logger, FileHandle file, u32 flProtect, u64 maxSize, const tchar* hint);
+	u8* FileMapping_Map(Logger& logger, FileMappingHandle fileMappingObject, u32 dwDesiredAccess, u64 offset, u64 dwNumberOfBytesToMap, bool allowDiscard = false);
+	bool FileMapping_Commit(Logger& logger, void* address, u64 size, bool allowDiscard = false);
+	bool FileMapping_Unmap(Logger& logger, const void* lpBaseAddress, u64 bytesToUnmap, const tchar* hint, bool allowDiscard = false);
+	bool FileMapping_Close(Logger& logger, FileMappingHandle h, const tchar* hint);
+	void* FileMapping_ReservePlaceholder(void* baseAddress, u64 capacity);
+	void* FileMapping_MapPlaceholder(FileMappingHandle handle, MemoryMapType mapType, void* targetAddress, u64 targetOffset, u64 targetCapacity, u64 handleOffset, u64 size);
+	void FileMapping_CopyMem(void* dest, const void* source, u64 size);
+
+	FileMappingHandle FileMapping_DuplicateFromHost(FileMappingHandle handle, const char* hint);
+	u8* FileMapping_MapFromHost(FileMappingHandle handle, u64 size, u64 offset, bool writable, const char* hint);
 
 	////////////////////////////////////////////////////////////////////////////////////////////////////
 }
diff --git a/Engine/Source/Programs/UnrealBuildAccelerator/Core/Public/UbaFileMapping.inl b/Engine/Source/Programs/UnrealBuildAccelerator/Core/Public/UbaFileMapping.inl
--- a/Engine/Source/Programs/UnrealBuildAccelerator/Core/Public/UbaFileMapping.inl
+++ b/Engine/Source/Programs/UnrealBuildAccelerator/Core/Public/UbaFileMapping.inl
@@ -153,7 +153,7 @@ bool GarbageCollectFileMappings(Logger& logger)
 }
 #endif
 
-FileMappingHandle CreateMemoryMappingW(Logger& logger, u32 flProtect, u64 maxSize, const tchar* name, const tchar* hint)
+FileMappingHandle FileMapping_Create(Logger& logger, u32 flProtect, u64 maxSize, const tchar* name, const tchar* hint)
 {
 #if PLATFORM_WINDOWS
 	ExtendedTimerScope ts(KernelStats::GetCurrent().createFileMapping);
@@ -239,7 +239,7 @@ FileMappingHandle CreateMemoryMappingW(Logger& logger, u32 flProtect, u64 maxSiz
 #endif
 }
 
-FileMappingHandle CreateFileMappingW(Logger& logger, FileHandle file, u32 protect, u64 maxSize, const tchar* hint)
+FileMappingHandle FileMapping_CreateFromFile(Logger& logger, FileHandle file, u32 protect, u64 maxSize, const tchar* hint)
 {
 #if PLATFORM_WINDOWS
 	ExtendedTimerScope ts(KernelStats::GetCurrent().createFileMapping);
@@ -278,7 +278,7 @@ FileMappingHandle CreateFileMappingW(Logger& logger, FileHandle file, u32 protec
 #endif
 }
 
-u8* MapViewOfFile(Logger& logger, FileMappingHandle fileMappingObject, u32 desiredAccess, u64 offset, u64 bytesToMap, bool allowDiscard)
+u8* FileMapping_Map(Logger& logger, FileMappingHandle fileMappingObject, u32 desiredAccess, u64 offset, u64 bytesToMap, bool allowDiscard)
 {
 #if PLATFORM_WINDOWS
 	ExtendedTimerScope ts(KernelStats::GetCurrent().mapViewOfFile);
@@ -319,7 +319,7 @@ u8* MapViewOfFile(Logger& logger, FileMappingHandle fileMappingObject, u32 desir
 	return res;
 }
 
-bool MapViewCommit(Logger& logger, void* address, u64 size, bool allowDiscard)
+bool FileMapping_Commit(Logger& logger, void* address, u64 size, bool allowDiscard)
 {
 #if PLATFORM_WINDOWS
 	ExtendedTimerScope ts(KernelStats::GetCurrent().virtualAlloc);
@@ -333,7 +333,7 @@ bool MapViewCommit(Logger& logger, void* address, u64 size, bool allowDiscard)
 #endif
 }
 
-bool UnmapViewOfFile(Logger& logger, const void* address, u64 bytesToUnmap, const tchar* hint, bool allowDiscard)
+bool FileMapping_Unmap(Logger& logger, const void* address, u64 bytesToUnmap, const tchar* hint, bool allowDiscard)
 {
 	if (!address)
 		return true;
@@ -364,7 +364,7 @@ bool UnmapViewOfFile(Logger& logger, const void* address, u64 bytesToUnmap, cons
 #endif
 }
 
-bool CloseFileMapping(Logger& logger, FileMappingHandle h, const tchar* hint)
+bool FileMapping_Close(Logger& logger, FileMappingHandle h, const tchar* hint)
 {
 	if (!h.IsValid())
 		return true;
@@ -406,7 +406,7 @@ bool CloseFileMapping(Logger& logger, FileMappingHandle h, const tchar* hint)
 #endif
 }
 
-void* ReserveMemory(void* baseAddress, u64 capacity)
+void* FileMapping_ReservePlaceholder(void* baseAddress, u64 capacity)
 {
 #if PLATFORM_WINDOWS
 	void* memory = VirtualAlloc2(GetCurrentProcess(), baseAddress, capacity, MEM_RESERVE | MEM_RESERVE_PLACEHOLDER, PAGE_NOACCESS, nullptr, 0);
@@ -419,7 +419,7 @@ void* ReserveMemory(void* baseAddress, u64 capacity)
 	return memory;
 }
 
-void* MapMemory(FileMappingHandle handle, MemoryMapType mapType, void* targetAddress, u64 targetOffset, u64 targetCapacity, u64 handleOffset, u64 size)
+void* FileMapping_MapPlaceholder(FileMappingHandle handle, MemoryMapType mapType, void* targetAddress, u64 targetOffset, u64 targetCapacity, u64 handleOffset, u64 size)
 {
 #if PLATFORM_WINDOWS
 	u32 protect = PAGE_READWRITE;
@@ -455,7 +455,7 @@ void* MapMemory(FileMappingHandle handle, MemoryMapType mapType, void* targetAdd
 	return viewA;
 }
 
-FileMappingHandle DuplicateMappingFromHost(FileMappingHandle handle, const char* hint)
+FileMappingHandle FileMapping_DuplicateFromHost(FileMappingHandle handle, const char* hint)
 {
 #if PLATFORM_WINDOWS
 	HANDLE realHandle;
@@ -470,7 +470,7 @@ FileMappingHandle DuplicateMappingFromHost(FileMappingHandle handle, const char*
 #endif
 }
 
-u8* GetMemoryFromMapping(FileMappingHandle handle, u64 size, u64 offset, bool writable, const char* hint)
+u8* FileMapping_MapFromHost(FileMappingHandle handle, u64 size, u64 offset, bool writable, const char* hint)
 {
 #if PLATFORM_WINDOWS
 	DWORD dwDesiredAccess = FILE_MAP_READ;
diff --git a/Engine/Source/Programs/UnrealBuildAccelerator/Core/Public/UbaFileMappingWin.inl b/Engine/Source/Programs/UnrealBuildAccelerator/Core/Public/UbaFileMappingWin.inl
--- a/Engine/Source/Programs/UnrealBuildAccelerator/Core/Public/UbaFileMappingWin.inl
+++ b/Engine/Source/Programs/UnrealBuildAccelerator/Core/Public/UbaFileMappingWin.inl
@@ -13,54 +13,54 @@ namespace uba
 		#include "UbaFileMapping.inl"
 	}
 
-	FileMappingHandle CreateMemoryMappingW(Logger& logger, u32 flProtect, u64 maxSize, const tchar* name, const tchar* hint)
+	FileMappingHandle FileMapping_Create(Logger& logger, u32 flProtect, u64 maxSize, const tchar* name, const tchar* hint)
 	{
-		return internal::CreateMemoryMappingW(logger, flProtect, maxSize, name, hint);
+		return internal::FileMapping_Create(logger, flProtect, maxSize, name, hint);
 	}
 
-	FileMappingHandle CreateFileMappingW(Logger& logger, FileHandle file, u32 flProtect, u64 maxSize, const tchar* hint)
+	FileMappingHandle FileMapping_CreateFromFile(Logger& logger, FileHandle file, u32 flProtect, u64 maxSize, const tchar* hint)
 	{
-		return internal::CreateFileMappingW(logger, file, flProtect, maxSize, hint);
+		return internal::FileMapping_CreateFromFile(logger, file, flProtect, maxSize, hint);
 	}
 
-	u8* MapViewOfFile(Logger& logger, FileMappingHandle fileMappingObject, u32 dwDesiredAccess, u64 offset, u64 dwNumberOfBytesToMap, bool allowDiscard)
+	u8* FileMapping_Map(Logger& logger, FileMappingHandle fileMappingObject, u32 dwDesiredAccess, u64 offset, u64 dwNumberOfBytesToMap, bool allowDiscard)
 	{
-		return internal::MapViewOfFile(logger, fileMappingObject, dwDesiredAccess, offset, dwNumberOfBytesToMap, allowDiscard);
+		return internal::FileMapping_Map(logger, fileMappingObject, dwDesiredAccess, offset, dwNumberOfBytesToMap, allowDiscard);
 	}
 
-	bool MapViewCommit(Logger& logger, void* address, u64 size, bool allowDiscard)
+	bool FileMapping_Commit(Logger& logger, void* address, u64 size, bool allowDiscard)
 	{
-		return internal::MapViewCommit(logger, address, size, allowDiscard);
+		return internal::FileMapping_Commit(logger, address, size, allowDiscard);
 	}
 
-	bool UnmapViewOfFile(Logger& logger, const void* lpBaseAddress, u64 bytesToUnmap, const tchar* hint, bool allowDiscard)
+	bool FileMapping_Unmap(Logger& logger, const void* lpBaseAddress, u64 bytesToUnmap, const tchar* hint, bool allowDiscard)
 	{
-		return internal::UnmapViewOfFile(logger, lpBaseAddress, bytesToUnmap, hint, allowDiscard);
+		return internal::FileMapping_Unmap(logger, lpBaseAddress, bytesToUnmap, hint, allowDiscard);
 	}
 
-	bool CloseFileMapping(Logger& logger, FileMappingHandle h, const tchar* hint)
+	bool FileMapping_Close(Logger& logger, FileMappingHandle h, const tchar* hint)
 	{
-		return internal::CloseFileMapping(logger, h, hint);
+		return internal::FileMapping_Close(logger, h, hint);
 	}
 
-	void* ReserveMemory(void* baseAddress, u64 capacity)
+	void* FileMapping_ReservePlaceholder(void* baseAddress, u64 capacity)
 	{
-		return internal::ReserveMemory(baseAddress, capacity);
+		return internal::FileMapping_ReservePlaceholder(baseAddress, capacity);
 	}
 
-	void* MapMemory(FileMappingHandle handle, MemoryMapType mapType, void* targetAddress, u64 targetOffset, u64 targetCapacity, u64 handleOffset, u64 size)
+	void* FileMapping_MapPlaceholder(FileMappingHandle handle, MemoryMapType mapType, void* targetAddress, u64 targetOffset, u64 targetCapacity, u64 handleOffset, u64 size)
 	{
-		return internal::MapMemory(handle, mapType, targetAddress, targetOffset, targetCapacity, handleOffset, size);
+		return internal::FileMapping_MapPlaceholder(handle, mapType, targetAddress, targetOffset, targetCapacity, handleOffset, size);
 	}
 
-	FileMappingHandle DuplicateMappingFromHost(FileMappingHandle handle, const char* hint)
+	FileMappingHandle FileMapping_DuplicateFromHost(FileMappingHandle handle, const char* hint)
 	{
-		return internal::DuplicateMappingFromHost(handle, hint);
+		return internal::FileMapping_DuplicateFromHost(handle, hint);
 	}
 
-	u8* GetMemoryFromMapping(FileMappingHandle handle, u64 size, u64 offset, bool writable, const char* hint)
+	u8* FileMapping_MapFromHost(FileMappingHandle handle, u64 size, u64 offset, bool writable, const char* hint)
 	{
-		return internal::GetMemoryFromMapping(handle, size, offset, writable, hint);
+		return internal::FileMapping_MapFromHost(handle, size, offset, writable, hint);
 	}
 }
 
diff --git a/Engine/Source/Programs/UnrealBuildAccelerator/Detours/Private/Windows/UbaDetoursFunctionsWin.cpp b/Engine/Source/Programs/UnrealBuildAccelerator/Detours/Private/Windows/UbaDetoursFunctionsWin.cpp
--- a/Engine/Source/Programs/UnrealBuildAccelerator/Detours/Private/Windows/UbaDetoursFunctionsWin.cpp
+++ b/Engine/Source/Programs/UnrealBuildAccelerator/Detours/Private/Windows/UbaDetoursFunctionsWin.cpp
@@ -1151,16 +1151,16 @@ void Init(const DetoursPayload& payload, u64 startTime)
 	}
 	#endif
 
-	g_sharedMemoryAllocator = DuplicateMappingFromHost(sharedMemoryAllocatorHandle, "SharedMem");
+	g_sharedMemoryAllocator = FileMapping_DuplicateFromHost(sharedMemoryAllocatorHandle, "SharedMem");
 
 	TimerScope ts1(g_stats.fileTable);
-	u8* mappedFileTableMem = GetMemoryFromMapping(DuplicateMappingFromHost(mappedFileTableHandle, "FileTable"), FileMappingTableMaxSize, 0, false, "FileTable");
+	u8* mappedFileTableMem = FileMapping_MapFromHost(FileMapping_DuplicateFromHost(mappedFileTableHandle, "FileTable"), FileMappingTableMaxSize, 0, false, "FileTable");
 	g_mappedFileTable.Init(mappedFileTableMem, mappedFileTableCount, mappedFileTableSize);
 	ts1.Leave();
 
 	TimerScope ts2(g_stats.dirTable);
-	u8* directoryTableMem = GetMemoryFromMapping(DuplicateMappingFromHost(directoryTableHandle, "DirTable"), DirTableMaxSize, 0, false, "DirTable");
-	u8* overlayTableMem = GetMemoryFromMapping(DuplicateMappingFromHost(overlayFilesHandle, "OverlayTable"), OverlayTableMaxSize, 0, false, "OverlayTable");
+	u8* directoryTableMem = FileMapping_MapFromHost(FileMapping_DuplicateFromHost(directoryTableHandle, "DirTable"), DirTableMaxSize, 0, false, "DirTable");
+	u8* overlayTableMem = FileMapping_MapFromHost(FileMapping_DuplicateFromHost(overlayFilesHandle, "OverlayTable"), OverlayTableMaxSize, 0, false, "OverlayTable");
 	g_directoryTable.Init(directoryTableMem, directoryTableCount, directoryTableSize.main);
 	g_directoryTable.InitOverlay(overlayTableMem, directoryTableSize.overlay);
 	ts2.Leave();
diff --git a/Engine/Source/Programs/UnrealBuildAccelerator/Detours/Private/Windows/UbaDetoursMainWin.cpp b/Engine/Source/Programs/UnrealBuildAccelerator/Detours/Private/Windows/UbaDetoursMainWin.cpp
--- a/Engine/Source/Programs/UnrealBuildAccelerator/Detours/Private/Windows/UbaDetoursMainWin.cpp
+++ b/Engine/Source/Programs/UnrealBuildAccelerator/Detours/Private/Windows/UbaDetoursMainWin.cpp
@@ -57,8 +57,8 @@ namespace uba
 		if (!DuplicateHandle(g_hostProcess, payload.writeEvent, GetCurrentProcess(), &g_readEvent, SYNCHRONIZE, false, 0))
 			return -3;
 
-		FileMappingHandle communicationHandle = DuplicateMappingFromHost(FileMappingHandle::FromU64(payload.communicationHandle), "Communication");
-		u8* mem = GetMemoryFromMapping(communicationHandle, CommunicationMemSize, payload.communicationOffset, true, "Communication");
+		FileMappingHandle communicationHandle = FileMapping_DuplicateFromHost(FileMappingHandle::FromU64(payload.communicationHandle), "Communication");
+		u8* mem = FileMapping_MapFromHost(communicationHandle, CommunicationMemSize, payload.communicationOffset, true, "Communication");
 		if (!mem)
 		{
 			return -5;
diff --git a/Engine/Source/Programs/UnrealBuildAccelerator/Visualizer/Private/UbaTraceReader.cpp b/Engine/Source/Programs/UnrealBuildAccelerator/Visualizer/Private/UbaTraceReader.cpp
--- a/Engine/Source/Programs/UnrealBuildAccelerator/Visualizer/Private/UbaTraceReader.cpp
+++ b/Engine/Source/Programs/UnrealBuildAccelerator/Visualizer/Private/UbaTraceReader.cpp
@@ -111,15 +111,15 @@ namespace uba
 		if (!uba::GetFileSizeEx(fileSize, readHandle))
 			return false;
 
-		m_memoryHandle = uba::CreateFileMappingW(m_logger, readHandle, PAGE_READONLY, 0, fileName);
+		m_memoryHandle = FileMapping_CreateFromFile(m_logger, readHandle, PAGE_READONLY, 0, fileName);
 		if (!m_memoryHandle.IsValid())
 			return false;
-		auto closeFileMapping = MakeGuard([&]() { CloseFileMapping(m_logger, m_memoryHandle, fileName); m_memoryHandle = {}; });
+		auto closeFileMapping = MakeGuard([&]() { FileMapping_Close(m_logger, m_memoryHandle, fileName); m_memoryHandle = {}; });
 
-		m_memoryBegin = MapViewOfFile(m_logger, m_memoryHandle, FILE_MAP_READ, 0, fileSize);
+		m_memoryBegin = FileMapping_Map(m_logger, m_memoryHandle, FILE_MAP_READ, 0, fileSize);
 		if (!m_memoryBegin)
 			return false;
-		auto closeView = MakeGuard([&]() { UnmapViewOfFile(m_logger, m_memoryBegin, fileSize, TC("TraceReader")); m_memoryBegin = nullptr; m_memoryPos = nullptr; m_memoryEnd = nullptr; });
+		auto closeView = MakeGuard([&]() { FileMapping_Unmap(m_logger, m_memoryBegin, fileSize, TC("TraceReader")); m_memoryBegin = nullptr; m_memoryPos = nullptr; m_memoryEnd = nullptr; });
 
 		m_memoryPos = m_memoryBegin;
 		m_memoryEnd = m_memoryBegin + fileSize;
@@ -196,10 +196,10 @@ namespace uba
 		Reset(out);
 
 		u32 traceMemSize = 128 * 1024 * 1024;
-		m_memoryHandle = CreateMemoryMappingW(m_logger, PAGE_READWRITE, traceMemSize, 0, TC("NetworkClient"));
+		m_memoryHandle = FileMapping_Create(m_logger, PAGE_READWRITE, traceMemSize, 0, TC("NetworkClient"));
 		if (!m_memoryHandle.IsValid())
 			return false;
-		m_memoryBegin = MapViewOfFile(m_logger, m_memoryHandle, FILE_MAP_ALL_ACCESS, 0, traceMemSize);
+		m_memoryBegin = FileMapping_Map(m_logger, m_memoryHandle, FILE_MAP_ALL_ACCESS, 0, traceMemSize);
 		m_memoryPos = m_memoryBegin;
 		m_memoryEnd = m_memoryBegin;
 
@@ -296,7 +296,7 @@ namespace uba
 					m_logger.Error(TC("OpenFileMappingW - Failed to open file mapping %s (%s)"), namedTrace, LastErrorToText().data);
 				return false;
 			}
-			m_memoryBegin = MapViewOfFile(m_logger, m_memoryHandle, FILE_MAP_READ, 0, traceCapacity);
+			m_memoryBegin = FileMapping_Map(m_logger, m_memoryHandle, FILE_MAP_READ, 0, traceCapacity);
 			if (!m_memoryBegin)
 			{
 				if (!silentFail)
@@ -340,9 +340,9 @@ namespace uba
 		// Move memory to local mapping in case we want to replay
 		u64 traceMemSize = m_memoryEnd - m_memoryBegin;
 		u8* memoryBegin = nullptr;
-		FileMappingHandle memoryHandle = CreateMemoryMappingW(m_logger, PAGE_READWRITE, traceMemSize, nullptr, TC("UpdateReadNamed"));
+		FileMappingHandle memoryHandle = FileMapping_Create(m_logger, PAGE_READWRITE, traceMemSize, nullptr, TC("UpdateReadNamed"));
 		if (memoryHandle.IsValid())
-			if ((memoryBegin = MapViewOfFile(m_logger, memoryHandle, FILE_MAP_ALL_ACCESS, 0, traceMemSize)) != nullptr)
+			if ((memoryBegin = FileMapping_Map(m_logger, memoryHandle, FILE_MAP_ALL_ACCESS, 0, traceMemSize)) != nullptr)
 				memcpy(memoryBegin, m_memoryBegin, traceMemSize);
 
 		u64 pos = m_memoryPos - m_memoryBegin;
@@ -1582,10 +1582,10 @@ namespace uba
 		m_hostProcess = 0;
 
 		if (m_memoryBegin)
-			UnmapViewOfFile(m_logger, m_memoryBegin, unmapSize, TC("TraceReader"));
+			FileMapping_Unmap(m_logger, m_memoryBegin, unmapSize, TC("TraceReader"));
 		m_memoryBegin = nullptr;
 		if (m_memoryHandle.IsValid())
-			CloseFileMapping(m_logger, m_memoryHandle, TC("TraceReader"));
+			FileMapping_Close(m_logger, m_memoryHandle, TC("TraceReader"));
 		m_memoryHandle = {};
 		m_namedTrace.clear();
 	}
