[UBA]
* Moved last piece of file mapping code that talked directly to platform api into UbaFileMapping.inl. Now it will be possible to implement code to call out directly to linux on wine

diff --git a/Engine/Source/Programs/UnrealBuildAccelerator/Common/Private/UbaSharedMemoryAllocator.cpp b/Engine/Source/Programs/UnrealBuildAccelerator/Common/Private/UbaSharedMemoryAllocator.cpp
--- a/Engine/Source/Programs/UnrealBuildAccelerator/Common/Private/UbaSharedMemoryAllocator.cpp
+++ b/Engine/Source/Programs/UnrealBuildAccelerator/Common/Private/UbaSharedMemoryAllocator.cpp
@@ -525,38 +525,17 @@ namespace uba
 			return;
 		}
 
+		Vector<u64> subMappings;
+		subMappings.reserve(h.slices.size() - 1);
+		u64 committed = 0;
+		for (Slice slice : h.slices)
 		{
-			auto& kernelStats = KernelStats::GetCurrent();
-
-			#if PLATFORM_WINDOWS
-			u64 pos = 0;
-			u32 res; (void)res;
-
-			for (Slice slice : h.slices)
-			{
-				TimerScope ts(kernelStats.unmapViewOfFile2);
-				res = UnmapViewOfFile2(GetCurrentProcess(), memory + pos, MEM_PRESERVE_PLACEHOLDER);
-				UBA_ASSERTF(res, TC("UnmapViewOfFile2 failed on pos %llu for %s (%s)"), pos, h.GetHint(), LastErrorToText().data);
-				pos += slice.size;
-			}
-			if (h.slices.size() > 1)
-			{
-				TimerScope ts(kernelStats.virtualFree);
-				res = VirtualFree(memory, pos, MEM_RELEASE | MEM_COALESCE_PLACEHOLDERS);
-				UBA_ASSERTF(res, TC("VirtualFree failed (%s)"), LastErrorToText().data);
-			}
-
-			TimerScope ts(kernelStats.virtualFree);
-			res = VirtualFree(memory, 0, MEM_RELEASE);
-			UBA_ASSERTF(res, TC("VirtualFree failed (%s)"), LastErrorToText().data);
-			#else
-			u64 committed = 0;
-			for (Slice slice : h.slices)
-				committed += slice.size;
-			int res = munmap(memory, committed);
-			UBA_ASSERT(res == 0);(void)res;
-			#endif
+			if (committed != 0)
+				subMappings.push_back(committed);
+			committed += slice.size;
 		}
+		FileMapping_UnmapPlaceholder(memory, committed, committed, subMappings.data(), subMappings.size());
+
 		ReleaseHandle(handle);
 	}
 
diff --git a/Engine/Source/Programs/UnrealBuildAccelerator/Core/Private/UbaSharedMemoryView.cpp b/Engine/Source/Programs/UnrealBuildAccelerator/Core/Private/UbaSharedMemoryView.cpp
--- a/Engine/Source/Programs/UnrealBuildAccelerator/Core/Private/UbaSharedMemoryView.cpp
+++ b/Engine/Source/Programs/UnrealBuildAccelerator/Core/Private/UbaSharedMemoryView.cpp
@@ -57,42 +57,9 @@ namespace uba
 			return;
 		}
 
-		#if PLATFORM_WINDOWS
-		bool res; (void)res;
-		if (m_mappedSize)
-		{
-			auto& kernelStats = KernelStats::GetCurrent();
-			{
-				TimerScope ts(kernelStats.unmapViewOfFile2);
-				res = UnmapViewOfFile2(GetCurrentProcess(), m_memory, MEM_PRESERVE_PLACEHOLDER);
-				UBA_ASSERT(res);
-			}
-
-			for (u64 mapping : m_additionalMappings)
-			{
-				TimerScope ts(kernelStats.unmapViewOfFile2);
-				res = UnmapViewOfFile2(GetCurrentProcess(), m_memory + mapping, MEM_PRESERVE_PLACEHOLDER);
-				UBA_ASSERT(res);
-			}
-
-			if (m_mappedSize < m_capacity || !m_additionalMappings.empty())
-			{
-				TimerScope ts(kernelStats.virtualFree);
-				res = VirtualFree(m_memory, m_capacity, MEM_RELEASE | MEM_COALESCE_PLACEHOLDERS);
-				UBA_ASSERTF(res, TC("VirtualFree failed (%u)"), GetLastError());
-			}
-
-			m_additionalMappings.clear();
-
-			m_mappedSize = 0;
-		}
-		res = VirtualFree(m_memory, 0, MEM_RELEASE);
-		UBA_ASSERTF(res, TC("VirtualFree failed (%u)"), GetLastError());
-		#else
-		int res = munmap(m_memory, m_capacity);
-		UBA_ASSERT(res == 0); (void)res;
-		#endif
-
+		FileMapping_UnmapPlaceholder(m_memory, m_capacity, m_mappedSize, m_additionalMappings.data(), m_additionalMappings.size());
+		m_additionalMappings.clear();
+		m_mappedSize = 0;
 		m_memory = nullptr;
 		m_capacity = 0;
 	}
diff --git a/Engine/Source/Programs/UnrealBuildAccelerator/Core/Public/UbaFileMapping.h b/Engine/Source/Programs/UnrealBuildAccelerator/Core/Public/UbaFileMapping.h
--- a/Engine/Source/Programs/UnrealBuildAccelerator/Core/Public/UbaFileMapping.h
+++ b/Engine/Source/Programs/UnrealBuildAccelerator/Core/Public/UbaFileMapping.h
@@ -36,6 +36,7 @@ namespace uba
 	bool FileMapping_Close(Logger& logger, FileMappingHandle h, const tchar* hint);
 	void* FileMapping_ReservePlaceholder(void* baseAddress, u64 capacity);
 	void* FileMapping_MapPlaceholder(FileMappingHandle handle, MemoryMapType mapType, void* targetAddress, u64 targetOffset, u64 targetCapacity, u64 handleOffset, u64 size);
+	bool FileMapping_UnmapPlaceholder(void* memory, u64 capacity, u64 mappedSize, u64* subMappings, u64 subMappingsCount);
 	void FileMapping_CopyMem(void* dest, const void* source, u64 size);
 
 	FileMappingHandle FileMapping_DuplicateFromHost(FileMappingHandle handle, const char* hint);
diff --git a/Engine/Source/Programs/UnrealBuildAccelerator/Core/Public/UbaFileMapping.inl b/Engine/Source/Programs/UnrealBuildAccelerator/Core/Public/UbaFileMapping.inl
--- a/Engine/Source/Programs/UnrealBuildAccelerator/Core/Public/UbaFileMapping.inl
+++ b/Engine/Source/Programs/UnrealBuildAccelerator/Core/Public/UbaFileMapping.inl
@@ -455,6 +455,43 @@ void* FileMapping_MapPlaceholder(FileMappingHandle handle, MemoryMapType mapType
 	return viewA;
 }
 
+bool FileMapping_UnmapPlaceholder(void* memory, u64 capacity, u64 mappedSize, u64* subMappings, u64 subMappingsCount)
+{
+#if PLATFORM_WINDOWS
+	bool res; (void)res;
+	if (mappedSize)
+	{
+		auto& kernelStats = KernelStats::GetCurrent();
+		{
+			TimerScope ts(kernelStats.unmapViewOfFile2);
+			res = UnmapViewOfFile2(GetCurrentProcess(), memory, MEM_PRESERVE_PLACEHOLDER);
+			UBA_ASSERT(res);
+		}
+
+		for (u64 i=0;i!=subMappingsCount; ++i)
+		{
+			TimerScope ts(kernelStats.unmapViewOfFile2);
+			res = UnmapViewOfFile2(GetCurrentProcess(), (u8*)memory + subMappings[i], MEM_PRESERVE_PLACEHOLDER);
+			UBA_ASSERT(res);
+		}
+
+		if (mappedSize < capacity || subMappingsCount)
+		{
+			TimerScope ts(kernelStats.virtualFree);
+			res = VirtualFree(memory, capacity, MEM_RELEASE | MEM_COALESCE_PLACEHOLDERS);
+			UBA_ASSERTF(res, TC("VirtualFree failed (%u)"), GetLastError());
+		}
+		mappedSize = 0;
+	}
+	res = VirtualFree(memory, 0, MEM_RELEASE);
+	UBA_ASSERTF(res, TC("VirtualFree failed (%u)"), GetLastError());
+#else
+	int res = munmap(memory, capacity);
+	UBA_ASSERT(res == 0); (void)res;
+#endif
+	return true;
+}
+
 FileMappingHandle FileMapping_DuplicateFromHost(FileMappingHandle handle, const char* hint)
 {
 #if PLATFORM_WINDOWS
diff --git a/Engine/Source/Programs/UnrealBuildAccelerator/Core/Public/UbaFileMappingWin.inl b/Engine/Source/Programs/UnrealBuildAccelerator/Core/Public/UbaFileMappingWin.inl
--- a/Engine/Source/Programs/UnrealBuildAccelerator/Core/Public/UbaFileMappingWin.inl
+++ b/Engine/Source/Programs/UnrealBuildAccelerator/Core/Public/UbaFileMappingWin.inl
@@ -53,6 +53,11 @@ namespace uba
 		return internal::FileMapping_MapPlaceholder(handle, mapType, targetAddress, targetOffset, targetCapacity, handleOffset, size);
 	}
 
+	bool FileMapping_UnmapPlaceholder(void* memory, u64 capacity, u64 mappedSize, u64* subMappings, u64 subMappingsCount)
+	{
+		return internal::FileMapping_UnmapPlaceholder(memory, capacity, mappedSize, subMappings, subMappingsCount);
+	}
+
 	FileMappingHandle FileMapping_DuplicateFromHost(FileMappingHandle handle, const char* hint)
 	{
 		return internal::FileMapping_DuplicateFromHost(handle, hint);
diff --git a/Engine/Source/Programs/UnrealBuildAccelerator/Core/Public/UbaSharedMemoryView.h b/Engine/Source/Programs/UnrealBuildAccelerator/Core/Public/UbaSharedMemoryView.h
--- a/Engine/Source/Programs/UnrealBuildAccelerator/Core/Public/UbaSharedMemoryView.h
+++ b/Engine/Source/Programs/UnrealBuildAccelerator/Core/Public/UbaSharedMemoryView.h
@@ -55,7 +55,7 @@ namespace uba
 		u64 m_mappedSize = 0;
 		u64 m_capacity = 0;
 
-		List<u64> m_additionalMappings;
+		Vector<u64> m_additionalMappings;
 
 		SharedMemoryView(const SharedMemoryView&) = delete;
 		SharedMemoryView& operator=(const SharedMemoryView&) = delete;
diff --git a/Engine/Source/Programs/UnrealBuildAccelerator/Detours/Private/PosixOS/UbaDetoursFunctionsPosix.cpp b/Engine/Source/Programs/UnrealBuildAccelerator/Detours/Private/PosixOS/UbaDetoursFunctionsPosix.cpp
--- a/Engine/Source/Programs/UnrealBuildAccelerator/Detours/Private/PosixOS/UbaDetoursFunctionsPosix.cpp
+++ b/Engine/Source/Programs/UnrealBuildAccelerator/Detours/Private/PosixOS/UbaDetoursFunctionsPosix.cpp
@@ -2950,13 +2950,13 @@ namespace uba
 		LogHeader();
 
 		TimerScope ts1(g_stats.fileTable);
-		u8* mappedFileTableMem = GetMemoryFromMapping(mappedFileTableHandle, FileMappingTableMaxSize,0, false, "FileTable");
+		u8* mappedFileTableMem = FileMapping_MapFromHost(mappedFileTableHandle, FileMappingTableMaxSize,0, false, "FileTable");
 		g_mappedFileTable.Init(mappedFileTableMem, mappedFileTableCount, mappedFileTableSize);
 		ts1.Leave();
 
 		TimerScope ts2(g_stats.dirTable);
-		u8* dirTableMem = GetMemoryFromMapping(directoryTableHandle, DirTableMaxSize,0, false, "DirTable");
-		u8* overlayTableMem = GetMemoryFromMapping(overlayTableHandle, OverlayTableMaxSize,0, false, "OverlayTable");
+		u8* dirTableMem = FileMapping_MapFromHost(directoryTableHandle, DirTableMaxSize,0, false, "DirTable");
+		u8* overlayTableMem = FileMapping_MapFromHost(overlayTableHandle, OverlayTableMaxSize,0, false, "OverlayTable");
 		g_directoryTable.Init(dirTableMem, directoryTableCount, directoryTableSize.main);
 		g_directoryTable.InitOverlay(overlayTableMem , directoryTableSize.overlay);
 		ts2.Leave();
