[UBA]
* Implemented so file mappings are linux native when running through wine and UbaWine.dll.so is present. This works around the bug in wine where MapViewOfFile3 take up commit space causing wine machines with less memory to fail

diff --git a/Engine/Source/Programs/UnrealBuildAccelerator/Common/Private/UbaEnvironment.cpp b/Engine/Source/Programs/UnrealBuildAccelerator/Common/Private/UbaEnvironment.cpp
--- a/Engine/Source/Programs/UnrealBuildAccelerator/Common/Private/UbaEnvironment.cpp
+++ b/Engine/Source/Programs/UnrealBuildAccelerator/Common/Private/UbaEnvironment.cpp
@@ -1,6 +1,7 @@
 // Copyright Epic Games, Inc. All Rights Reserved.
 
 #include "UbaEnvironment.h"
+#include "UbaFile.h"
 #include "UbaPlatform.h"
 #include "UbaStringBuffer.h"
 
@@ -597,4 +598,32 @@ namespace uba
 		}
 		#endif
 	}
+
+	#if PLATFORM_WINDOWS
+	StringView GetUbaWineModulePath()
+	{
+		static StringView sv = []()
+			{
+				static StringBuffer<> dllPath;
+				LoggerWithWriter logger(g_consoleLogWriter);
+				if (GetDirectoryOfCurrentModule(logger, dllPath))
+					dllPath.EnsureEndsWithSlash();
+				dllPath.Append(TCV("UbaWine.dll.so"));
+				return StringView(dllPath);
+			}();
+		return sv;
+	}
+	HMODULE GetUbaWineModule()
+	{
+		static HMODULE ubaWineModule = []()
+			{
+				StringView sv = GetUbaWineModulePath();
+				HMODULE res = LoadLibraryW(sv.data);
+				if (!res)
+					LoggerWithWriter(g_consoleLogWriter).Info(TC("Failed to load UbaWine module from %s"), sv.data);
+				return res;
+			}();
+		return ubaWineModule;
+	}
+	#endif
 }
diff --git a/Engine/Source/Programs/UnrealBuildAccelerator/Common/Private/UbaNetworkBackendTcp.cpp b/Engine/Source/Programs/UnrealBuildAccelerator/Common/Private/UbaNetworkBackendTcp.cpp
--- a/Engine/Source/Programs/UnrealBuildAccelerator/Common/Private/UbaNetworkBackendTcp.cpp
+++ b/Engine/Source/Programs/UnrealBuildAccelerator/Common/Private/UbaNetworkBackendTcp.cpp
@@ -225,12 +225,6 @@ namespace uba
 
 	////////////////////////////////////////////////////////////////////////////////////////////////////
 
-	#if PLATFORM_WINDOWS
-	HMODULE g_wineDll;
-	#endif
-
-	////////////////////////////////////////////////////////////////////////////////////////////////////
-
 	bool InitNetwork(Logger& logger)
 	{
 #if PLATFORM_WINDOWS
@@ -240,15 +234,9 @@ namespace uba
 
 		if (IsRunningWine())
 		{
-			StringBuffer<> dllPath;
-			if (GetDirectoryOfCurrentModule(logger, dllPath))
-				dllPath.EnsureEndsWithSlash();
-			dllPath.Append(TCV("UbaWine.dll.so"));
-
-			g_wineDll = LoadLibrary(dllPath.data);
-			if (g_wineDll)
-				SocketInitWine(logger, g_wineDll);
-			else// if (FileExists(logger, TC("UbaWine.dll.so")))
+			if (HMODULE wineDll = GetUbaWineModule())
+				SocketInitWine(logger, wineDll);
+			else
 				logger.Warning(TC("Failed to load UbaWine.dll.so (%s)"), LastErrorToText().data);
 
 		}
@@ -1702,7 +1690,7 @@ namespace uba
 
 		#if PLATFORM_WINDOWS
 		if (!GetTcpAutoTuning && IsRunningWine())
-			GetTcpAutoTuning = (decltype(GetTcpAutoTuning))GetProcAddress(g_wineDll, "GetTcpAutoTuning");
+			GetTcpAutoTuning = (decltype(GetTcpAutoTuning))GetProcAddress(GetUbaWineModule(), "GetTcpAutoTuning");
 		#else
 		GetTcpAutoTuning = UnixGetTcpAutoTuning;
 		#endif
@@ -2209,7 +2197,7 @@ namespace uba
 		#if 0
 		if (IsRunningWine())
 		{
-			static HMODULE wineDll = LoadLibrary(TC("UbaWine.dll.so"));
+			static HMODULE wineDll = GetUbaWineModule();
 			if (!wineDll)
 				return false;
 			using GetLinuxTcpInfoFunc = int WINAPI(Socket, void*, int*);
diff --git a/Engine/Source/Programs/UnrealBuildAccelerator/Common/Private/UbaProcess.cpp b/Engine/Source/Programs/UnrealBuildAccelerator/Common/Private/UbaProcess.cpp
--- a/Engine/Source/Programs/UnrealBuildAccelerator/Common/Private/UbaProcess.cpp
+++ b/Engine/Source/Programs/UnrealBuildAccelerator/Common/Private/UbaProcess.cpp
@@ -1392,14 +1392,26 @@ namespace uba
 
 		auto& process = *(ProcessImpl*)h.m_process;
 
-		const char* detoursLib = m_session.m_detoursLibrary[IsLinux][IsArmBinary].c_str();
-		u32 detoursLibLen = u32(m_session.m_detoursLibrary[IsLinux][IsArmBinary].size());
+		auto& detoursLib = m_session.m_detoursLibrary[IsLinux][IsArmBinary];
+		u32 detoursLibLen = u32(detoursLib.size());
 		UBA_ASSERT(detoursLibLen < 1024);
 
 		writer.WriteU32(childProcessId);
 		writer.WriteU32(info.rules->index);
+		
 		writer.WriteU32(detoursLibLen);
-		writer.WriteBytes(detoursLib, detoursLibLen);
+		writer.WriteBytes(detoursLib.data(), detoursLibLen);
+		#if PLATFORM_WINDOWS
+		if (m_session.m_sharedMemory.GetFileMappingBackend().linuxBackend)
+		{
+			StringView ubaWineModulePath = GetUbaWineModulePath();
+			char ubaWineModule[256];
+			ubaWineModulePath.Parse(ubaWineModule, 256);
+			writer.WriteU32(ubaWineModulePath.count);
+			writer.WriteBytes(ubaWineModule, ubaWineModulePath.count);
+		}
+		#endif
+		writer.WriteU32(0);
 
 		writer.WriteString(m_realWorkingDir);
 
@@ -1902,7 +1914,16 @@ namespace uba
 				commandLine.erase(pos - commandLine.c_str(), sizeof_array(g_extractExportsStr));
 			}
 
-			LPCSTR dlls[] = { detoursLib };
+			LPCSTR dlls[2];
+			u32 dllCount = 0;
+			dlls[dllCount++] = detoursLib;
+
+			char ubaWineModule[256];
+			if (m_session.m_sharedMemory.GetFileMappingBackend().linuxBackend)
+			{
+				GetUbaWineModulePath().Parse(ubaWineModule, 256);
+				dlls[dllCount++] = ubaWineModule;
+			}
 
 			STARTUPINFOEX siex;
 			STARTUPINFO& si = siex.StartupInfo;
@@ -1978,7 +1999,7 @@ namespace uba
 
 				if (m_detourEnabled)
 				{
-					if (DetourCreateProcessWithDlls(m_realApplication.c_str(), (tchar*)commandLine.c_str(), NULL, NULL, inheritHandles, creationFlags, environment, workingDir, &si, &processInfo, sizeof_array(dlls), dlls, NULL))
+					if (DetourCreateProcessWithDlls(m_realApplication.c_str(), (tchar*)commandLine.c_str(), NULL, NULL, inheritHandles, creationFlags, environment, workingDir, &si, &processInfo, dllCount, dlls, NULL))
 					{
 						DisableProcessPowerThrottling(processInfo.hProcess);
 						break;
diff --git a/Engine/Source/Programs/UnrealBuildAccelerator/Common/Private/UbaSessionServer.cpp b/Engine/Source/Programs/UnrealBuildAccelerator/Common/Private/UbaSessionServer.cpp
--- a/Engine/Source/Programs/UnrealBuildAccelerator/Common/Private/UbaSessionServer.cpp
+++ b/Engine/Source/Programs/UnrealBuildAccelerator/Common/Private/UbaSessionServer.cpp
@@ -1279,8 +1279,6 @@ namespace uba
 			modules.emplace_back(TC("UbaWine.dll.so"), extra.data, DefaultAttributes());
 			extra.Clear().AppendDir(applicationName).EnsureEndsWithSlash().Append(TCV("UbaDetours.dll"));
 			modules.emplace_back(TC("UbaDetours.dll"), extra.data, DefaultAttributes());
-			extra.Clear().AppendDir(applicationName).EnsureEndsWithSlash().Append(TCV("UbaTestApp.exe"));
-			modules.emplace_back(TC("UbaTestApp.exe"), extra.data, DefaultAttributes());
 		}
 		#endif
 
diff --git a/Engine/Source/Programs/UnrealBuildAccelerator/Common/Private/UbaSharedMemoryAllocator.cpp b/Engine/Source/Programs/UnrealBuildAccelerator/Common/Private/UbaSharedMemoryAllocator.cpp
--- a/Engine/Source/Programs/UnrealBuildAccelerator/Common/Private/UbaSharedMemoryAllocator.cpp
+++ b/Engine/Source/Programs/UnrealBuildAccelerator/Common/Private/UbaSharedMemoryAllocator.cpp
@@ -70,22 +70,22 @@ namespace uba
 			if (tempFileSparse)
 				if (!MakeSparse(m_logger, tempFile, m_tempFile, m_capacity))
 					return false;
-			m_fileMapping = FileMapping_CreateFromFile(m_logger, m_tempFile, PAGE_READWRITE, m_capacity, TC("SharedMemoryAllocator"));
+			m_fileMapping = m_fileMappingBackend.CreateFromFile(m_logger, m_tempFile, PAGE_READWRITE, m_capacity, TC("SharedMemoryAllocator"));
 		}
 		else
 		#endif
 		{
-			m_fileMapping = FileMapping_Create(m_logger, PAGE_READWRITE|SEC_RESERVE, m_capacity, nullptr, TC("SharedMemoryAllocator"));
+			m_fileMapping = m_fileMappingBackend.Create(m_logger, PAGE_READWRITE|SEC_RESERVE, m_capacity, nullptr, TC("SharedMemoryAllocator"));
 		}
 		if (!m_fileMapping.IsValid())
 			return false;
 
-		m_memory = FileMapping_Map(m_logger, m_fileMapping, FILE_MAP_READ|FILE_MAP_WRITE, 0, m_capacity);
+		m_memory = m_fileMappingBackend.Map(m_logger, m_fileMapping, FILE_MAP_READ|FILE_MAP_WRITE, 0, m_capacity, false);
 		UBA_ASSERTF(m_memory, TC("MapViewOfFile failed trying to map %llu bytes (%s)"), m_capacity, LastErrorToText().data);
 
 		if (!m_memory)
 		{
-			FileMapping_Close(m_logger, m_fileMapping, TC("SharedMemoryAllocator"));
+			m_fileMappingBackend.Close(m_logger, m_fileMapping, TC("SharedMemoryAllocator"));
 			return false;
 		}
 
@@ -96,11 +96,11 @@ namespace uba
 	void SharedMemoryAllocator::Deinit(bool keepFileMapping)
 	{
 		if (m_memory)
-			FileMapping_Unmap(m_logger, m_memory, m_capacity, TC("SharedMemoryAllocator"));
+			m_fileMappingBackend.Unmap(m_logger, m_memory, m_capacity, TC("SharedMemoryAllocator"), false);
 		m_memory = nullptr;
 
 		if (m_fileMapping.IsValid() && !keepFileMapping)
-			FileMapping_Close(m_logger, m_fileMapping, TC("SharedMemoryAllocator"));
+			m_fileMappingBackend.Close(m_logger, m_fileMapping, TC("SharedMemoryAllocator"));
 		m_fileMapping = {};
 
 		if (m_tempFile != InvalidFileHandle && !keepFileMapping)
@@ -294,7 +294,7 @@ namespace uba
 		UBA_ASSERT(h.slices.size() == 1);
 		auto& slice = h.slices.back();
 		if (m_tempFile == InvalidFileHandle)
-			FileMapping_Commit(m_logger, m_memory + slice.offset + offset, size);
+			m_fileMappingBackend.Commit(m_logger, m_memory + slice.offset + offset, size, false);
 
 		SCOPED_FUTEX(m_lock, lock);
 		m_committed += size;
@@ -355,6 +355,11 @@ namespace uba
 		return h.refCount;
 	}
 
+	FileMappingBackend& SharedMemoryAllocator::GetFileMappingBackend()
+	{
+		return m_fileMappingBackend;
+	}
+
 	void SharedMemoryAllocator::WriteViewSlices(BinaryWriter& out, SharedMemoryHandle handle, u64 offset, u64 size)
 	{
 		auto& h = GetHandle(handle);
@@ -404,7 +409,7 @@ namespace uba
 			// TODO: Right now we leak this mapping. Not good..
 			// .. but only used in very special situations
 			FileMappingHandle file = FileMappingHandle::FromU64(handle.internalHandle & ~0x8000'0000'0000'0000);
-			u8* mem = FileMapping_Map(m_logger, file, FILE_MAP_READ, offset, size);
+			u8* mem = m_fileMappingBackend.Map(m_logger, file, FILE_MAP_READ, offset, size, false);
 			UBA_ASSERT(mem);
 			view.ReferenceExternalMemory(mem, size);
 			return true;
@@ -492,7 +497,7 @@ namespace uba
 		auto& kernelStats = KernelStats::GetCurrent();
 
 		TimerScope ts(kernelStats.virtualAlloc2);
-		u8* memory = (u8*)FileMapping_ReservePlaceholder(nullptr, committed);
+		u8* memory = (u8*)m_fileMappingBackend.ReservePlaceholder(nullptr, committed);
 		ts.Leave();
 
 		u64 pos = 0;
@@ -501,7 +506,7 @@ namespace uba
 			kernelStats.mapViewOfFile3.bytes += slice.size;
 			TimerScope ts2(kernelStats.mapViewOfFile3);
 
-			void* viewA = FileMapping_MapPlaceholder(m_fileMapping, (MemoryMapType)mapType, memory, pos, committed, slice.offset, slice.size);
+			void* viewA = m_fileMappingBackend.MapPlaceholder(m_fileMapping, (MemoryMapType)mapType, memory, pos, committed, slice.offset, slice.size);
 			UBA_ASSERTF(viewA == memory+pos, TC("MapViewOfFile3 failed in MapView 0x%p %llu for % s (%s)"), memory, pos, h.GetHint(), LastErrorToText().data); (void)viewA;
 			pos += slice.size;
 		}
@@ -534,13 +539,14 @@ namespace uba
 				subMappings.push_back(committed);
 			committed += slice.size;
 		}
-		FileMapping_UnmapPlaceholder(memory, committed, committed, subMappings.data(), subMappings.size());
+		m_fileMappingBackend.UnmapPlaceholder(memory, committed, committed, subMappings.data(), subMappings.size());
 
 		ReleaseHandle(handle);
 	}
 
-	SharedMemoryAllocator::SharedMemoryAllocator(Logger& logger)
+	SharedMemoryAllocator::SharedMemoryAllocator(Logger& logger, FileMappingBackend& fileMappingBackend)
 	:	m_logger(logger)
+	,	m_fileMappingBackend(fileMappingBackend)
 	{
 	}
 
@@ -564,7 +570,7 @@ namespace uba
 	void SharedMemoryAllocator::UnmapExternalMapping(SharedMemoryHandle handle, u8* memory, u64 size)
 	{
 		UBA_ASSERT((handle.internalHandle & 0x8000'0000'0000'0000) != 0);
-		bool res = FileMapping_Unmap(m_logger, memory, size, TC("ExternalMapping"));
+		bool res = m_fileMappingBackend.Unmap(m_logger, memory, size, TC("ExternalMapping"), false);
 		UBA_ASSERT(res); (void)res;
 	}
 
@@ -677,7 +683,7 @@ namespace uba
 			UBA_SHIPPING_ASSERTF(m_reserved + toCommit <= m_capacity, TC("Not enough space in capacity for commit (Capacity: %llu, Reserved: %llu Requested: %llu"), m_capacity, m_reserved, toCommit);
 			if (m_tempFile == InvalidFileHandle)
 			{
-				bool res = FileMapping_Commit(m_logger, m_memory + m_reserved, toCommit);
+				bool res = m_fileMappingBackend.Commit(m_logger, m_memory + m_reserved, toCommit, false);
 				UBA_SHIPPING_ASSERTF(res, TC("VirtualAlloc failed commit %llu bytes at address 0x%llx + %llu (%s)"), toCommit, u64(m_memory), m_reserved, LastErrorToText().data);
 			}
 			u64 offset = m_reserved;
diff --git a/Engine/Source/Programs/UnrealBuildAccelerator/Common/Private/UbaStorage.cpp b/Engine/Source/Programs/UnrealBuildAccelerator/Common/Private/UbaStorage.cpp
--- a/Engine/Source/Programs/UnrealBuildAccelerator/Common/Private/UbaStorage.cpp
+++ b/Engine/Source/Programs/UnrealBuildAccelerator/Common/Private/UbaStorage.cpp
@@ -46,6 +46,7 @@ namespace uba
 		table.GetValueAsBool(allowDeferredDeletes, TC("AllowDeferredDeletes"));
 		table.GetValueAsBool(skipCleanupOnShutdown, TC("SkipCleanupOnShutdown"));
 		table.GetValueAsBool(allowUnDropOfCas, TC("AllowUnDropOfCas"));
+		table.GetValueAsBool(sideStepWineForSharedMemory, TC("SideStepWineForSharedMemory"));
 		table.GetValueAsU32(fileEntryTableWipeThreshold, TC("FileEntryTableWipeThreshold"));
 		table.GetValueAsU32(maxParallelCopyOrLink, TC("MaxParallelCopyOrLink"));
 		table.GetValueAsU32(deleteCasBatchSize, TC("DeleteCasBatchSize"));
@@ -86,7 +87,7 @@ namespace uba
 			return nullptr;
 		}
 
-		MemoryMappingScope(Logger& l, SharedMemoryAllocator& a) : logger(l), allocator(a) {}
+		MemoryMappingScope(Logger& l, SharedMemoryAllocator& a) : logger(l), allocator(a), view(a.GetFileMappingBackend()) {}
 		~MemoryMappingScope()
 		{
 			if (!mappedMem)
@@ -1040,10 +1041,28 @@ namespace uba
 		return success;
 	}
 
+	FileMappingBackend& GetFileMappingBackend(const StorageCreateInfo& info, Logger& logger)
+	{
+		#if PLATFORM_WINDOWS
+		if (info.sideStepWineForSharedMemory && IsRunningWine())
+		{
+			static FileMappingBackend wineBackend;
+			if (HMODULE h = GetUbaWineModule())
+				if (FileMapping_GetWineBackend(wineBackend, logger, h))
+				{
+					logger.Debug(TC("Using wine filemapping backend"));
+					return wineBackend;
+				}
+		}
+		#endif
+		static FileMappingBackend backend;
+		return backend;
+	}
+
 	StorageImpl::StorageImpl(const StorageCreateInfo& info, const tchar* logPrefix)
 	:	m_logger(info.writer, logPrefix)
 	,	m_workManager(info.workManager)
-	,	m_sharedMemory(m_logger)
+	,	m_sharedMemory(m_logger, GetFileMappingBackend(info, m_logger))
 	,	m_bufferSlots(BufferSlotSize)
 	,	m_maxParallelCopyOrLinkBottleneck(info.maxParallelCopyOrLink)
 	,	m_casDataBuffer(m_logger, m_sharedMemory)
diff --git a/Engine/Source/Programs/UnrealBuildAccelerator/Common/Public/UbaEnvironment.h b/Engine/Source/Programs/UnrealBuildAccelerator/Common/Public/UbaEnvironment.h
--- a/Engine/Source/Programs/UnrealBuildAccelerator/Common/Public/UbaEnvironment.h
+++ b/Engine/Source/Programs/UnrealBuildAccelerator/Common/Public/UbaEnvironment.h
@@ -32,4 +32,9 @@ namespace uba
 
 	void SetMemoryWorkingSet(Logger& logger, u64 minWorkingSet);
 	void LockMemory(Logger& logger, void* dest, u64 size, const tchar* hint);
+
+	#if PLATFORM_WINDOWS
+	StringView GetUbaWineModulePath();
+	HMODULE GetUbaWineModule();
+	#endif
 }
\ No newline at end of file
diff --git a/Engine/Source/Programs/UnrealBuildAccelerator/Common/Public/UbaSharedMemoryAllocator.h b/Engine/Source/Programs/UnrealBuildAccelerator/Common/Public/UbaSharedMemoryAllocator.h
--- a/Engine/Source/Programs/UnrealBuildAccelerator/Common/Public/UbaSharedMemoryAllocator.h
+++ b/Engine/Source/Programs/UnrealBuildAccelerator/Common/Public/UbaSharedMemoryAllocator.h
@@ -12,9 +12,10 @@
 
 namespace uba
 {
-	struct BinaryWriter;
 	class Logger;
 	class Trace;
+	struct BinaryWriter;
+	struct FileMappingBackend;
 
 	////////////////////////////////////////////////////////////////////////////////////////////////////
 
@@ -45,6 +46,7 @@ namespace uba
 
 		u64 GetCommitted(SharedMemoryHandle handle);
 		u64 GetRefCount(SharedMemoryHandle handle);
+		FileMappingBackend& GetFileMappingBackend();
 
 		void WriteViewSlices(BinaryWriter& out, SharedMemoryHandle handle, u64 offset = 0, u64 size = ~0llu);
 		bool MapView(SharedMemoryView& view, SharedMemoryHandle handle, const tchar* hint, SharedMemoryMapType mapType = SharedMemoryMapType_ReadOnly, u64 offset = 0, u64 size = ~0llu);
@@ -56,7 +58,7 @@ namespace uba
 
 		void TraceStats(Trace& trace, u32 startRow);
 
-		SharedMemoryAllocator(Logger& logger);
+		SharedMemoryAllocator(Logger& logger, FileMappingBackend& fileMappingBackend);
 		~SharedMemoryAllocator();
 
 		// This is really ugly. But very contained for virtual file mappings
@@ -76,6 +78,7 @@ namespace uba
 		void FreeNoLock(Slice slice);
 
 		Logger& m_logger;
+		FileMappingBackend& m_fileMappingBackend;
 
 		u64 m_commitStepSize;
 		FileHandle m_tempFile = InvalidFileHandle;
diff --git a/Engine/Source/Programs/UnrealBuildAccelerator/Common/Public/UbaStorage.h b/Engine/Source/Programs/UnrealBuildAccelerator/Common/Public/UbaStorage.h
--- a/Engine/Source/Programs/UnrealBuildAccelerator/Common/Public/UbaStorage.h
+++ b/Engine/Source/Programs/UnrealBuildAccelerator/Common/Public/UbaStorage.h
@@ -134,6 +134,7 @@ namespace uba
 		bool allowDeferredDeletes = true;
 		bool skipCleanupOnShutdown = false;
 		bool allowUnDropOfCas = true;
+		bool sideStepWineForSharedMemory = true;
 		u64 sharedMemoryCapacity = 1024ull*1024*1024*1024; // 1TB
 		u64 sharedMemoryCommitSize = 64ull*1024*1024;
 		const tchar* sharedMemoryTempFile = TC("");
diff --git a/Engine/Source/Programs/UnrealBuildAccelerator/Core/Private/UbaSharedMemoryView.cpp b/Engine/Source/Programs/UnrealBuildAccelerator/Core/Private/UbaSharedMemoryView.cpp
--- a/Engine/Source/Programs/UnrealBuildAccelerator/Core/Private/UbaSharedMemoryView.cpp
+++ b/Engine/Source/Programs/UnrealBuildAccelerator/Core/Private/UbaSharedMemoryView.cpp
@@ -17,6 +17,11 @@ namespace uba
 {
 	////////////////////////////////////////////////////////////////////////////////////////////////////
 
+	SharedMemoryView::SharedMemoryView(FileMappingBackend& backend)
+	:	m_fileMappingBackend(backend)
+	{
+	}
+
 	bool SharedMemoryView::Init(SharedMemoryAllocatorHandle handle, u64 capacity, void* baseAddress, const tchar* hint)
 	{
 		return InternalInit(handle, capacity, baseAddress, hint);
@@ -57,7 +62,7 @@ namespace uba
 			return;
 		}
 
-		FileMapping_UnmapPlaceholder(m_memory, m_capacity, m_mappedSize, m_additionalMappings.data(), m_additionalMappings.size());
+		m_fileMappingBackend.UnmapPlaceholder(m_memory, m_capacity, m_mappedSize, m_additionalMappings.data(), m_additionalMappings.size());
 		m_additionalMappings.clear();
 		m_mappedSize = 0;
 		m_memory = nullptr;
@@ -84,7 +89,7 @@ namespace uba
 		UBA_ASSERT(!m_memory);
 		m_capacity = AlignUp(capacity, PageSize);
 		TimerScope ts(KernelStats::GetCurrent().virtualAlloc2);
-		m_memory = (u8*)FileMapping_ReservePlaceholder(baseAddress, m_capacity);
+		m_memory = (u8*)m_fileMappingBackend.ReservePlaceholder(baseAddress, m_capacity);
 		return m_memory != nullptr;
 	}
 
@@ -138,7 +143,7 @@ namespace uba
 		auto& kernelStats = KernelStats::GetCurrent();
 
 		TimerScope ts(kernelStats.mapViewOfFile3);
-		void* viewA = FileMapping_MapPlaceholder(handle, (MemoryMapType)mapType, m_memory, m_mappedSize, m_capacity, offset, size);
+		void* viewA = m_fileMappingBackend.MapPlaceholder(handle, (MemoryMapType)mapType, m_memory, m_mappedSize, m_capacity, offset, size);
 		UBA_ASSERTF(viewA == m_memory + m_mappedSize, TC("MapViewOfFile3 failed extending memory with %llu bytes at offset %llu - %s (%s)"), size, offset, LastErrorToText().data, hint);
 		if (!viewA)
 			return false;
diff --git a/Engine/Source/Programs/UnrealBuildAccelerator/Core/Private/UbaStringBuffer.cpp b/Engine/Source/Programs/UnrealBuildAccelerator/Core/Private/UbaStringBuffer.cpp
--- a/Engine/Source/Programs/UnrealBuildAccelerator/Core/Private/UbaStringBuffer.cpp
+++ b/Engine/Source/Programs/UnrealBuildAccelerator/Core/Private/UbaStringBuffer.cpp
@@ -458,16 +458,7 @@ namespace uba
 
 	u32 StringBufferBase::Parse(char* out, u64 outCapacity) const
 	{
-		#if PLATFORM_WINDOWS
-		return WideCharToMultiByte(CP_ACP, 0, data, count+1, out, int(outCapacity), nullptr, nullptr);
-		#else
-		if (outCapacity == 0)
-			return 0;
-		u32 toCopy = Min(u32(outCapacity - 1), count);
-		memcpy(out, data, toCopy);
-		out[toCopy] = 0;
-		return toCopy;
-		#endif
+		return StringView(*this).Parse(out, outCapacity);
 	}
 
 	////////////////////////////////////////////////////////////////////////////////////////////////////
@@ -513,6 +504,21 @@ namespace uba
 		return {};
 	}
 
+	u32 StringView::Parse(char* out, u64 outCapacity) const
+	{
+		#if PLATFORM_WINDOWS
+		return WideCharToMultiByte(CP_ACP, 0, data, count+1, out, int(outCapacity), nullptr, nullptr);
+		#else
+		if (outCapacity == 0)
+			return 0;
+		u32 toCopy = Min(u32(outCapacity - 1), count);
+		memcpy(out, data, toCopy);
+		out[toCopy] = 0;
+		return toCopy;
+		#endif
+	}
+
+
 	StringView ToView(const tchar* s)
 	{
 		return s ? StringView(s, TStrlen(s)) : StringView();;
diff --git a/Engine/Source/Programs/UnrealBuildAccelerator/Core/Public/UbaFileMapping.h b/Engine/Source/Programs/UnrealBuildAccelerator/Core/Public/UbaFileMapping.h
--- a/Engine/Source/Programs/UnrealBuildAccelerator/Core/Public/UbaFileMapping.h
+++ b/Engine/Source/Programs/UnrealBuildAccelerator/Core/Public/UbaFileMapping.h
@@ -9,9 +9,15 @@ namespace uba
 	enum FileHandle : u64;
 	class Logger;
 
+#ifdef UBA_FM_STRING_TYPE
+	using fmchar = UBA_FM_STRING_TYPE;
+#else
+	using fmchar = tchar;
+#endif
+
 	////////////////////////////////////////////////////////////////////////////////////////////////////
 
-	#if !PLATFORM_WINDOWS
+#if !PLATFORM_WINDOWS
 	inline constexpr u32 FILE_MAP_WRITE = 0x0002;
 	inline constexpr u32 FILE_MAP_READ = 0x0004;
 	inline constexpr u32 FILE_MAP_ALL_ACCESS = FILE_MAP_WRITE | FILE_MAP_READ;
@@ -19,21 +25,21 @@ namespace uba
 	inline constexpr u32 PAGE_READONLY = PROT_READ;
 	inline constexpr u32 PAGE_READWRITE = PROT_READ | PROT_WRITE;
 
-	#if PLATFORM_MAC
+#if PLATFORM_MAC
 	inline constexpr u32 SHM_MAX_FILENAME = PSHMNAMLEN;
-	#else
+#else
 	inline constexpr u32 SHM_MAX_FILENAME = 38;
-	#endif
-	#endif
+#endif
+#endif
 
 	enum MemoryMapType { MemoryMapType_ReadOnly, MemoryMapType_ReadWrite, MemoryMapType_CopyOnWrite };
 
-	FileMappingHandle FileMapping_Create(Logger& logger, u32 flProtect, u64 maxSize, const tchar* name, const tchar* hint);
-	FileMappingHandle FileMapping_CreateFromFile(Logger& logger, FileHandle file, u32 flProtect, u64 maxSize, const tchar* hint);
+	FileMappingHandle FileMapping_Create(Logger& logger, u32 flProtect, u64 maxSize, const fmchar* name, const fmchar* hint);
+	FileMappingHandle FileMapping_CreateFromFile(Logger& logger, FileHandle file, u32 flProtect, u64 maxSize, const fmchar* hint);
 	u8* FileMapping_Map(Logger& logger, FileMappingHandle fileMappingObject, u32 dwDesiredAccess, u64 offset, u64 dwNumberOfBytesToMap, bool allowDiscard = false);
 	bool FileMapping_Commit(Logger& logger, void* address, u64 size, bool allowDiscard = false);
-	bool FileMapping_Unmap(Logger& logger, const void* lpBaseAddress, u64 bytesToUnmap, const tchar* hint, bool allowDiscard = false);
-	bool FileMapping_Close(Logger& logger, FileMappingHandle h, const tchar* hint);
+	bool FileMapping_Unmap(Logger& logger, const void* lpBaseAddress, u64 bytesToUnmap, const fmchar* hint, bool allowDiscard = false);
+	bool FileMapping_Close(Logger& logger, FileMappingHandle h, const fmchar* hint);
 	void* FileMapping_ReservePlaceholder(void* baseAddress, u64 capacity);
 	void* FileMapping_MapPlaceholder(FileMappingHandle handle, MemoryMapType mapType, void* targetAddress, u64 targetOffset, u64 targetCapacity, u64 handleOffset, u64 size);
 	bool FileMapping_UnmapPlaceholder(void* memory, u64 capacity, u64 mappedSize, u64* subMappings, u64 subMappingsCount);
@@ -43,4 +49,35 @@ namespace uba
 	u8* FileMapping_MapFromHost(FileMappingHandle handle, u64 size, u64 offset, bool writable, const char* hint);
 
 	////////////////////////////////////////////////////////////////////////////////////////////////////
+
+	#define UBA_FM_FUNCTIONS \
+		UBA_FM_FUNC(Create) \
+		UBA_FM_FUNC(CreateFromFile) \
+		UBA_FM_FUNC(Map) \
+		UBA_FM_FUNC(Commit) \
+		UBA_FM_FUNC(Unmap) \
+		UBA_FM_FUNC(Close) \
+		UBA_FM_FUNC(ReservePlaceholder) \
+		UBA_FM_FUNC(MapPlaceholder) \
+		UBA_FM_FUNC(UnmapPlaceholder) \
+		UBA_FM_FUNC(DuplicateFromHost) \
+		UBA_FM_FUNC(MapFromHost) \
+		//UBA_FM_FUNC(CopyMem) \
+
+	struct FileMappingBackend
+	{
+		#define UBA_FM_FUNC(func) decltype(FileMapping_##func)* func = FileMapping_##func;
+		UBA_FM_FUNCTIONS
+		#undef UBA_FM_FUNC
+
+		#if PLATFORM_WINDOWS
+		bool linuxBackend = false;
+		#endif
+	};
+
+	#if PLATFORM_WINDOWS
+	bool FileMapping_GetWineBackend(FileMappingBackend& out, Logger& logger, HMODULE wineDll);
+	#endif
+
+	////////////////////////////////////////////////////////////////////////////////////////////////////
 }
diff --git a/Engine/Source/Programs/UnrealBuildAccelerator/Core/Public/UbaFileMapping.inl b/Engine/Source/Programs/UnrealBuildAccelerator/Core/Public/UbaFileMapping.inl
--- a/Engine/Source/Programs/UnrealBuildAccelerator/Core/Public/UbaFileMapping.inl
+++ b/Engine/Source/Programs/UnrealBuildAccelerator/Core/Public/UbaFileMapping.inl
@@ -6,6 +6,15 @@
 
 ////////////////////////////////////////////////////////////////////////////////////////////////////
 
+#if 0
+#define WRITE_DEBUG(format, ...) \
+do { \
+    char str[1024]; \
+    int len = snprintf(str, sizeof(str), format "\n", ##__VA_ARGS__); \
+	ssize_t ignored_result = write(STDERR_FILENO, str, (size_t)len); \
+} while (0)
+#endif
+
 #if PLATFORM_WINDOWS
 Bottleneck& g_createFileHandleBottleneck = *new Bottleneck(8); // Allocated and leaked just to prevent shutdown asserts in debug
 
@@ -153,7 +162,7 @@ bool GarbageCollectFileMappings(Logger& logger)
 }
 #endif
 
-FileMappingHandle FileMapping_Create(Logger& logger, u32 flProtect, u64 maxSize, const tchar* name, const tchar* hint)
+FileMappingHandle FileMapping_Create(Logger& logger, u32 flProtect, u64 maxSize, const fmchar* name, const fmchar* hint)
 {
 #if PLATFORM_WINDOWS
 	ExtendedTimerScope ts(KernelStats::GetCurrent().createFileMapping);
@@ -239,7 +248,7 @@ FileMappingHandle FileMapping_Create(Logger& logger, u32 flProtect, u64 maxSize,
 #endif
 }
 
-FileMappingHandle FileMapping_CreateFromFile(Logger& logger, FileHandle file, u32 protect, u64 maxSize, const tchar* hint)
+FileMappingHandle FileMapping_CreateFromFile(Logger& logger, FileHandle file, u32 protect, u64 maxSize, const fmchar* hint)
 {
 #if PLATFORM_WINDOWS
 	ExtendedTimerScope ts(KernelStats::GetCurrent().createFileMapping);
@@ -333,7 +342,7 @@ bool FileMapping_Commit(Logger& logger, void* address, u64 size, bool allowDisca
 #endif
 }
 
-bool FileMapping_Unmap(Logger& logger, const void* address, u64 bytesToUnmap, const tchar* hint, bool allowDiscard)
+bool FileMapping_Unmap(Logger& logger, const void* address, u64 bytesToUnmap, const fmchar* hint, bool allowDiscard)
 {
 	if (!address)
 		return true;
@@ -364,7 +373,7 @@ bool FileMapping_Unmap(Logger& logger, const void* address, u64 bytesToUnmap, co
 #endif
 }
 
-bool FileMapping_Close(Logger& logger, FileMappingHandle h, const tchar* hint)
+bool FileMapping_Close(Logger& logger, FileMappingHandle h, const fmchar* hint)
 {
 	if (!h.IsValid())
 		return true;
@@ -411,9 +420,13 @@ void* FileMapping_ReservePlaceholder(void* baseAddress, u64 capacity)
 #if PLATFORM_WINDOWS
 	void* memory = VirtualAlloc2(GetCurrentProcess(), baseAddress, capacity, MEM_RESERVE | MEM_RESERVE_PLACEHOLDER, PAGE_NOACCESS, nullptr, 0);
 #else
-	void* memory = mmap(NULL, capacity, PROT_NONE, MAP_PRIVATE | MAP_ANONYMOUS, -1, 0);
+	int flags = MAP_PRIVATE | MAP_ANONYMOUS;
+	if (baseAddress)
+		flags |= MAP_FIXED;
+	void* memory = mmap(baseAddress, capacity, PROT_NONE, flags, -1, 0);
 	if (memory == MAP_FAILED)
 		memory = nullptr;
+	//WRITE_DEBUG("RESERVE: Base: 0x%llx Capacity: %llu -> 0x%llx", u64(baseAddress), capacity, u64(memory));
 #endif
 	UBA_ASSERTF(memory, TC("VirtualAlloc2 failed (%s)"), LastErrorToText().data);
 	return memory;
@@ -443,14 +456,33 @@ void* FileMapping_MapPlaceholder(FileMappingHandle handle, MemoryMapType mapType
 	//TimerScope ts(kernelStats.mapViewOfFile3);
 	void* viewA = MapViewOfFile3(AsHANDLE(handle), GetCurrentProcess(), (u8*)targetAddress+targetOffset, handleOffset, size, MEM_REPLACE_PLACEHOLDER, protect, nullptr, 0);
 #else
-	UBA_ASSERT(mapType != MemoryMapType_CopyOnWrite);
+	//UBA_ASSERT(mapType != MemoryMapType_CopyOnWrite);
 	u32 protect = PROT_READ;
-	if (mapType == MemoryMapType_ReadWrite)
+	u32 mapFlags = MAP_FIXED;
+
+	if (mapType == MemoryMapType_CopyOnWrite)
+	{
+		//WRITE_DEBUG("HEEERE: %llu", size);
+		mapFlags |= MAP_PRIVATE;
 		protect |= PROT_WRITE;
+	}
+	else
+	{
+		mapFlags |= MAP_SHARED;
+		if (mapType == MemoryMapType_ReadWrite)
+			protect |= PROT_WRITE;
+	}
 	//TimerScope ts(kernelStats.mapViewOfFile3);
-	void* viewA = mmap((u8*)targetAddress+targetOffset, size, protect, MAP_SHARED | MAP_FIXED, handle.shmFd, handleOffset);
+	void* viewA = mmap((u8*)targetAddress+targetOffset, size, protect, mapFlags, handle.shmFd, handleOffset);
 	if (viewA == MAP_FAILED)
 		viewA = nullptr;
+	//else if (mapType == MemoryMapType_CopyOnWrite)
+	//{
+	//	struct stat st;
+	//	fstat(handle.shmFd, &st);
+	//	WRITE_DEBUG("shmFd=%i size=%llu handleOffset=%llu mapSize=%llu -> 0x%llx", handle.shmFd, (u64)st.st_size, handleOffset, size, u64(viewA));
+	//}
+	//WRITE_DEBUG("mmap: Base: %llu Protect: %i Offset: %llu Size: %llu HandleOffset: %llu: SHM: %i -> 0x%llx (%s)", u64(targetAddress), protect, targetOffset, size, handleOffset, handle.shmFd, u64(viewA), viewA ? "success" : strerror(errno));
 #endif
 	return viewA;
 }
@@ -503,7 +535,15 @@ FileMappingHandle FileMapping_DuplicateFromHost(FileMappingHandle handle, const
 	}
 	return FileMappingHandle(u64(realHandle));
 #else
-	return handle;
+	char uidName[128];
+	GetMappingHandleName(uidName, handle.internal);
+	int shmFd = shm_open(uidName, O_RDWR, S_IRUSR | S_IWUSR);
+	UBA_ASSERTF(shmFd != -1, "shm_open failed for %s and handle %s (%s)", hint, uidName, strerror(errno));
+	if (shmFd == -1)
+		return {};
+	FileMappingHandle newHandle = handle;
+	newHandle.shmFd = shmFd;
+	return newHandle;
 #endif
 }
 
diff --git a/Engine/Source/Programs/UnrealBuildAccelerator/Core/Public/UbaFileMappingWin.inl b/Engine/Source/Programs/UnrealBuildAccelerator/Core/Public/UbaFileMappingWin.inl
--- a/Engine/Source/Programs/UnrealBuildAccelerator/Core/Public/UbaFileMappingWin.inl
+++ b/Engine/Source/Programs/UnrealBuildAccelerator/Core/Public/UbaFileMappingWin.inl
@@ -2,10 +2,10 @@
 
 #pragma once
 
-////////////////////////////////////////////////////////////////////////////////////////////////////
-
 namespace uba
 {
+	////////////////////////////////////////////////////////////////////////////////////////////////////
+
 	HANDLE g_hostProcess;
 
 	namespace internal
@@ -13,60 +13,96 @@ namespace uba
 		#include "UbaFileMapping.inl"
 	}
 
+	////////////////////////////////////////////////////////////////////////////////////////////////////
+
+	#define UBA_FM_FUNC(func) decltype(FileMapping_##func)* g_##func##Func = internal::FileMapping_##func;
+	UBA_FM_FUNCTIONS
+	#undef UBA_FM_FUNC
+
+	////////////////////////////////////////////////////////////////////////////////////////////////////
+
 	FileMappingHandle FileMapping_Create(Logger& logger, u32 flProtect, u64 maxSize, const tchar* name, const tchar* hint)
 	{
+		//return g_CreateFunc(logger, flProtect, maxSize, name, hint);
 		return internal::FileMapping_Create(logger, flProtect, maxSize, name, hint);
 	}
 
 	FileMappingHandle FileMapping_CreateFromFile(Logger& logger, FileHandle file, u32 flProtect, u64 maxSize, const tchar* hint)
 	{
+		//return g_CreateFromFileFunc(logger, file, flProtect, maxSize, hint);
 		return internal::FileMapping_CreateFromFile(logger, file, flProtect, maxSize, hint);
 	}
 
 	u8* FileMapping_Map(Logger& logger, FileMappingHandle fileMappingObject, u32 dwDesiredAccess, u64 offset, u64 dwNumberOfBytesToMap, bool allowDiscard)
 	{
+		//return g_MapFunc(logger, fileMappingObject, dwDesiredAccess, offset, dwNumberOfBytesToMap, allowDiscard);
 		return internal::FileMapping_Map(logger, fileMappingObject, dwDesiredAccess, offset, dwNumberOfBytesToMap, allowDiscard);
 	}
 
 	bool FileMapping_Commit(Logger& logger, void* address, u64 size, bool allowDiscard)
 	{
+		//return g_CommitFunc(logger, address, size, allowDiscard);
 		return internal::FileMapping_Commit(logger, address, size, allowDiscard);
 	}
 
 	bool FileMapping_Unmap(Logger& logger, const void* lpBaseAddress, u64 bytesToUnmap, const tchar* hint, bool allowDiscard)
 	{
+		//return g_UnmapFunc(logger, lpBaseAddress, bytesToUnmap, hint, allowDiscard);
 		return internal::FileMapping_Unmap(logger, lpBaseAddress, bytesToUnmap, hint, allowDiscard);
 	}
 
 	bool FileMapping_Close(Logger& logger, FileMappingHandle h, const tchar* hint)
 	{
+		//return g_CloseFunc(logger, h, hint);
 		return internal::FileMapping_Close(logger, h, hint);
 	}
 
 	void* FileMapping_ReservePlaceholder(void* baseAddress, u64 capacity)
 	{
+		//return g_ReservePlaceholderFunc(baseAddress, capacity);
 		return internal::FileMapping_ReservePlaceholder(baseAddress, capacity);
 	}
 
 	void* FileMapping_MapPlaceholder(FileMappingHandle handle, MemoryMapType mapType, void* targetAddress, u64 targetOffset, u64 targetCapacity, u64 handleOffset, u64 size)
 	{
+		//return g_MapPlaceholderFunc(handle, mapType, targetAddress, targetOffset, targetCapacity, handleOffset, size);
 		return internal::FileMapping_MapPlaceholder(handle, mapType, targetAddress, targetOffset, targetCapacity, handleOffset, size);
 	}
 
 	bool FileMapping_UnmapPlaceholder(void* memory, u64 capacity, u64 mappedSize, u64* subMappings, u64 subMappingsCount)
 	{
+		//return g_UnmapPlaceholderFunc(memory, capacity, mappedSize, subMappings, subMappingsCount);
 		return internal::FileMapping_UnmapPlaceholder(memory, capacity, mappedSize, subMappings, subMappingsCount);
 	}
 
 	FileMappingHandle FileMapping_DuplicateFromHost(FileMappingHandle handle, const char* hint)
 	{
+		//return g_DuplicateFromHostFunc(handle, hint);
 		return internal::FileMapping_DuplicateFromHost(handle, hint);
 	}
 
 	u8* FileMapping_MapFromHost(FileMappingHandle handle, u64 size, u64 offset, bool writable, const char* hint)
 	{
+		//return g_MapFromHostFunc(handle, size, offset, writable, hint);
 		return internal::FileMapping_MapFromHost(handle, size, offset, writable, hint);
 	}
-}
 
-////////////////////////////////////////////////////////////////////////////////////////////////////
+	////////////////////////////////////////////////////////////////////////////////////////////////////
+
+	bool FileMapping_GetWineBackend(FileMappingBackend& out, Logger& logger, HMODULE wineDll)
+	{
+		#define UBA_FM_FUNC(func) \
+			auto func##Func = (decltype(FileMapping_##func)*)GetProcAddress(wineDll, "FileMapping_" #func); \
+			if (!func##Func) return logger.Warning(TC("FileMapping_" #func " is not exported from wine dll"));
+		UBA_FM_FUNCTIONS
+		#undef UBA_FM_FUNC
+
+		#define UBA_FM_FUNC(func) out.func = func##Func;
+		UBA_FM_FUNCTIONS
+		#undef UBA_FM_FUNC
+		out.linuxBackend = true;
+		return true;
+	}
+
+	////////////////////////////////////////////////////////////////////////////////////////////////////
+}
diff --git a/Engine/Source/Programs/UnrealBuildAccelerator/Core/Public/UbaSharedMemoryView.h b/Engine/Source/Programs/UnrealBuildAccelerator/Core/Public/UbaSharedMemoryView.h
--- a/Engine/Source/Programs/UnrealBuildAccelerator/Core/Public/UbaSharedMemoryView.h
+++ b/Engine/Source/Programs/UnrealBuildAccelerator/Core/Public/UbaSharedMemoryView.h
@@ -8,6 +8,7 @@
 namespace uba
 {
 	struct BinaryReader;
+	struct FileMappingBackend;
 
 	////////////////////////////////////////////////////////////////////////////////////////////////////
 
@@ -29,7 +30,7 @@ namespace uba
 	class SharedMemoryView
 	{
 	public:
-		SharedMemoryView() = default;
+		SharedMemoryView(FileMappingBackend& backend);
 		SharedMemoryView(SharedMemoryView&&) = default;
 
 		bool Init(SharedMemoryAllocatorHandle handle, u64 capacity, void* baseAddress = nullptr, const tchar* hint = TC(""));
@@ -51,6 +52,7 @@ namespace uba
 		void ReferenceExternalMemory(u8* memory, u64 mappedSize);
 		bool ExtendMemory(SharedMemoryAllocatorHandle handle, u64 offset, u64 size, SharedMemoryMapType mapType, const tchar* hint, u8 init = 0);
 
+		FileMappingBackend& m_fileMappingBackend;
 		u8* m_memory = nullptr;
 		u64 m_mappedSize = 0;
 		u64 m_capacity = 0;
diff --git a/Engine/Source/Programs/UnrealBuildAccelerator/Core/Public/UbaStringBuffer.h b/Engine/Source/Programs/UnrealBuildAccelerator/Core/Public/UbaStringBuffer.h
--- a/Engine/Source/Programs/UnrealBuildAccelerator/Core/Public/UbaStringBuffer.h
+++ b/Engine/Source/Programs/UnrealBuildAccelerator/Core/Public/UbaStringBuffer.h
@@ -128,6 +128,8 @@ namespace uba
 
 		TString ToString() const { return TString(data, data + count); }
 
+		u32 Parse(char* out, u64 outCapacity) const; // Note, return value contains length + null termination.. zero if failed
+
 		const tchar* data;
 		u32 count;
 	};
diff --git a/Engine/Source/Programs/UnrealBuildAccelerator/Detours/Private/PosixOS/UbaDetoursFunctionsPosix.cpp b/Engine/Source/Programs/UnrealBuildAccelerator/Detours/Private/PosixOS/UbaDetoursFunctionsPosix.cpp
--- a/Engine/Source/Programs/UnrealBuildAccelerator/Detours/Private/PosixOS/UbaDetoursFunctionsPosix.cpp
+++ b/Engine/Source/Programs/UnrealBuildAccelerator/Detours/Private/PosixOS/UbaDetoursFunctionsPosix.cpp
@@ -34,8 +34,6 @@
 
 namespace uba
 {
-	enum LogEntryType : u8 { LogEntryType_Error = 0, LogEntryType_Warning = 1, };
-	class Logger { public: void Logf(...) {} void Info(...) {} void Warning(...) {} void Error(...) {} };
 	#include "UbaFileMapping.inl"
 
 	constexpr bool g_logToScreen = false;
@@ -2267,8 +2265,13 @@ int shared_posix_spawn(pid_t* pid, const char* path, const posix_spawn_file_acti
 		
 		rulesStr.Append("UBA_RULES=").AppendValue(reader.ReadU32());
 
-		u32 dllNameSize = reader.ReadU32();
-		reader.Skip(dllNameSize);
+		while (true)
+		{
+			u32 dllNameSize = reader.ReadU32();
+			if (!dllNameSize)
+				break;
+			reader.Skip(dllNameSize);
+		}
 
 		currentDir.Append("UBA_CWD=");
 		reader.ReadString(currentDir);
diff --git a/Engine/Source/Programs/UnrealBuildAccelerator/Detours/Private/UbaDetoursFileMappingTable.cpp b/Engine/Source/Programs/UnrealBuildAccelerator/Detours/Private/UbaDetoursFileMappingTable.cpp
--- a/Engine/Source/Programs/UnrealBuildAccelerator/Detours/Private/UbaDetoursFileMappingTable.cpp
+++ b/Engine/Source/Programs/UnrealBuildAccelerator/Detours/Private/UbaDetoursFileMappingTable.cpp
@@ -329,7 +329,8 @@ namespace uba
 		u64 committed = size;
 		if (!committed)
 			committed = reader.Read7BitEncoded();
-		view.Init(allocatorHandle, allocatorMem, reader, committed, hint, type);
+		bool res = view.Init(allocatorHandle, allocatorMem, reader, committed, hint, type);
+		UBA_ASSERT(res);(void)res;
 		//UBA_ASSERTF(u64(view.GetMemory()) != 0xFFFFFFFFFFFFFFFFllu, TC("%llu %llu %llu"), offset, size, committed);
 		//view.AddRequestedMemory(allocatorHandle, reader, SharedMemoryMapType_ReadOnly);
 	}
diff --git a/Engine/Source/Programs/UnrealBuildAccelerator/Detours/Private/UbaDetoursShared.h b/Engine/Source/Programs/UnrealBuildAccelerator/Detours/Private/UbaDetoursShared.h
--- a/Engine/Source/Programs/UnrealBuildAccelerator/Detours/Private/UbaDetoursShared.h
+++ b/Engine/Source/Programs/UnrealBuildAccelerator/Detours/Private/UbaDetoursShared.h
@@ -154,4 +154,8 @@ namespace uba
 		TimerScope ts(g_stats.timerName); \
 		BinaryWriter writer(ProcessCommunication_Write); \
 		writer.WriteByte(MessageType_##messageName);
+
+	// Dummy code
+	enum LogEntryType : u8 { LogEntryType_Error = 0, LogEntryType_Warning = 1, };
+	class Logger { public: void Logf(...) {} void Info(...) {} bool Warning(...) { return false; } bool Error(...) { return false; } };
 }
diff --git a/Engine/Source/Programs/UnrealBuildAccelerator/Detours/Private/Windows/UbaDetoursFunctionsKernelBase.inl b/Engine/Source/Programs/UnrealBuildAccelerator/Detours/Private/Windows/UbaDetoursFunctionsKernelBase.inl
--- a/Engine/Source/Programs/UnrealBuildAccelerator/Detours/Private/Windows/UbaDetoursFunctionsKernelBase.inl
+++ b/Engine/Source/Programs/UnrealBuildAccelerator/Detours/Private/Windows/UbaDetoursFunctionsKernelBase.inl
@@ -2288,14 +2288,16 @@ LPVOID Detoured_MapViewOfFileEx(HANDLE hFileMappingObject, DWORD dwDesiredAccess
 				// This is where cl.exe "continues" its compile from the pch
 				UBA_ASSERT(dwDesiredAccess == FILE_MAP_COPY);
 
-				SharedMemoryView& view = g_specialViews.emplace_back();
+				SharedMemoryView& view = g_specialViews.emplace_back(g_sharedMemoryFileMappingBackend);
 				RPC_MESSAGE(GetSharedMemory, getSharedMemory)
 				writer.WriteU64(fi.trueFileMapHandle.ToU64());
 				writer.WriteU64(offset + fi.trueFileMapOffset);
 				writer.WriteU64(dwNumberOfBytesToMap);
 				BinaryReader reader = writer.Flush();
 				view.Init(g_sharedMemoryAllocator, dwNumberOfBytesToMap, lpBaseAddress);
-				view.AddRequestedMemory(g_sharedMemoryAllocator, reader, SharedMemoryMapType_CopyOnWrite, 0);
+				pcs.Leave();
+				UBA_FOR_ASSERT(bool success =) view.AddRequestedMemory(g_sharedMemoryAllocator, reader, SharedMemoryMapType_CopyOnWrite, 0);
+				UBA_ASSERT(success);
 				SetLastError(ERROR_SUCCESS);
 				//UBA_ASSERT(!fi.trueFileMapOffset);
 				void* res = view.GetMemory();
@@ -3006,7 +3008,10 @@ BOOL Detoured_CreateProcessW(LPCWSTR lpApplicationName, LPWSTR lpCommandLine, LP
 
 	TString currentDir;
 	u32 processId = 0;
-	char dll[1024];
+	char dllBuffer[1024];
+	char* dllIt = dllBuffer;
+	u32 dllCount = 0;
+	LPCSTR dlls[2];
 	{
 		RPC_MESSAGE(CreateProcess, createProcess)
 		writer.WriteString(application);
@@ -3026,9 +3031,16 @@ BOOL Detoured_CreateProcessW(LPCWSTR lpApplicationName, LPWSTR lpCommandLine, LP
 
 		reader.Skip(sizeof(u32)); // Rules index
 
-		u32 dllNameSize = reader.ReadU32();
-		reader.ReadBytes(dll, dllNameSize);
-		dll[dllNameSize] = 0;
+		while (true)
+		{
+			u32 dllNameSize = reader.ReadU32();
+			if (!dllNameSize)
+				break;
+			dlls[dllCount++] = dllIt;
+			reader.ReadBytes(dllIt, dllNameSize);
+			dllIt += dllNameSize;
+			*dllIt++ = 0;
+		}
 
 		currentDir = reader.ReadString();
 		reader.ReadString(application.Clear());
@@ -3216,12 +3228,10 @@ BOOL Detoured_CreateProcessW(LPCWSTR lpApplicationName, LPWSTR lpCommandLine, LP
 
 	++t_disallowDetour;
 
-	LPCSTR dlls[] = { dll };
-
 	while (true)
 	{
 		res = true;
-		if (DetourCreateProcessWithDllsW(application.data, finalCommandLine, NULL, NULL, bInheritHandles, dwCreationFlags, lpEnvironment, currentDir.c_str(), lpStartupInfo, lpProcessInformation, sizeof_array(dlls), dlls, True_CreateProcessW))
+		if (DetourCreateProcessWithDllsW(application.data, finalCommandLine, NULL, NULL, bInheritHandles, dwCreationFlags, lpEnvironment, currentDir.c_str(), lpStartupInfo, lpProcessInformation, dllCount, dlls, True_CreateProcessW))
 			break;
 		res = false;
 		lastError = GetLastError();
diff --git a/Engine/Source/Programs/UnrealBuildAccelerator/Detours/Private/Windows/UbaDetoursFunctionsNtDll.inl b/Engine/Source/Programs/UnrealBuildAccelerator/Detours/Private/Windows/UbaDetoursFunctionsNtDll.inl
--- a/Engine/Source/Programs/UnrealBuildAccelerator/Detours/Private/Windows/UbaDetoursFunctionsNtDll.inl
+++ b/Engine/Source/Programs/UnrealBuildAccelerator/Detours/Private/Windows/UbaDetoursFunctionsNtDll.inl
@@ -1473,7 +1473,7 @@ NTSTATUS NTAPI Shared_NtCreateFile(bool IsCreateFunc, PHANDLE hFileHandle, ACCES
 		if (!g_directoryTable.GetLatestOffset(dirTableOffset2, dir))
 		{
 			IoStatusBlock->Status = STATUS_OBJECT_PATH_NOT_FOUND;
-			IoStatusBlock->Information = 0;
+			IoStatusBlock->Information = FILE_DOES_NOT_EXIST;
 			DEBUG_LOG_DETOURED(funcName, L"(DIRTABLE) %llu, (%s) -> STATUS_OBJECT_PATH_NOT_FOUND", uintptr_t(*hFileHandle), fileName.data);
 			return STATUS_OBJECT_PATH_NOT_FOUND;
 		}
diff --git a/Engine/Source/Programs/UnrealBuildAccelerator/Detours/Private/Windows/UbaDetoursFunctionsWin.cpp b/Engine/Source/Programs/UnrealBuildAccelerator/Detours/Private/Windows/UbaDetoursFunctionsWin.cpp
--- a/Engine/Source/Programs/UnrealBuildAccelerator/Detours/Private/Windows/UbaDetoursFunctionsWin.cpp
+++ b/Engine/Source/Programs/UnrealBuildAccelerator/Detours/Private/Windows/UbaDetoursFunctionsWin.cpp
@@ -63,6 +63,7 @@ namespace uba
 	extern HANDLE g_hostProcess;
 	SharedMemoryAllocatorHandle g_sharedMemoryAllocator;
 	u8* g_sharedMemory;
+	FileMappingBackend g_sharedMemoryFileMappingBackend;
 
 	SharedMemoryHandle g_permanentFilesHandle;
 	u8* g_permanentFilesMemory;
@@ -260,13 +261,15 @@ const wchar_t* HandleToName(const DetouredHandle& dh)
 BlockAllocator<FileObject>& g_fileObjectAllocator = *new BlockAllocator<FileObject>(g_memoryBlock);
 
 MemoryFile::MemoryFile(u8* data, bool localOnly)
-:	baseAddress(data)
+:	memoryView(g_sharedMemoryFileMappingBackend)
+,	baseAddress(data)
 ,	isLocalOnly(localOnly)
 {
 }
 
 MemoryFile::MemoryFile(bool localOnly, u64 reserveSize_, bool isThrowAway_, u32 commitStepSize_, const tchar* fileName)
 :	commitStepSize(commitStepSize_)
+,	memoryView(g_sharedMemoryFileMappingBackend)
 ,	isLocalOnly(localOnly)
 ,	isThrowAway(isThrowAway_)
 {
@@ -474,7 +477,7 @@ bool EnsureMapped(DetouredHandle& handle, DWORD dwFileOffsetHigh = 0, DWORD dwFi
 	u64 alignedOffsetStart = AlignUp(offset - (PageSize - 1), PageSize);
 	u64 mapSize = (offset - alignedOffsetStart) + numberOfBytesToMap;
 
-	SharedMemoryView view;
+	SharedMemoryView view(g_sharedMemoryFileMappingBackend);
 	Rpc_GetSharedMemory(view, g_sharedMemoryAllocator, g_sharedMemory, info.trueFileMapHandle, alignedOffsetStart, mapSize, info.originalName, type);
 
 	info.fileMapMem = view.DetachMemory() + (offset - alignedOffsetStart);
@@ -1151,7 +1154,19 @@ void Init(const DetoursPayload& payload, u64 startTime)
 	}
 	#endif
 
-	g_sharedMemoryAllocator = FileMapping_DuplicateFromHost(sharedMemoryAllocatorHandle, "SharedMem");
+	if (g_isRunningWine)
+	{
+		// TODO: This should probably check a variable sent through payload instead
+		if (HMODULE m = GetModuleHandleW(L"UbaWine.dll"))
+		{
+			Logger logger;
+			bool res = FileMapping_GetWineBackend(g_sharedMemoryFileMappingBackend, logger, m);
+			UBA_ASSERT(res);(void)res;
+		}
+	}
+
+	g_sharedMemoryAllocator = g_sharedMemoryFileMappingBackend.DuplicateFromHost(sharedMemoryAllocatorHandle, "SharedMem");
+	UBA_ASSERT(g_sharedMemoryAllocator.IsValid());
 
 	TimerScope ts1(g_stats.fileTable);
 	u8* mappedFileTableMem = FileMapping_MapFromHost(FileMapping_DuplicateFromHost(mappedFileTableHandle, "FileTable"), FileMappingTableMaxSize, 0, false, "FileTable");
@@ -1167,7 +1182,7 @@ void Init(const DetoursPayload& payload, u64 startTime)
 
 	if (permanentFilesHandle.IsValid())
 	{
-		SharedMemoryView permanentFilesView;
+		SharedMemoryView permanentFilesView(g_sharedMemoryFileMappingBackend);
 		Rpc_GetSharedMemory(permanentFilesView, g_sharedMemoryAllocator, nullptr, permanentFilesHandle, 0, 0, TC("PermanentFiles"));
 		g_permanentFilesHandle = permanentFilesHandle;
 		g_permanentFilesMemory = permanentFilesView.DetachMemory();
diff --git a/Engine/Source/Programs/UnrealBuildAccelerator/Detours/Private/Windows/UbaDetoursMainWin.cpp b/Engine/Source/Programs/UnrealBuildAccelerator/Detours/Private/Windows/UbaDetoursMainWin.cpp
--- a/Engine/Source/Programs/UnrealBuildAccelerator/Detours/Private/Windows/UbaDetoursMainWin.cpp
+++ b/Engine/Source/Programs/UnrealBuildAccelerator/Detours/Private/Windows/UbaDetoursMainWin.cpp
@@ -18,8 +18,6 @@
 
 namespace uba
 {
-	enum LogEntryType : u8 { LogEntryType_Error = 0, LogEntryType_Warning = 1, };
-	class Logger { public: void Logf(...) {} void Info(...) {} void Warning(...) {} void Error(...) {} };
 	struct Bottleneck { Bottleneck(...) {} };
 	struct BottleneckScope { BottleneckScope(Bottleneck&, Timer&) {} };
 }
diff --git a/Engine/Source/Programs/UnrealBuildAccelerator/Test/Private/UbaTestAll.h b/Engine/Source/Programs/UnrealBuildAccelerator/Test/Private/UbaTestAll.h
--- a/Engine/Source/Programs/UnrealBuildAccelerator/Test/Private/UbaTestAll.h
+++ b/Engine/Source/Programs/UnrealBuildAccelerator/Test/Private/UbaTestAll.h
@@ -30,6 +30,7 @@ namespace uba
 		UBA_TEST(TestParseArguments) \
 		UBA_TEST(TestBinaryWriter) \
 		UBA_TEST(TestSharedMemory) \
+		UBA_TEST(TestFileMapping) \
 		UBA_TEST(TestFileMappingBuffer) \
 		UBA_TEST(TestCrypto) \
 		UBA_TEST(TestSockets) \
diff --git a/Engine/Source/Programs/UnrealBuildAccelerator/Test/Private/UbaTestFileMappings.cpp b/Engine/Source/Programs/UnrealBuildAccelerator/Test/Private/UbaTestFileMappings.cpp
--- a/Engine/Source/Programs/UnrealBuildAccelerator/Test/Private/UbaTestFileMappings.cpp
+++ b/Engine/Source/Programs/UnrealBuildAccelerator/Test/Private/UbaTestFileMappings.cpp
@@ -1,16 +1,45 @@
 // Copyright Epic Games, Inc. All Rights Reserved.
 
 #include "UbaLogger.h"
+#include "UbaFileMapping.h"
 #include "UbaFileMappingBuffer.h"
 #include "UbaSharedMemoryAllocator.h"
 #include "UbaTest.h"
 
 namespace uba
 {
-	bool TestFileMappingBuffer(LoggerWithWriter& logger, const StringBufferBase& rootDir)
+	bool TestFileMappingInner(LoggerWithWriter& logger, const StringBufferBase& rootDir, FileMappingBackend& backend)
+	{
+		FileMappingHandle fmh = backend.Create(logger, PAGE_READWRITE|SEC_RESERVE, 1024*1024, nullptr, TC("Test"));
+		CHECK_TRUE(fmh.IsValid());
+		u8* memory = backend.Map(logger, fmh, FILE_MAP_READ|FILE_MAP_WRITE, 0, 64*1024, false);
+		CHECK_TRUE(memory);
+		CHECK_TRUE(backend.Unmap(logger, memory, 64*1024, TC(""), false));
+		CHECK_TRUE(backend.Close(logger, fmh, TC("")));
+		return true;
+	}
+
+	bool TestFileMapping(LoggerWithWriter& logger, const StringBufferBase& rootDir)
+	{
+		FileMappingBackend backend;
+		CHECK_TRUE(TestFileMappingInner(logger, rootDir, backend));
+
+		#if PLATFORM_WINDOWS
+		if (!IsRunningWine())
+			return true;
+		HMODULE m = LoadLibraryW(L"UbaWine.dll.so");
+		CHECK_TRUEF(m, TC("Failed to load UbaWine module"));
+		FileMappingBackend wineBackend;
+		CHECK_TRUE(FileMapping_GetWineBackend(wineBackend, logger, m));
+		CHECK_TRUE(TestFileMappingInner(logger, rootDir, wineBackend));
+		#endif
+		return true;
+	}
+
+	bool TestFileMappingBufferInner(LoggerWithWriter& logger, const StringBufferBase& rootDir, FileMappingBackend& backend)
 	{
 		TrackSystemUsage tsu(logger);
-		SharedMemoryAllocator allocator(logger);
+		SharedMemoryAllocator allocator(logger, backend);
 		allocator.Init(64ull*1024*1024);
 		FileMappingBuffer mappingBuffer(logger, allocator);
 
@@ -54,4 +83,21 @@ namespace uba
 		}
 		return true;
 	}
+
+	bool TestFileMappingBuffer(LoggerWithWriter& logger, const StringBufferBase& rootDir)
+	{
+		FileMappingBackend backend;
+		CHECK_TRUE(TestFileMappingBufferInner(logger, rootDir, backend));
+
+#if PLATFORM_WINDOWS
+		if (!IsRunningWine())
+			return true;
+		HMODULE m = LoadLibraryW(L"UbaWine.dll.so");
+		CHECK_TRUEF(m, TC("Failed to load UbaWine module"));
+		FileMappingBackend wineBackend;
+		CHECK_TRUE(FileMapping_GetWineBackend(wineBackend, logger, m));
+		CHECK_TRUE(TestFileMappingBufferInner(logger, rootDir, wineBackend));
+#endif
+		return true;
+	}
 }
diff --git a/Engine/Source/Programs/UnrealBuildAccelerator/Test/Private/UbaTestSharedMemory.cpp b/Engine/Source/Programs/UnrealBuildAccelerator/Test/Private/UbaTestSharedMemory.cpp
--- a/Engine/Source/Programs/UnrealBuildAccelerator/Test/Private/UbaTestSharedMemory.cpp
+++ b/Engine/Source/Programs/UnrealBuildAccelerator/Test/Private/UbaTestSharedMemory.cpp
@@ -2,6 +2,7 @@
 
 #include "UbaLogger.h"
 #include "UbaBinaryReaderWriter.h"
+#include "UbaFileMapping.h"
 #include "UbaSharedMemoryAllocator.h"
 #include "UbaSharedMemoryView.h"
 #include "UbaTest.h"
@@ -11,17 +12,18 @@ namespace uba
 	bool TestSharedMemory(LoggerWithWriter& logger, const StringBufferBase& rootDir)
 	{
 		TrackSystemUsage tsu(logger);
+		FileMappingBackend backend;
 
 		{
-			SharedMemoryAllocator mem(logger);
+			SharedMemoryAllocator mem(logger, backend);
 			mem.Init(64ull * 1024 * 1024 * 1024);
 			SharedMemoryAllocatorHandle handle = mem;
 
 			SharedMemoryHandle viewHandle1 = mem.CreateHandle(logger, TC(""));
 			SharedMemoryHandle viewHandle2 = mem.CreateHandle(logger, TC(""));
 
-			SharedMemoryView view1;
-			SharedMemoryView view2;
+			SharedMemoryView view1(backend);
+			SharedMemoryView view2(backend);
 
 			view1.Init(handle, 1024 * 1024);
 			view2.Init(handle, 1024 * 1024);
@@ -45,7 +47,7 @@ namespace uba
 			u8* ptr = mem.MapView(viewHandle1, TC(""), SharedMemoryMapType_ReadOnly);
 			mem.UnmapView(viewHandle1, ptr);
 
-			SharedMemoryView view3;
+			SharedMemoryView view3(backend);
 			mem.MapView(view3, viewHandle1, TC(""), SharedMemoryMapType_ReadOnly);
 
 			mem.CloseHandle(logger, viewHandle2, TC(""));
@@ -53,14 +55,14 @@ namespace uba
 
 
 			SharedMemoryHandle viewHandle4 = mem.CreateHandle(logger, TC(""));
-			SharedMemoryView view4;
+			SharedMemoryView view4(backend);
 			view4.Init(handle, 1024 * 1024);
 			mem.ExtendMemory(view4, viewHandle4, PageSize, TC(""), false);
 			mem.ExtendMemory(view4, viewHandle4, PageSize, TC(""), false);
 		}
 
 		{
-			SharedMemoryAllocator mem(logger);
+			SharedMemoryAllocator mem(logger, backend);
 			mem.Init(1024 * 1024);
 			SharedMemoryHandle viewHandle1 = mem.CreateHandle(logger, TC(""));
 			SharedMemoryHandle viewHandle2 = mem.CreateHandle(logger, TC(""));
diff --git a/Engine/Source/Programs/UnrealBuildAccelerator/Wine/UbaWineLinux.cpp b/Engine/Source/Programs/UnrealBuildAccelerator/Wine/UbaWineLinux.cpp
--- a/Engine/Source/Programs/UnrealBuildAccelerator/Wine/UbaWineLinux.cpp
+++ b/Engine/Source/Programs/UnrealBuildAccelerator/Wine/UbaWineLinux.cpp
@@ -49,6 +49,7 @@ int UnixGetTcpInfo(int fd, void* buf, int* len)
 
 #define UBA_ASSERT(...)
 #define UBA_ASSERTF(...) 
+#define UBA_FM_STRING_TYPE wchar_t
 
 #include "../Core/Public/UbaFileHandle.h"
 #include "../Core/Public/UbaFileMappingHandle.h"
@@ -56,7 +57,7 @@ int UnixGetTcpInfo(int fd, void* buf, int* len)
 
 namespace uba
 {
-	struct Logger { void Logf(...) {} void Info(...) {} void Warning(...) {} void Error(...) {} };
+	class Logger { public: void Logf(...) {} void Info(...) {} bool Warning(...) { return false; } bool Error(...) { return false; } };
 	template<typename Type> using Atomic = std::atomic<Type>;
 	void SetLastError(int err) {}
 	enum LogEntryType : u8 { LogEntryType_Error = 0, LogEntryType_Warning = 1, };
diff --git a/Engine/Source/Programs/UnrealBuildAccelerator/Wine/UbaWineWindows.cpp b/Engine/Source/Programs/UnrealBuildAccelerator/Wine/UbaWineWindows.cpp
--- a/Engine/Source/Programs/UnrealBuildAccelerator/Wine/UbaWineWindows.cpp
+++ b/Engine/Source/Programs/UnrealBuildAccelerator/Wine/UbaWineWindows.cpp
@@ -163,3 +163,56 @@ extern "C" int WINAPI SocketGetSendBuf(uba::Socket s, int& outSize)
 {
 	return uba::SocketGetSendBuf(s, outSize);
 }
+
+using FileMappingHandle = uba::FileMappingHandle;
+using tchar = uba::tchar;
+using u8 = uba::u8;
+using MemoryMapType = uba::MemoryMapType;
+using u32 = uba::u32;
+using u64 = uba::u64;
+using FileHandle = uba::FileHandle;
+
+extern "C" FileMappingHandle WINAPI FileMapping_Create(uba::Logger& logger, u32 flProtect, u64 maxSize, const tchar* name, const tchar* hint)
+{
+	return uba::FileMapping_Create(logger, flProtect, maxSize, name, hint);
+}
+extern "C" FileMappingHandle WINAPI FileMapping_CreateFromFile(uba::Logger& logger, FileHandle file, u32 flProtect, u64 maxSize, const tchar* hint)
+{
+	return uba::FileMapping_CreateFromFile(logger, file, flProtect, maxSize, hint);
+}
+extern "C" u8* WINAPI FileMapping_Map(uba::Logger& logger, FileMappingHandle fileMappingObject, u32 dwDesiredAccess, u64 offset, u64 dwNumberOfBytesToMap, bool allowDiscard)
+{
+	return uba::FileMapping_Map(logger, fileMappingObject, dwDesiredAccess, offset, dwNumberOfBytesToMap, allowDiscard);
+}
+extern "C" bool WINAPI FileMapping_Commit(uba::Logger& logger, void* address, u64 size, bool allowDiscard)
+{
+	return uba::FileMapping_Commit(logger, address, size, allowDiscard);
+}
+extern "C" bool WINAPI FileMapping_Unmap(uba::Logger& logger, const void* lpBaseAddress, u64 bytesToUnmap, const tchar* hint, bool allowDiscard)
+{
+	return uba::FileMapping_Unmap(logger, lpBaseAddress, bytesToUnmap, hint, allowDiscard);
+}
+extern "C" bool WINAPI FileMapping_Close(uba::Logger& logger, FileMappingHandle h, const tchar* hint)
+{
+	return uba::FileMapping_Close(logger, h, hint);
+}
+extern "C" void* WINAPI FileMapping_ReservePlaceholder(void* baseAddress, u64 capacity)
+{
+	return uba::FileMapping_ReservePlaceholder(baseAddress, capacity);
+}
+extern "C" void* WINAPI FileMapping_MapPlaceholder(FileMappingHandle handle, MemoryMapType mapType, void* targetAddress, u64 targetOffset, u64 targetCapacity, u64 handleOffset, u64 size)
+{
+	return uba::FileMapping_MapPlaceholder(handle, mapType, targetAddress, targetOffset, targetCapacity, handleOffset, size);
+}
+extern "C" bool WINAPI FileMapping_UnmapPlaceholder(void* memory, u64 capacity, u64 mappedSize, u64* subMappings, u64 subMappingsCount)
+{
+	return uba::FileMapping_UnmapPlaceholder(memory, capacity, mappedSize, subMappings, subMappingsCount);
+}
+extern "C" FileMappingHandle WINAPI FileMapping_DuplicateFromHost(FileMappingHandle handle, const char* hint)
+{
+	return uba::FileMapping_DuplicateFromHost(handle, hint);
+}
+extern "C" u8* WINAPI FileMapping_MapFromHost(FileMappingHandle handle, u64 size, u64 offset, bool writable, const char* hint)
+{
+	return uba::FileMapping_MapFromHost(handle, size, offset, writable, hint);
+}
diff --git a/Engine/Source/Programs/UnrealBuildAccelerator/Wine/UbaWineWindows.def b/Engine/Source/Programs/UnrealBuildAccelerator/Wine/UbaWineWindows.def
--- a/Engine/Source/Programs/UnrealBuildAccelerator/Wine/UbaWineWindows.def
+++ b/Engine/Source/Programs/UnrealBuildAccelerator/Wine/UbaWineWindows.def
@@ -29,3 +29,14 @@ EXPORTS
 	SocketGetCongestionAlgorithm
 	SocketGetRecvBuf
 	SocketGetSendBuf
+	FileMapping_Create
+	FileMapping_CreateFromFile
+	FileMapping_Map
+	FileMapping_Commit
+	FileMapping_Unmap
+	FileMapping_Close
+	FileMapping_ReservePlaceholder
+	FileMapping_MapPlaceholder
+	FileMapping_UnmapPlaceholder
+	FileMapping_DuplicateFromHost
+	FileMapping_MapFromHost
