[UBA]
* Removed debug code not needed anymore

diff --git a/Engine/Source/Programs/UnrealBuildAccelerator/Core/Public/UbaFileMapping.inl b/Engine/Source/Programs/UnrealBuildAccelerator/Core/Public/UbaFileMapping.inl
--- a/Engine/Source/Programs/UnrealBuildAccelerator/Core/Public/UbaFileMapping.inl
+++ b/Engine/Source/Programs/UnrealBuildAccelerator/Core/Public/UbaFileMapping.inl
@@ -426,7 +426,6 @@ void* FileMapping_ReservePlaceholder(void* baseAddress, u64 capacity)
 	void* memory = mmap(baseAddress, capacity, PROT_NONE, flags, -1, 0);
 	if (memory == MAP_FAILED)
 		memory = nullptr;
-	//WRITE_DEBUG("RESERVE: Base: 0x%llx Capacity: %llu -> 0x%llx", u64(baseAddress), capacity, u64(memory));
 #endif
 	UBA_ASSERTF(memory, TC("VirtualAlloc2 failed (%s)"), LastErrorToText().data);
 	return memory;
@@ -462,7 +461,6 @@ void* FileMapping_MapPlaceholder(FileMappingHandle handle, MemoryMapType mapType
 
 	if (mapType == MemoryMapType_CopyOnWrite)
 	{
-		//WRITE_DEBUG("HEEERE: %llu", size);
 		mapFlags |= MAP_PRIVATE;
 		protect |= PROT_WRITE;
 	}
@@ -476,13 +474,6 @@ void* FileMapping_MapPlaceholder(FileMappingHandle handle, MemoryMapType mapType
 	void* viewA = mmap((u8*)targetAddress+targetOffset, size, protect, mapFlags, handle.shmFd, handleOffset);
 	if (viewA == MAP_FAILED)
 		viewA = nullptr;
-	//else if (mapType == MemoryMapType_CopyOnWrite)
-	//{
-	//	struct stat st;
-	//	fstat(handle.shmFd, &st);
-	//	WRITE_DEBUG("shmFd=%i size=%llu handleOffset=%llu mapSize=%llu -> 0x%llx", handle.shmFd, (u64)st.st_size, handleOffset, size, u64(viewA));
-	//}
-	//WRITE_DEBUG("mmap: Base: %llu Protect: %i Offset: %llu Size: %llu HandleOffset: %llu: SHM: %i -> 0x%llx (%s)", u64(targetAddress), protect, targetOffset, size, handleOffset, handle.shmFd, u64(viewA), viewA ? "success" : strerror(errno));
 #endif
 	return viewA;
 }
