[UBA]
* Fixed linux compile errors

diff --git a/Engine/Source/Developer/UbaCoordinatorHorde/Private/UbaCoordinatorHorde.cpp b/Engine/Source/Developer/UbaCoordinatorHorde/Private/UbaCoordinatorHorde.cpp
--- a/Engine/Source/Developer/UbaCoordinatorHorde/Private/UbaCoordinatorHorde.cpp
+++ b/Engine/Source/Developer/UbaCoordinatorHorde/Private/UbaCoordinatorHorde.cpp
@@ -49,7 +49,11 @@ public:
 
 	virtual void SetUpdateStatusCallback(UpdateStatusCallback* callback, void* userData) override
 	{
+		#if PLATFORM_WINDOWS
 		m_manager.SetUpdateStatusCallback(callback, userData);
+		#else
+		// TODO: Callback is char*  while unreal is wchar_t*
+		#endif
 	}
 
 	FUbaHordeAgentManager m_manager;
diff --git a/Engine/Source/Programs/UnrealBuildAccelerator/Core/Public/UbaFileMapping.h b/Engine/Source/Programs/UnrealBuildAccelerator/Core/Public/UbaFileMapping.h
--- a/Engine/Source/Programs/UnrealBuildAccelerator/Core/Public/UbaFileMapping.h
+++ b/Engine/Source/Programs/UnrealBuildAccelerator/Core/Public/UbaFileMapping.h
@@ -62,7 +62,8 @@ namespace uba
 		UBA_FM_FUNC(UnmapPlaceholder) \
 		UBA_FM_FUNC(DuplicateFromHost) \
 		UBA_FM_FUNC(MapFromHost) \
-		//UBA_FM_FUNC(CopyMem) \
+
+		//UBA_FM_FUNC(CopyMem)
 
 	struct FileMappingBackend
 	{
diff --git a/Engine/Source/Programs/UnrealBuildAccelerator/Test/Private/UbaTestBasics.cpp b/Engine/Source/Programs/UnrealBuildAccelerator/Test/Private/UbaTestBasics.cpp
--- a/Engine/Source/Programs/UnrealBuildAccelerator/Test/Private/UbaTestBasics.cpp
+++ b/Engine/Source/Programs/UnrealBuildAccelerator/Test/Private/UbaTestBasics.cpp
@@ -130,7 +130,8 @@ namespace uba
 		#if !PLATFORM_WINDOWS
 		if (IsRunningDarling()) // bug in darling, this does not work
 			return true;
-		SharedMemoryAllocator allocator(logger);
+		FileMappingBackend backend;
+		SharedMemoryAllocator allocator(logger, backend);
 		allocator.Init(1024 * 1024);
 		FileMappingBuffer mappingBuffer(logger, allocator);
 		mappingBuffer.Init(TC(""), 64*1024, 64 * 1024);
