[UBA]
* Added option to write out how many entries that are in scheduler's queues

diff --git a/Engine/Source/Programs/UnrealBuildAccelerator/Common/Private/UbaScheduler.cpp b/Engine/Source/Programs/UnrealBuildAccelerator/Common/Private/UbaScheduler.cpp
--- a/Engine/Source/Programs/UnrealBuildAccelerator/Common/Private/UbaScheduler.cpp
+++ b/Engine/Source/Programs/UnrealBuildAccelerator/Common/Private/UbaScheduler.cpp
@@ -19,15 +19,14 @@
 #include "UbaSessionServer.h"
 #include "UbaStringBuffer.h"
 
-// TODO. This is how we want scheduler to work.
-// 
-// 1. Cache should be queried from front of queue.
-//    We need to split up cache in fetching cache entry and fetching cas.
-//    .. if there are 10000 actions and the last is a cache miss, then we want to find this out as early as possible.
-// 2. Local processing should prioritize cache misses first, then query from the back of they queue
-// 
-// At some point cached processes and local will race which is fine.
-//
+#define UBA_SCHEDULER_PROCESS_QUEUES \
+			UBA_SCHEDULER_PROCESS_QUEUE(TestCache) \
+			UBA_SCHEDULER_PROCESS_QUEUE(DownloadCache) \
+			UBA_SCHEDULER_PROCESS_QUEUE(RunLocal) \
+			UBA_SCHEDULER_PROCESS_QUEUE(RunRemote) \
+			UBA_SCHEDULER_PROCESS_QUEUE(RunCrossArchitecture) \
+			UBA_SCHEDULER_PROCESS_QUEUE(RunCrossPlatform) \
+			UBA_SCHEDULER_PROCESS_QUEUE(Skip) \
 
 namespace uba
 {
@@ -168,6 +167,7 @@ namespace uba
 
 			table->GetValueAsBool(writeTrackedInputs, TC("WriteTrackedInputs"));
 			table->GetValueAsBool(deferRemoteCapableProcesses, TC("DeferRemoteCapableProcesses"));
+			table->GetValueAsBool(traceSchedulerQueues, TC("TraceSchedulerQueues"));
 
 			u32 temp = 0;
 			if (table->GetValueAsU32(temp, TC("MemTracingLevel")))
@@ -220,6 +220,7 @@ namespace uba
 	,	m_useThreadToExitRemoteProcess(info.useThreadToExitRemoteProcess)
 	,	m_writeTrackedInputs(info.writeTrackedInputs)
 	,	m_deferRemoteCapableProcesses(info.deferRemoteCapableProcesses)
+	,	m_traceSchedulerQueues(info.traceSchedulerQueues)
 	,	m_processConfigs(info.processConfigs)
 	,	m_writeToCache(info.writeToCache && info.cacheClientCount)
 	{
@@ -678,11 +679,35 @@ namespace uba
 
 	void Scheduler::ThreadLoop()
 	{
+		#define UBA_SCHEDULER_PROCESS_QUEUE(queue) u32 queue##Len = 0;
+		UBA_SCHEDULER_PROCESS_QUEUES
+		#undef UBA_SCHEDULER_PROCESS_QUEUE
+
+		if (m_traceSchedulerQueues)
+			m_session.GetTrace().StatusUpdate(2, 1, TCV("Queues"), LogEntryType_Info);
+
 		while (m_loop)
 		{
 			if (!m_updateThreadWakeup.IsSet())
 				break;
 			
+			if (m_traceSchedulerQueues)
+			{
+				bool changed = false;
+				SCOPED_FUTEX_READ(m_processEntriesLock, lock)
+				#define UBA_SCHEDULER_PROCESS_QUEUE(queue) u32 queue##LenNew = u32(m_readyTo##queue.size()); changed |= queue##LenNew != queue##Len; queue##Len = queue##LenNew;
+				UBA_SCHEDULER_PROCESS_QUEUES
+				#undef UBA_SCHEDULER_PROCESS_QUEUE
+				if (changed)
+				{
+					StringBuffer<> temp;
+					#define UBA_SCHEDULER_PROCESS_QUEUE(queue) temp.Append(TC("") #queue ": ").AppendValue(queue##Len).Append(TCV("  "));
+					UBA_SCHEDULER_PROCESS_QUEUES
+					#undef UBA_SCHEDULER_PROCESS_QUEUE
+					m_session.GetTrace().StatusUpdate(2, 6, temp, LogEntryType_Info);
+				}
+			}				
+
 			m_session.FlushDeadProcesses();
 
 			while (RunCacheQuery())
diff --git a/Engine/Source/Programs/UnrealBuildAccelerator/Common/Public/UbaScheduler.h b/Engine/Source/Programs/UnrealBuildAccelerator/Common/Public/UbaScheduler.h
--- a/Engine/Source/Programs/UnrealBuildAccelerator/Common/Public/UbaScheduler.h
+++ b/Engine/Source/Programs/UnrealBuildAccelerator/Common/Public/UbaScheduler.h
@@ -40,6 +40,7 @@ namespace uba
 		bool memWatchdog = false; // Will watch memory and throttle/kill processes if needed
 		bool writeTrackedInputs = false; // Set to true to write out log files with tracked inputs
 		bool deferRemoteCapableProcesses = false; // Set to true to have local session first pick local-only processes
+		bool traceSchedulerQueues = false; // Show in visualizer how many processes that are queued up for state change
 		u32 maxRacingPercent = 0u; // If processor weights go below this percent, then we will start filling up with racing processes
 		u8 memTracingLevel = 0;
 		u8 memStartWaitPercent = 85; // When memory usage goes above this percent, no new processes will be spawned until back below
@@ -226,6 +227,7 @@ namespace uba
 		bool m_allowDisableRemoteExecution = false;
 		bool m_writeTrackedInputs = false;
 		bool m_deferRemoteCapableProcesses = false;
+		bool m_traceSchedulerQueues = false;
 		bool m_cancelled = false;
 		bool m_hadRemoteRequest = false;
 		bool m_hadCrossSystemRequest = false;
