[UBA]
* Fixed bug in logic garbage collecting file mappings on posix

diff --git a/Engine/Source/Programs/UnrealBuildAccelerator/Common/Private/UbaDirectoryTableHolder.cpp b/Engine/Source/Programs/UnrealBuildAccelerator/Common/Private/UbaDirectoryTableHolder.cpp
--- a/Engine/Source/Programs/UnrealBuildAccelerator/Common/Private/UbaDirectoryTableHolder.cpp
+++ b/Engine/Source/Programs/UnrealBuildAccelerator/Common/Private/UbaDirectoryTableHolder.cpp
@@ -206,7 +206,7 @@ namespace uba
 		UBA_ASSERT(!overlay.handle.IsValid());
 		SCOPED_WRITE_LOCK(overlay.lock, lock);
 		overlay.handle = FileMapping_Create(m_logger, PAGE_READWRITE|SEC_RESERVE, OverlayTableMaxSize, nullptr, TC("Overlay"));
-		UBA_ASSERT(overlay.handle.IsValid());
+		UBA_ASSERTF(overlay.handle.IsValid(), TC("Failed to create overlay handle (size %llu)"), OverlayTableMaxSize);
 		overlay.memory = FileMapping_Map(m_logger, overlay.handle, FILE_MAP_WRITE, 0, OverlayTableMaxSize);
 		UBA_ASSERT(overlay.memory);
 	}
diff --git a/Engine/Source/Programs/UnrealBuildAccelerator/Core/Public/UbaFileMapping.inl b/Engine/Source/Programs/UnrealBuildAccelerator/Core/Public/UbaFileMapping.inl
--- a/Engine/Source/Programs/UnrealBuildAccelerator/Core/Public/UbaFileMapping.inl
+++ b/Engine/Source/Programs/UnrealBuildAccelerator/Core/Public/UbaFileMapping.inl
@@ -88,6 +88,20 @@ void GetMappingHandleName(char* out, u64 uid)
 	#endif
 	AppendPrefixAndValue(out, "/uba_", uid);
 }
+u64 ParseMappingUid(const char* str)
+{
+	u64 result = 0;
+	int shift = 0;
+	while (str[0] && str[1] && shift < 64)
+	{
+		auto hexVal = [](char c) -> u8 { return c >= 'a' ? c - 'a' + 10 : c - '0'; };
+		result |= (u64)((hexVal(str[0]) << 4) | hexVal(str[1])) << shift;
+		shift += 8;
+		str += 2;
+	}
+	return result;
+}
+
 bool GarbageCollectFileMappings(Logger& logger)
 {
 	// Since we need to not leak file mappings we use files as a trick to know which ones are used and not
@@ -124,7 +138,7 @@ bool GarbageCollectFileMappings(Logger& logger)
 			if (fileName[0] == '.' && (fileName[1] == 0 || (fileName[1] == '.' && fileName[2] == 0))) 
 				continue;
 
-			u32 uid = strtoul(fileName, nullptr, 10);
+			u32 uid = ParseMappingUid(fileName);
 
 			strcpy(filePath+pathLen, fileName);
 			int lockFd = open(filePath, O_RDWR, S_IRUSR | S_IWUSR | O_CLOEXEC);
