[UBA]
* Fixed more bugs in file mapping lock file logic and garbage collection

diff --git a/Engine/Source/Programs/UnrealBuildAccelerator/Core/Public/UbaFileMapping.inl b/Engine/Source/Programs/UnrealBuildAccelerator/Core/Public/UbaFileMapping.inl
--- a/Engine/Source/Programs/UnrealBuildAccelerator/Core/Public/UbaFileMapping.inl
+++ b/Engine/Source/Programs/UnrealBuildAccelerator/Core/Public/UbaFileMapping.inl
@@ -141,7 +141,7 @@ bool GarbageCollectFileMappings(Logger& logger)
 			u32 uid = ParseMappingUid(fileName);
 
 			strcpy(filePath+pathLen, fileName);
-			int lockFd = open(filePath, O_RDWR, S_IRUSR | S_IWUSR | O_CLOEXEC);
+			int lockFd = open(filePath, O_RDWR | O_CLOEXEC, S_IRUSR | S_IWUSR);
 			if (lockFd == -1)
 			{
 				if (errno == EPERM)
@@ -203,7 +203,7 @@ FileMappingHandle FileMapping_Create(Logger& logger, u32 flProtect, u64 maxSize,
 
 		AppendPrefixAndValue(lockFilePath, "/tmp/uba_shm_locks/", uid);
 
-		lockFd = open(lockFilePath, O_CREAT | O_RDWR | O_NOFOLLOW | O_EXCL, S_IRUSR | S_IWUSR | O_CLOEXEC);
+		lockFd = open(lockFilePath, O_CREAT | O_RDWR | O_NOFOLLOW | O_EXCL | O_CLOEXEC, S_IRUSR | S_IWUSR);
 		if (lockFd == -1)
 		{
 			if (errno == EEXIST)
