[UBA]
* More bug fixes for shm lock files being wrong after refactoring
* Fixed so child processes can properly get ctrl-c

diff --git a/Engine/Source/Programs/UnrealBuildAccelerator/Common/Private/UbaProcess.cpp b/Engine/Source/Programs/UnrealBuildAccelerator/Common/Private/UbaProcess.cpp
--- a/Engine/Source/Programs/UnrealBuildAccelerator/Common/Private/UbaProcess.cpp
+++ b/Engine/Source/Programs/UnrealBuildAccelerator/Common/Private/UbaProcess.cpp
@@ -2290,7 +2290,13 @@ namespace uba
 			auto attrGuard = MakeGuard([&]() { posix_spawnattr_destroy(&attr); });
 
 			// We set process group because we want to make sure that all processes get killed when ctrl-c is pressed
-			res = posix_spawnattr_setflags(&attr, POSIX_SPAWN_SETPGROUP);
+			// We also reset the signal mask because UBA threads block SIGINT, and posix_spawn inherits the calling
+			// thread's signal mask - meaning spawned children would never receive SIGINT from ctrl-c
+			sigset_t emptymask;
+			sigemptyset(&emptymask);
+			res = posix_spawnattr_setsigmask(&attr, &emptymask);
+			UBA_ASSERTF(res == 0, TC("posix_spawnattr_setsigmask (%s)"), strerror(errno));
+			res = posix_spawnattr_setflags(&attr, POSIX_SPAWN_SETPGROUP | POSIX_SPAWN_SETSIGMASK);
 			UBA_ASSERTF(res == 0, TC("posix_spawnattr_setflags (%s)"), strerror(errno));
 			res = posix_spawnattr_setpgroup(&attr, getpgrp());
 			UBA_ASSERTF(res == 0, TC("posix_spawnattr_setpgroup (%s)"), strerror(errno));
diff --git a/Engine/Source/Programs/UnrealBuildAccelerator/Core/Public/UbaFileMapping.inl b/Engine/Source/Programs/UnrealBuildAccelerator/Core/Public/UbaFileMapping.inl
--- a/Engine/Source/Programs/UnrealBuildAccelerator/Core/Public/UbaFileMapping.inl
+++ b/Engine/Source/Programs/UnrealBuildAccelerator/Core/Public/UbaFileMapping.inl
@@ -195,6 +195,7 @@ FileMappingHandle FileMapping_Create(Logger& logger, u32 flProtect, u64 maxSize,
 	int shmFd;
 	int lockFd;
 	char lockFilePath[128];
+	char uidName[64];
 
 	int retryCount = 4;
 	while (true)
@@ -219,7 +220,6 @@ FileMappingHandle FileMapping_Create(Logger& logger, u32 flProtect, u64 maxSize,
 			continue;
 		}
 
-		char uidName[64];
 		GetMappingHandleName(uidName, uid);
 
 		int oflags = O_CREAT | O_RDWR | O_NOFOLLOW | O_EXCL;
@@ -227,17 +227,30 @@ FileMappingHandle FileMapping_Create(Logger& logger, u32 flProtect, u64 maxSize,
 		if (shmFd != -1)
 			break;
 
+		int err = errno;
+		if (err == EEXIST)
+		{
+			// We hold the exclusive flock, so no other process owns this uid.
+			// The shm is an orphan from a previous run that lost its lock file
+			// (e.g. ftruncate failure). Unlink it and retry with a fresh shm.
+			shm_unlink(uidName);
+			shmFd = shm_open(uidName, oflags, S_IRUSR | S_IWUSR);
+			if (shmFd != -1)
+				break;
+			err = errno;
+		}
+
 		bool retry = retryCount-- > 0;
 
-		int err = errno;
 		auto logType = retry ? LogEntryType_Warning : LogEntryType_Error;
 		logger.Logf(logType, TC("Failed to create shm %s after getting lock-file %s (%s)"), uidName, lockFilePath, strerror(err));
 
+		remove(lockFilePath);
+		close(lockFd);
+
 		if (retry)
 			continue;
 
-		remove(lockFilePath);
-		close(lockFd);
 		SetLastError(err);
 		return {};
 	}
@@ -248,6 +261,7 @@ FileMappingHandle FileMapping_Create(Logger& logger, u32 flProtect, u64 maxSize,
 		{
 			SetLastError(errno);
 			close(shmFd);
+			shm_unlink(uidName);
 			remove(lockFilePath);
 			close(lockFd);
 			//UBA_ASSERTF(false, "Failed to truncate file mapping '%s' to size %llu (%s)", name, maxSize, strerror(errno));
