[UBA]
* Added more handling of when process handle is set in GetModuleFileNameExW/A

diff --git a/Engine/Source/Programs/UnrealBuildAccelerator/Detours/Private/Windows/UbaDetoursFunctionsKernelBase.inl b/Engine/Source/Programs/UnrealBuildAccelerator/Detours/Private/Windows/UbaDetoursFunctionsKernelBase.inl
--- a/Engine/Source/Programs/UnrealBuildAccelerator/Detours/Private/Windows/UbaDetoursFunctionsKernelBase.inl
+++ b/Engine/Source/Programs/UnrealBuildAccelerator/Detours/Private/Windows/UbaDetoursFunctionsKernelBase.inl
@@ -1038,10 +1038,17 @@ DWORD Detoured_GetModuleFileNameExW(HANDLE hProcess, HMODULE hModule, LPWSTR lpF
 {
 	if (hProcess != (HANDLE)-1)
 	{
-		UBA_ASSERT(!g_runningRemote); // Not implemented
-		DWORD res = True_GetModuleFileNameExW(hProcess, hModule, lpFilename, nSize);
-		DEBUG_LOG_TRUE(L"GetModuleFileNameExW", L"%llu  %u (%ls) -> %u", uintptr_t(hModule), nSize, lpFilename, res);
-		return res;
+		HANDLE trueHandle = hProcess;
+		if (isDetouredHandle(hProcess))
+			trueHandle = asDetouredHandle(hProcess).trueHandle;
+		if (GetProcessId(trueHandle) != GetCurrentProcessId())
+		{
+			DWORD res = True_GetModuleFileNameExW(trueHandle, hModule, lpFilename, nSize);
+			DEBUG_LOG_TRUE(L"GetModuleFileNameExW", L"%llu  %u (%ls) -> %u", uintptr_t(hModule), nSize, lpFilename, res);
+			UBA_ASSERTF(!isDetouredHandle(hProcess), TC("GetModuleFileNameExW getting module filename from detoured process, path is potentially wrong. Need implementation"));
+			UBA_ASSERTF(!g_runningRemote, TC("GetModuleFileNameExW on other processes not implemented (%s)"), lpFilename); // Not implemented
+			return res;
+		}
 	}
 	return Shared_GetModuleFileNameW(L"GetModuleFileNameExW", hModule, lpFilename, nSize);
 }
@@ -1077,10 +1084,17 @@ DWORD Detoured_GetModuleFileNameExA(HANDLE hProcess, HMODULE hModule, LPSTR lpFi
 	DETOURED_CALL(GetModuleFileNameExA);
 	if (hProcess != (HANDLE)-1)
 	{
-		UBA_ASSERT(!g_runningRemote); // Not implemented
-		DWORD res = True_GetModuleFileNameExA(hProcess, hModule, lpFilename, nSize);
-		DEBUG_LOG_TRUE(L"GetModuleFileNameExA", L"%llu  %u (%hs) -> %u", uintptr_t(hModule), nSize, lpFilename, res);
-		return res;
+		HANDLE trueHandle = hProcess;
+		if (isDetouredHandle(hProcess))
+			trueHandle = asDetouredHandle(hProcess).trueHandle;
+		if (GetProcessId(trueHandle) != GetCurrentProcessId())
+		{
+			DWORD res = True_GetModuleFileNameExA(trueHandle, hModule, lpFilename, nSize);
+			DEBUG_LOG_TRUE(L"GetModuleFileNameExA", L"%llu  %u (%hs) -> %u", uintptr_t(hModule), nSize, lpFilename, res);
+			UBA_ASSERTF(!isDetouredHandle(hProcess), TC("GetModuleFileNameExA getting module filename from detoured process, path is potentially wrong. Need implementation"));
+			UBA_ASSERTF(!g_runningRemote, TC("GetModuleFileNameExA on other processes not implemented (%hs)"), lpFilename); // Not implemented
+			return res;
+		}
 	}
 	return Shared_GetModuleFileNameA(L"GetModuleFileNameExA", hModule, lpFilename, nSize);
 }
