[UBA]
* Attempt at implementing Detoured_RtlDosPathNameToRelativeNtPathName_U_WithStatus when relativeName->ContainingDirectory is set. Based on bug report this can happen in the wild so hopefully this will fix that
* Fixed compile warning in logging when building debug

diff --git a/Engine/Source/Programs/UnrealBuildAccelerator/Common/Private/UbaDirectoryTableHolder.cpp b/Engine/Source/Programs/UnrealBuildAccelerator/Common/Private/UbaDirectoryTableHolder.cpp
--- a/Engine/Source/Programs/UnrealBuildAccelerator/Common/Private/UbaDirectoryTableHolder.cpp
+++ b/Engine/Source/Programs/UnrealBuildAccelerator/Common/Private/UbaDirectoryTableHolder.cpp
@@ -206,7 +206,7 @@ namespace uba
 		UBA_ASSERT(!overlay.handle.IsValid());
 		SCOPED_WRITE_LOCK(overlay.lock, lock);
 		overlay.handle = FileMapping_Create(m_logger, PAGE_READWRITE|SEC_RESERVE, OverlayTableMaxSize, nullptr, TC("Overlay"));
-		UBA_ASSERTF(overlay.handle.IsValid(), TC("Failed to create overlay handle (size %llu)"), OverlayTableMaxSize);
+		UBA_ASSERTF(overlay.handle.IsValid(), TC("Failed to create overlay handle (size %u)"), OverlayTableMaxSize);
 		overlay.memory = FileMapping_Map(m_logger, overlay.handle, FILE_MAP_WRITE, 0, OverlayTableMaxSize);
 		UBA_ASSERT(overlay.memory);
 	}
diff --git a/Engine/Source/Programs/UnrealBuildAccelerator/Detours/Private/Windows/UbaDetoursFunctionsNtDll.inl b/Engine/Source/Programs/UnrealBuildAccelerator/Detours/Private/Windows/UbaDetoursFunctionsNtDll.inl
--- a/Engine/Source/Programs/UnrealBuildAccelerator/Detours/Private/Windows/UbaDetoursFunctionsNtDll.inl
+++ b/Engine/Source/Programs/UnrealBuildAccelerator/Detours/Private/Windows/UbaDetoursFunctionsNtDll.inl
@@ -125,13 +125,60 @@ NTSTATUS Detoured_RtlDosPathNameToNtPathName_U_WithStatus(PCWSTR dosPath, PUNICO
 
 NTSTATUS Detoured_RtlDosPathNameToRelativeNtPathName_U_WithStatus(PCWSTR dosPath, PUNICODE_STRING ntPath, PWSTR* filePart, RTL_RELATIVE_NAME_U* relativeName)
 {
-	const tchar* original = dosPath;(void)original;
-	StringBuffer<> temp;
-	UBA_ASSERT(!relativeName || !relativeName->ContainingDirectory);
-	FixupDosPath(dosPath, temp);
-	auto res = True_RtlDosPathNameToRelativeNtPathName_U_WithStatus(dosPath, ntPath, filePart, relativeName);
-	//DEBUG_LOG(L"RtlDosPathNameToRelativeNtPathName_U_WithStatus:   %s -> %.*s", original, ntPath->Length/2, ntPath->Buffer);
-	return res;
+	//const tchar* original = dosPath;
+
+	auto trueRun = [&]()
+		{
+			StringBuffer<> temp;
+			FixupDosPath(dosPath, temp);
+			auto res = True_RtlDosPathNameToRelativeNtPathName_U_WithStatus(dosPath, ntPath, filePart, relativeName);
+			//DEBUG_LOG(L"RtlDosPathNameToRelativeNtPathName_U_WithStatus:   %s -> %.*s", original, ntPath->Length/2, ntPath->Buffer);
+			return res;
+		};
+
+	if (!relativeName || !relativeName->ContainingDirectory)
+		return trueRun();
+
+	HANDLE dirHandle = relativeName->ContainingDirectory;
+	const tchar* dirPath = nullptr;
+	if (isDetouredHandle(dirHandle))
+	{
+		DetouredHandle& dh = asDetouredHandle(dirHandle);
+		if (dh.fileObject && dh.fileObject->fileInfo)
+			dirPath = dh.fileObject->fileInfo->name;
+		dirHandle = dh.trueHandle;
+		UBA_ASSERT(dirPath || dirHandle);
+	}
+	else if (isListDirectoryHandle(dirHandle))
+		dirPath = asListDirectoryHandle(dirHandle).originalName;
+
+	StringBuffer<> combined;
+
+	auto trueRun2 = [&]()
+		{
+			combined.Append(dirPath).EnsureEndsWithSlash().Append(dosPath);
+			dosPath = combined.data;
+			HANDLE originalDir = relativeName->ContainingDirectory;
+			relativeName->ContainingDirectory = nullptr;
+			NTSTATUS res = trueRun();
+			relativeName->ContainingDirectory = originalDir;
+			return res;
+		};
+
+	if (dirPath)
+		return trueRun2();
+
+	tchar temp[1024];
+	u32 dirLen = True_GetFinalPathNameByHandleW(dirHandle, temp, 1024, 0);
+	if (dirLen == 0 || dirLen >= 1024)
+	{
+		UBA_ASSERTF(false, TC("True_GetFinalPathNameByHandleW failed for handle %llu"), u64(dirHandle));
+		return trueRun();
+	}
+
+	UBA_ASSERT(StartsWith(temp, TC("\\\\?\\")));
+	dirPath = temp + 4;
+	return trueRun2();
 }
 
 DWORD WINAPI Detoured_RtlGetFullPathName_U(const WCHAR* name, ULONG size, WCHAR* buffer, WCHAR** file_part)
