[UBA]
* Removed code in RtlDosPathNameToRelativeNtPathName_U_WithStatus using ContainingDirectory since it is an output parameter, not input.

diff --git a/Engine/Source/Programs/UnrealBuildAccelerator/Detours/Private/Windows/UbaDetoursFunctionsNtDll.inl b/Engine/Source/Programs/UnrealBuildAccelerator/Detours/Private/Windows/UbaDetoursFunctionsNtDll.inl
--- a/Engine/Source/Programs/UnrealBuildAccelerator/Detours/Private/Windows/UbaDetoursFunctionsNtDll.inl
+++ b/Engine/Source/Programs/UnrealBuildAccelerator/Detours/Private/Windows/UbaDetoursFunctionsNtDll.inl
@@ -125,60 +125,12 @@ NTSTATUS Detoured_RtlDosPathNameToNtPathName_U_WithStatus(PCWSTR dosPath, PUNICO
 
 NTSTATUS Detoured_RtlDosPathNameToRelativeNtPathName_U_WithStatus(PCWSTR dosPath, PUNICODE_STRING ntPath, PWSTR* filePart, RTL_RELATIVE_NAME_U* relativeName)
 {
-	//const tchar* original = dosPath;
-
-	auto trueRun = [&]()
-		{
-			StringBuffer<> temp;
-			FixupDosPath(dosPath, temp);
-			auto res = True_RtlDosPathNameToRelativeNtPathName_U_WithStatus(dosPath, ntPath, filePart, relativeName);
-			//DEBUG_LOG(L"RtlDosPathNameToRelativeNtPathName_U_WithStatus:   %s -> %.*s", original, ntPath->Length/2, ntPath->Buffer);
-			return res;
-		};
-
-	if (!relativeName || !relativeName->ContainingDirectory)
-		return trueRun();
-
-	HANDLE dirHandle = relativeName->ContainingDirectory;
-	const tchar* dirPath = nullptr;
-	if (isDetouredHandle(dirHandle))
-	{
-		DetouredHandle& dh = asDetouredHandle(dirHandle);
-		if (dh.fileObject && dh.fileObject->fileInfo)
-			dirPath = dh.fileObject->fileInfo->name;
-		dirHandle = dh.trueHandle;
-		UBA_ASSERT(dirPath || dirHandle);
-	}
-	else if (isListDirectoryHandle(dirHandle))
-		dirPath = asListDirectoryHandle(dirHandle).originalName;
-
-	StringBuffer<> combined;
-
-	auto trueRun2 = [&]()
-		{
-			combined.Append(dirPath).EnsureEndsWithSlash().Append(dosPath);
-			dosPath = combined.data;
-			HANDLE originalDir = relativeName->ContainingDirectory;
-			relativeName->ContainingDirectory = nullptr;
-			NTSTATUS res = trueRun();
-			relativeName->ContainingDirectory = originalDir;
-			return res;
-		};
-
-	if (dirPath)
-		return trueRun2();
-
-	tchar temp[1024];
-	u32 dirLen = True_GetFinalPathNameByHandleW(dirHandle, temp, 1024, 0);
-	if (dirLen == 0 || dirLen >= 1024)
-	{
-		UBA_ASSERTF(false, TC("True_GetFinalPathNameByHandleW failed for handle %llu"), u64(dirHandle));
-		return trueRun();
-	}
-
-	UBA_ASSERT(StartsWith(temp, TC("\\\\?\\")));
-	dirPath = temp + 4;
-	return trueRun2();
+	const tchar* original = dosPath;(void)original;
+	StringBuffer<> temp;
+	FixupDosPath(dosPath, temp);
+	auto res = True_RtlDosPathNameToRelativeNtPathName_U_WithStatus(dosPath, ntPath, filePart, relativeName);
+	//DEBUG_LOG(L"RtlDosPathNameToRelativeNtPathName_U_WithStatus:   %s -> %.*s", original, ntPath->Length/2, ntPath->Buffer);
+	return res;
 }
 
 DWORD WINAPI Detoured_RtlGetFullPathName_U(const WCHAR* name, ULONG size, WCHAR* buffer, WCHAR** file_part)
@@ -205,7 +157,6 @@ BOOLEAN NTAPI Detoured_RtlDosPathNameToRelativeNtPathName_U(PCWSTR dosPath, PUNI
 {
 	const tchar* original = dosPath;(void)original;
 	StringBuffer<> temp;
-	UBA_ASSERT(!relativeName || !relativeName->ContainingDirectory);
 	FixupDosPath(dosPath, temp);
 	auto res = True_RtlDosPathNameToRelativeNtPathName_U(dosPath, ntPath, filePart, relativeName);
 	//DEBUG_LOG(L"RtlDosPathNameToRelativeNtPathName_U:              %s -> %.*s", original, ntPath->Length/2, ntPath->Buffer);
