[UBA]
* Removed code not needed anymore

diff --git a/Engine/Source/Programs/UnrealBuildAccelerator/Core/Public/UbaFileMappingWin.inl b/Engine/Source/Programs/UnrealBuildAccelerator/Core/Public/UbaFileMappingWin.inl
--- a/Engine/Source/Programs/UnrealBuildAccelerator/Core/Public/UbaFileMappingWin.inl
+++ b/Engine/Source/Programs/UnrealBuildAccelerator/Core/Public/UbaFileMappingWin.inl
@@ -7,85 +7,7 @@ namespace uba
 	////////////////////////////////////////////////////////////////////////////////////////////////////
 
 	HANDLE g_hostProcess;
-
-	namespace internal
-	{
-		#include "UbaFileMapping.inl"
-	}
-
-	////////////////////////////////////////////////////////////////////////////////////////////////////
-
-	#define UBA_FM_FUNC(func) decltype(FileMapping_##func)* g_##func##Func = internal::FileMapping_##func;
-	UBA_FM_FUNCTIONS
-	#undef UBA_FM_FUNC
-
-	////////////////////////////////////////////////////////////////////////////////////////////////////
-
-	FileMappingHandle FileMapping_Create(Logger& logger, u32 flProtect, u64 maxSize, const tchar* name, const tchar* hint)
-	{
-		//return g_CreateFunc(logger, flProtect, maxSize, name, hint);
-		return internal::FileMapping_Create(logger, flProtect, maxSize, name, hint);
-	}
-
-	FileMappingHandle FileMapping_CreateFromFile(Logger& logger, FileHandle file, u32 flProtect, u64 maxSize, const tchar* hint)
-	{
-		//return g_CreateFromFileFunc(logger, file, flProtect, maxSize, hint);
-		return internal::FileMapping_CreateFromFile(logger, file, flProtect, maxSize, hint);
-	}
-
-	u8* FileMapping_Map(Logger& logger, FileMappingHandle fileMappingObject, u32 dwDesiredAccess, u64 offset, u64 dwNumberOfBytesToMap, bool allowDiscard)
-	{
-		//return g_MapFunc(logger, fileMappingObject, dwDesiredAccess, offset, dwNumberOfBytesToMap, allowDiscard);
-		return internal::FileMapping_Map(logger, fileMappingObject, dwDesiredAccess, offset, dwNumberOfBytesToMap, allowDiscard);
-	}
-
-	bool FileMapping_Commit(Logger& logger, void* address, u64 size, bool allowDiscard)
-	{
-		//return g_CommitFunc(logger, address, size, allowDiscard);
-		return internal::FileMapping_Commit(logger, address, size, allowDiscard);
-	}
-
-	bool FileMapping_Unmap(Logger& logger, const void* lpBaseAddress, u64 bytesToUnmap, const tchar* hint, bool allowDiscard)
-	{
-		//return g_UnmapFunc(logger, lpBaseAddress, bytesToUnmap, hint, allowDiscard);
-		return internal::FileMapping_Unmap(logger, lpBaseAddress, bytesToUnmap, hint, allowDiscard);
-	}
-
-	bool FileMapping_Close(Logger& logger, FileMappingHandle h, const tchar* hint)
-	{
-		//return g_CloseFunc(logger, h, hint);
-		return internal::FileMapping_Close(logger, h, hint);
-	}
-
-	void* FileMapping_ReservePlaceholder(void* baseAddress, u64 capacity)
-	{
-		//return g_ReservePlaceholderFunc(baseAddress, capacity);
-		return internal::FileMapping_ReservePlaceholder(baseAddress, capacity);
-	}
-
-	void* FileMapping_MapPlaceholder(FileMappingHandle handle, MemoryMapType mapType, void* targetAddress, u64 targetOffset, u64 targetCapacity, u64 handleOffset, u64 size)
-	{
-		//return g_MapPlaceholderFunc(handle, mapType, targetAddress, targetOffset, targetCapacity, handleOffset, size);
-		return internal::FileMapping_MapPlaceholder(handle, mapType, targetAddress, targetOffset, targetCapacity, handleOffset, size);
-	}
-
-	bool FileMapping_UnmapPlaceholder(void* memory, u64 capacity, u64 mappedSize, u64* subMappings, u64 subMappingsCount)
-	{
-		//return g_UnmapPlaceholderFunc(memory, capacity, mappedSize, subMappings, subMappingsCount);
-		return internal::FileMapping_UnmapPlaceholder(memory, capacity, mappedSize, subMappings, subMappingsCount);
-	}
-
-	FileMappingHandle FileMapping_DuplicateFromHost(FileMappingHandle handle, const char* hint)
-	{
-		//return g_DuplicateFromHostFunc(handle, hint);
-		return internal::FileMapping_DuplicateFromHost(handle, hint);
-	}
-
-	u8* FileMapping_MapFromHost(FileMappingHandle handle, u64 size, u64 offset, bool writable, const char* hint)
-	{
-		//return g_MapFromHostFunc(handle, size, offset, writable, hint);
-		return internal::FileMapping_MapFromHost(handle, size, offset, writable, hint);
-	}
+	#include "UbaFileMapping.inl"
 
 	////////////////////////////////////////////////////////////////////////////////////////////////////
 
