[UBA]
* Added Write/ReadVectorCompressed to BinaryReader/Writer

diff --git a/Engine/Source/Programs/UnrealBuildAccelerator/Core/Private/UbaPlatform.cpp b/Engine/Source/Programs/UnrealBuildAccelerator/Core/Private/UbaPlatform.cpp
--- a/Engine/Source/Programs/UnrealBuildAccelerator/Core/Private/UbaPlatform.cpp
+++ b/Engine/Source/Programs/UnrealBuildAccelerator/Core/Private/UbaPlatform.cpp
@@ -788,4 +788,25 @@ namespace uba
 			return TString(temp.begin(), temp.end());
 		}
 	}
+
+	bool BinaryWriter::WriteBytesCompressed(const void* data, u64 size)
+	{
+		u64 needed = (u64)OodleLZ_GetCompressedBufferSizeNeeded(OodleLZ_Compressor_Kraken, size);
+		Vector<u8> mem;
+		mem.resize(needed);
+		u64 compressedSize = OodleLZ_Compress(OodleLZ_Compressor_Kraken, data, (OO_SINTa)size, mem.data(), OodleLZ_CompressionLevel_Normal);
+		UBA_ASSERT(compressedSize != OODLELZ_FAILED);
+		Write7BitEncoded(compressedSize);
+		WriteBytes(mem.data(), compressedSize);
+		return true;
+	}
+
+	void BinaryReader::ReadBytesCompressed(void* outData, u64 outSize)
+	{
+		Vector<u8> compressed;
+		compressed.resize(Read7BitEncoded());
+		ReadBytes(compressed.data(), compressed.size());
+		OO_SINTa decompLen = OodleLZ_Decompress(compressed.data(), (OO_SINTa)compressed.size(), outData, (OO_SINTa)outSize);
+		UBA_ASSERT(decompLen == (OO_SINTa)outSize);(void)decompLen;
+	}
 }
diff --git a/Engine/Source/Programs/UnrealBuildAccelerator/Core/Public/UbaBinaryReaderWriter.h b/Engine/Source/Programs/UnrealBuildAccelerator/Core/Public/UbaBinaryReaderWriter.h
--- a/Engine/Source/Programs/UnrealBuildAccelerator/Core/Public/UbaBinaryReaderWriter.h
+++ b/Engine/Source/Programs/UnrealBuildAccelerator/Core/Public/UbaBinaryReaderWriter.h
@@ -41,6 +41,7 @@ namespace uba
 		inline void WriteUtf8String(const CharType* str, u64 len);
 
 		void WriteLongString(const StringView& str, u32 threshold = 10000);
+		bool WriteBytesCompressed(const void* data, u64 size);
 
 		inline void WriteFileTime(u64 fileTime) { WriteU64(fileTime); }
 		inline void WriteFileAttributes(u32 attributes) { Write7BitEncoded(attributes); }
@@ -93,6 +94,7 @@ namespace uba
 		inline u64 Read7BitEncoded();
 		template<typename CharType> inline CharType ReadUtf8Char();
 		TString ReadLongString();
+		void ReadBytesCompressed(void* outData, u64 outCapacity);
 
 		inline u64 ReadFileTime() { return ReadU64(); }
 		inline u32 ReadFileAttributes() { return u32(Read7BitEncoded()); }
