[UBA]
* Changed so bcrypt and bryptprimitives are always loaded if mimalloc is enabled. These dlls are always loaded in first thread doing allocations and can cause race conditions if callstack is not in a good shape with locks etc

diff --git a/Engine/Source/Programs/UnrealBuildAccelerator/Core/Private/UbaApplicationRules.cpp b/Engine/Source/Programs/UnrealBuildAccelerator/Core/Private/UbaApplicationRules.cpp
--- a/Engine/Source/Programs/UnrealBuildAccelerator/Core/Private/UbaApplicationRules.cpp
+++ b/Engine/Source/Programs/UnrealBuildAccelerator/Core/Private/UbaApplicationRules.cpp
@@ -240,19 +240,6 @@ namespace uba
 	{
 		using Super = ApplicationRulesVcLink;
 
-		virtual const tchar* const* LibrariesToPreload() const override
-		{
-			// Special handling.. it seems loading bcrypt.dll can deadlock when using mimalloc so we make sure to load it here directly instead
-			// There is a setting to disable bcrypt dll loading inside mimalloc but with that change mimalloc does not work with older versions of windows
-			static constexpr const tchar* preloads[] = 
-			{
-				TC("bcrypt.dll"),
-				TC("bcryptprimitives.dll"),
-				nullptr,
-			};
-			return preloads;
-		}
-
 		virtual bool CloseFileMappingAfterUse(StringView fileName) const
 		{
 			return true;
@@ -473,19 +460,6 @@ namespace uba
 		{
 			return false; // fileName.EndsWith(TCV(".psym")); With psym as output file the BreakpadSymbolEncoder fails to output a .sym file
 		}
-
-		virtual const tchar* const* LibrariesToPreload() const override
-		{
-			// Special handling.. it seems loading bcrypt.dll can deadlock when using mimalloc so we make sure to load it here directly instead
-			// There is a setting to disable bcrypt dll loading inside mimalloc but with that change mimalloc does not work with older versions of windows
-			static constexpr const tchar* preloads[] = 
-			{
-				TC("bcrypt.dll"),
-				TC("bcryptprimitives.dll"),
-				nullptr,
-			};
-			return preloads;
-		}
 	};
 
 	class ApplicationRulesClangPlusPlusExePlatform1 : public ApplicationRulesClangPlusPlusExe
diff --git a/Engine/Source/Programs/UnrealBuildAccelerator/Detours/Private/Windows/UbaDetoursFunctionsWin.cpp b/Engine/Source/Programs/UnrealBuildAccelerator/Detours/Private/Windows/UbaDetoursFunctionsWin.cpp
--- a/Engine/Source/Programs/UnrealBuildAccelerator/Detours/Private/Windows/UbaDetoursFunctionsWin.cpp
+++ b/Engine/Source/Programs/UnrealBuildAccelerator/Detours/Private/Windows/UbaDetoursFunctionsWin.cpp
@@ -939,10 +939,21 @@ void PreInit(const DetoursPayload& payload)
 			FatalError(1349, L"Failed to reserve memory for cl.exe (%u)", GetLastError());
 	}
 
+	auto preloadDll = [](const tchar* dll)
+		{
+			if (!LoadLibraryExW(dll, NULL, LOAD_LIBRARY_SEARCH_SYSTEM32))
+				FatalError(1351, L"Failed to preload %s (%u)", dll, GetLastError());
+		};
+
+	// Mimalloc loads bcrypt and can deadlock processes, so we might just as well load it ourselves
+	if (g_useMiMalloc)
+	{
+		preloadDll(TC("bcrypt.dll"));
+		preloadDll(TC("bcryptprimitives.dll"));
+	}
 	if (const tchar* const* preloads = g_rules->LibrariesToPreload())
 		for (auto it = preloads; *it; ++it)
-			if (!LoadLibraryExW(*it, NULL, LOAD_LIBRARY_SEARCH_SYSTEM32))
-				FatalError(1351, L"Failed to preload %s (%u)", *it, GetLastError());
+			preloadDll(*it);
 }
 
 struct PatchImportInfo
