[UBA]
* Made track inputs thread safe and also pretend __pycache__ directories do not exist if track inputs is enabled

diff --git a/Engine/Source/Programs/UnrealBuildAccelerator/Detours/Private/UbaDetoursShared.cpp b/Engine/Source/Programs/UnrealBuildAccelerator/Detours/Private/UbaDetoursShared.cpp
--- a/Engine/Source/Programs/UnrealBuildAccelerator/Detours/Private/UbaDetoursShared.cpp
+++ b/Engine/Source/Programs/UnrealBuildAccelerator/Detours/Private/UbaDetoursShared.cpp
@@ -25,6 +25,7 @@ namespace uba
 	VARIABLE_MEM(DirectoryTable, g_directoryTable);
 	VARIABLE_MEM(MappedFileTable, g_mappedFileTable);
 	VARIABLE_MEM(ReaderWriterLock, g_consoleStringCs);
+	VARIABLE_MEM(Futex, g_trackInputsLock);
 
 	bool g_conEnabled[2];
 	u32 g_rulesIndex;
@@ -63,6 +64,8 @@ namespace uba
 		g_directoryTableMem.Create(g_memoryBlock);
 		g_mappedFileTableMem.Create(g_memoryBlock);
 		g_consoleStringCsMem.Create();
+
+		g_trackInputsLockMem.Create();
 	}
 
 #if UBA_DEBUG_LOG_ENABLED
@@ -212,8 +215,11 @@ namespace uba
 		if (!g_trackInputsMem)
 			return;
 
+		SCOPED_FUTEX(g_trackInputsLock, l);
+
 		if (g_trackInputsBufPos > TrackInputsMemCapacity - 2048)
 			SendInput();
+		UBA_ASSERTF(!file.StartsWith(TCV("\\")), TC("Error %s is not an absolute path. All tracked files must be absolute"), file.data);
 
 		BinaryWriter w(g_trackInputsMem, g_trackInputsBufPos, TrackInputsMemCapacity);
 		w.WriteString(file);
diff --git a/Engine/Source/Programs/UnrealBuildAccelerator/Detours/Private/Windows/UbaDetoursFunctionsKernelBase.inl b/Engine/Source/Programs/UnrealBuildAccelerator/Detours/Private/Windows/UbaDetoursFunctionsKernelBase.inl
--- a/Engine/Source/Programs/UnrealBuildAccelerator/Detours/Private/Windows/UbaDetoursFunctionsKernelBase.inl
+++ b/Engine/Source/Programs/UnrealBuildAccelerator/Detours/Private/Windows/UbaDetoursFunctionsKernelBase.inl
@@ -219,6 +219,12 @@ BOOL Detoured_CreateDirectoryW(LPCWSTR lpPathName, LPSECURITY_ATTRIBUTES lpSecur
 
 	StringBuffer<> lowerName(pathName);
 
+	if (IsTrackingInput() && pathName.EndsWith(TCV("__pycache__"))) // If we are populating cache we need to make sure not to use pycache
+	{
+		SetLastError(ERROR_ACCESS_DENIED);
+		return false;
+	}
+
 	if (pathName.StartsWith(g_systemTemp.data))
 	{
 		SuppressCreateFileDetourScope s;
@@ -1664,6 +1670,10 @@ BOOL Detoured_FindNextFileW(HANDLE hFindFile, LPWIN32_FIND_DATAW lpFindFileData)
 				SetLastError(ERROR_NO_MORE_FILES);
 				return false;
 			}
+
+			if (IsTrackingInput() && Equals(data.cFileName, L"__pycache__")) // If we are populating cache we need to make sure not to use pycache
+				continue;
+
 			if (listHandle.wildcard.empty() || PathMatchSpecW(data.cFileName, listHandle.wildcard.c_str()))
 			{
 				DEBUG_LOG_DETOURED(L"FindNextFileW", L"%llu (%ls) -> True", u64(hFindFile), data.cFileName);
@@ -2646,8 +2656,12 @@ HMODULE Shared_LoadLibrary(LPCWSTR lpLibFileName, HANDLE hFile, DWORD dwFlags, b
 	if (!g_runningRemote && !IsTrackingInput() || t_disallowDetour)
 		return NULL;
 
+	StringView pathView(ToView(lpLibFileName));
+	if (pathView.StartsWith(TCV("\\\\?\\")))
+		pathView = pathView.Skip(4);
+
 	StringBuffer<> path;
-	path.Append(lpLibFileName).FixPathSeparators();
+	path.Append(pathView).FixPathSeparators();
 
 	bool detourDll = path.EndsWith(L".exe") || path.EndsWith(L".dll") || path.EndsWith(L".node");
 	if (detourDll && g_systemRoot.count && path.StartsWith(g_systemRoot.data))
@@ -3008,7 +3022,7 @@ BOOL Detoured_CreateProcessW(LPCWSTR lpApplicationName, LPWSTR lpCommandLine, LP
 		application.Clear().Append(TCV("C:\\Windows\\system32\\cmd.exe"));
 	}
 
-	bool is32Bit = application.EndsWith(TC("gperf.exe")); // Hard coded for 32-bit application gperf.exe... this could be extended but will only work for apps that does not procuce output files since we don't track non-detoured applications
+	bool is32Bit = application.EndsWith(TC("gperf.exe")); // Hard coded for 32-bit application gperf.exe... this could be extended but will only work for apps that does not produce output files since we don't track non-detoured applications
 	if (is32Bit)
 	{
 		const wchar_t* realApplication = application.data;
