[UBA]
* Fixed so files sent by helpers are done in parallel.
* Fixed so tracked files are sent compressed (since they are just paths)

diff --git a/Engine/Source/Programs/UnrealBuildAccelerator/Common/Private/UbaSessionClient.cpp b/Engine/Source/Programs/UnrealBuildAccelerator/Common/Private/UbaSessionClient.cpp
--- a/Engine/Source/Programs/UnrealBuildAccelerator/Common/Private/UbaSessionClient.cpp
+++ b/Engine/Source/Programs/UnrealBuildAccelerator/Common/Private/UbaSessionClient.cpp
@@ -736,36 +736,39 @@ namespace uba
 		}
 	}
 
-	bool SessionClient::SendFiles(ProcessImpl& process, Timer& sendFiles)
+	bool SessionClient::SendFiles(ProcessImpl& process)
 	{
+		Atomic<bool> success = true;
 		auto& rules = *process.m_startInfo.rules;
+		auto& writtenFiles = process.m_shared.writtenFiles;
+		TimerScope timer(process.m_processStats.sendFiles);
 
-		StorageStatsScope storageStatsScope(process.m_storageStats);
-		for (auto& pair : process.m_shared.writtenFiles)
-		{
-			WrittenFile& file = pair.second;
-
-			if (StringView(file.name).StartsWith(m_tempPath) && !rules.KeepTempOutputFile(StringView(file.name)))
+		m_client.ParallelFor(u32(writtenFiles.size()), writtenFiles, [&](const WorkContext&, auto& it)
 			{
-				m_sharedMemory.CloseHandle(m_logger, file.memoryHandle, file.backedName.c_str());
-				file.memoryHandle = {};
-				#if UBA_DEBUG
-				m_logger.Info(TC("Skipping sending back temp file %s"), file.name.c_str());
+				StorageStatsScope storageStatsScope(process.m_storageStats);
+				WrittenFile& file = it->second;
+
+				if (StringView(file.name).StartsWith(m_tempPath) && !rules.KeepTempOutputFile(StringView(file.name)))
+				{
+					m_sharedMemory.CloseHandle(m_logger, file.memoryHandle, file.backedName.c_str());
+					file.memoryHandle = {};
+					#if UBA_DEBUG
+					m_logger.Info(TC("Skipping sending back temp file %s"), file.name.c_str());
+					#endif
+					return;
+				}
+
+				#ifdef _DEBUG
+				if (!pair.second.memoryHandle.IsValid())
+					m_logger.Warning(TC("%s is not using file mapping"), file.name.c_str());
 				#endif
-				continue;
-			}
 
-			TimerScope timer(sendFiles);
-#ifdef _DEBUG
-			if (!pair.second.memoryHandle.IsValid())
-				m_logger.Warning(TC("%s is not using file mapping"), file.name.c_str());
-#endif
-			bool keepMappingInMemory = IsWindows && IsMaybeReadAfterWritten(process, file.name) && file.memoryWritten < m_keepOutputFileMemoryMapsThreshold;
-			bool compressed = rules.SendFileCompressedFromClient(file.name);
-			if (!SendFile(file, process.GetId(), keepMappingInMemory, compressed))
-				return false;
-		}
-		return true;
+				bool keepMappingInMemory = IsWindows && IsMaybeReadAfterWritten(process, file.name) && file.memoryWritten < m_keepOutputFileMemoryMapsThreshold;
+				bool compressed = rules.SendFileCompressedFromClient(file.name);
+				if (!SendFile(file, process.GetId(), keepMappingInMemory, compressed))
+					success = false;
+			}, TCV("SendFile"));
+		return success;
 	}
 
 	bool SessionClient::SendFile(WrittenFile& file, u32 processId, bool keepMappingInMemory, bool compressed)
@@ -1495,7 +1498,6 @@ namespace uba
 		AddEnvironmentVariableNoLock(TCV("TEMP"), m_tempPath);
 		AddEnvironmentVariableNoLock(TCV("TMP"), m_tempPath);
 		//AddEnvironmentVariableNoLock(TCV("WINEDEBUG"), TCV("+process"));
-		AddEnvironmentVariableNoLock(TCV("PYTHONPYCACHEPREFIX"), m_tempPath);
 		#else
 		AddEnvironmentVariableNoLock(TCV("TMPDIR"), m_tempPath);
 		#endif
@@ -1626,25 +1628,13 @@ namespace uba
 	bool SessionClient::SendProcessInputs(ProcessImpl& process)
 	{
 		auto inputs = process.GetTrackedInputs();
-		u32 left = u32(inputs.size());
-		u32 capacityToAdd = left;
-		u8* readPos = inputs.data();
-		while (left)
-		{
-			StackBinaryWriter<SendMaxSize> writer;
-			NetworkMessage msg(m_client, ServiceId, SessionMessageType_ProcessInputs, writer);
-			writer.Write7BitEncoded(process.m_id);
-			writer.Write7BitEncoded(capacityToAdd);
-			capacityToAdd = 0;
-			u32 toWrite = Min(left, u32(writer.GetCapacityLeft()));
-			writer.WriteBytes(readPos, toWrite);
-			StackBinaryReader<32> reader;
-			if (!msg.Send(reader))
-				return false;
-			readPos += toWrite;
-			left -= toWrite;
-		}
-		return true;
+		StackBinaryWriter<SendMaxSize> writer;
+		NetworkMessage msg(m_client, ServiceId, SessionMessageType_ProcessInputs, writer);
+		writer.Write7BitEncoded(process.m_id);
+		writer.Write7BitEncoded(inputs.size());
+		writer.WriteBytesCompressed(inputs.data(), inputs.size());
+		StackBinaryReader<32> reader;
+		return msg.Send(reader);
 	}
 
 	bool SessionClient::SendProcessFinished(ProcessImpl& process, u32 exitCode)
@@ -2238,7 +2228,7 @@ namespace uba
 							// Should we decrease weight before or after sending files?
 							//decreaseWeight.Execute();
 
-							if (!session.SendFiles(process, process.m_processStats.sendFiles))
+							if (!session.SendFiles(process))
 							{
 								const tchar* desc = TC("Failed to send output files to host");
 								session.m_logger.Error(desc);
@@ -2893,7 +2883,7 @@ namespace uba
 	bool SessionClient::FlushWrittenFiles(ProcessImpl& process)
 	{
 		SCOPED_FUTEX(process.m_shared.writtenFilesLock, lock);
-		bool success = SendFiles(process, process.m_processStats.sendFiles);
+		bool success = SendFiles(process);
 		{
 			SCOPED_FUTEX(m_outputFilesLock, lock2);
 			for (auto& kv : process.m_shared.writtenFiles)
diff --git a/Engine/Source/Programs/UnrealBuildAccelerator/Common/Private/UbaSessionServer.cpp b/Engine/Source/Programs/UnrealBuildAccelerator/Common/Private/UbaSessionServer.cpp
--- a/Engine/Source/Programs/UnrealBuildAccelerator/Common/Private/UbaSessionServer.cpp
+++ b/Engine/Source/Programs/UnrealBuildAccelerator/Common/Private/UbaSessionServer.cpp
@@ -120,6 +120,7 @@ namespace uba
 		u64 m_peakMemory = 0;
 		Event m_done;
 		Vector<ProcessLogLine> m_logLines;
+		Futex m_receivedFilesLock;
 		Vector<ReceivedFile> m_receivedFiles;
 		Vector<u8> m_trackedInputs;
 		Vector<u8> m_trackedOutputs;
@@ -1578,6 +1579,8 @@ namespace uba
 		file.casRefcounted = true;
 
 		auto& process = *(RemoteProcess*)h.m_process;
+
+		SCOPED_FUTEX(m_receivedFilesLock, lock);
 		process.m_receivedFiles.push_back(file);
 		return true;
 	}
@@ -1903,12 +1906,9 @@ namespace uba
 		}
 		auto& process = *(RemoteProcess*)h.m_process;
 		auto& inputs = process.m_trackedInputs;
-		u64 size = inputs.size();
-		if (u64 addCapacity = reader.Read7BitEncoded())
-			inputs.reserve(size + addCapacity);
-		u64 toRead = reader.GetLeft();
-		inputs.resize(size + toRead);
-		reader.ReadBytes(inputs.data() + size, toRead);
+		u64 inputsSize = reader.Read7BitEncoded();
+		inputs.resize(inputsSize);
+		reader.ReadBytesCompressed(inputs.data(), inputsSize);
 		return true;
 	}
 
diff --git a/Engine/Source/Programs/UnrealBuildAccelerator/Common/Public/UbaNetwork.h b/Engine/Source/Programs/UnrealBuildAccelerator/Common/Public/UbaNetwork.h
--- a/Engine/Source/Programs/UnrealBuildAccelerator/Common/Public/UbaNetwork.h
+++ b/Engine/Source/Programs/UnrealBuildAccelerator/Common/Public/UbaNetwork.h
@@ -16,7 +16,7 @@ namespace uba
 	
 	static constexpr u32 SystemNetworkVersion = 1339;
 	static constexpr u32 StorageNetworkVersion = 4;
-	static constexpr u32 SessionNetworkVersion = 50;
+	static constexpr u32 SessionNetworkVersion = 51;
 	static constexpr u32 CacheNetworkVersion = 5;
 	static constexpr u32 CacheBucketVersion = 3;
 
diff --git a/Engine/Source/Programs/UnrealBuildAccelerator/Common/Public/UbaSessionClient.h b/Engine/Source/Programs/UnrealBuildAccelerator/Common/Public/UbaSessionClient.h
--- a/Engine/Source/Programs/UnrealBuildAccelerator/Common/Public/UbaSessionClient.h
+++ b/Engine/Source/Programs/UnrealBuildAccelerator/Common/Public/UbaSessionClient.h
@@ -91,7 +91,7 @@ namespace uba
 		bool ReadModules(List<ModuleInfo>& outModules, StringBufferBase& outVirtualApplication, u32 processId, StringView application, StringView workingDir);
 		bool EnsureBinaryFile(StringBufferBase& out, StringBufferBase& outVirtual, u32 processId, StringView fileName, const StringKey& fileNameKey, StringView applicationDir, StringView workingDir, const u8* loaderPaths, u32 loaderPathsSize);
 		bool WriteBinFile(StringBufferBase& out, const StringView& binaryName, const CasKey& casKey, const KeyToString& applicationDir, u32 fileAttributes);
-		bool SendFiles(ProcessImpl& process, Timer& sendFiles);
+		bool SendFiles(ProcessImpl& process);
 		bool SendFile(WrittenFile& file, u32 processId, bool keepMappingInMemory, bool compressed);
 		bool SendUpdateDirectoryTable(StackBinaryReader<SendMaxSize>& reader); // Note, reader is sent in to save stack space.
 		bool UpdateDirectoryTableFromServer(StackBinaryReader<SendMaxSize>& reader);
