[UBA]
* Changed assert of weird scenario related to same process hierarchy moving files to the same destination. now it skips one of the moves and prints out a warning

diff --git a/Engine/Source/Programs/UnrealBuildAccelerator/Common/Private/UbaSession.cpp b/Engine/Source/Programs/UnrealBuildAccelerator/Common/Private/UbaSession.cpp
--- a/Engine/Source/Programs/UnrealBuildAccelerator/Common/Private/UbaSession.cpp
+++ b/Engine/Source/Programs/UnrealBuildAccelerator/Common/Private/UbaSession.cpp
@@ -2908,12 +2908,18 @@ namespace uba
 
 				UBA_ASSERT(msg.toKey != StringKeyZero);
 				auto insres = shared.writtenFiles.try_emplace(msg.toKey);
-				UBA_ASSERTF(insres.second, TC("Failed to add new file %s when moving from %s (destination already exists)"), msg.fromName.data, msg.toName.data);
-				WrittenFile& newFile = insres.first->second;
-				newFile = oldFile;
-				newFile.key = msg.toKey;
-				newFile.name = msg.toName.data;
-				shared.writtenFiles.erase(findIt);
+				if (insres.second)
+				{
+					WrittenFile& newFile = insres.first->second;
+					newFile = oldFile;
+					newFile.key = msg.toKey;
+					newFile.name = msg.toName.data;
+					shared.writtenFiles.erase(findIt);
+				}
+				else
+				{
+					m_logger.Warning(TC("File %s written before by other process, skipped. (This could be bad if output is different and might also break caching)"), msg.toName.data);
+				}
 
 				if (isMapping)
 				{
