[UBA]
* Minor logging improvement

diff --git a/Engine/Source/Programs/UnrealBuildAccelerator/Common/Public/UbaDirectoryIterator.h b/Engine/Source/Programs/UnrealBuildAccelerator/Common/Public/UbaDirectoryIterator.h
--- a/Engine/Source/Programs/UnrealBuildAccelerator/Common/Public/UbaDirectoryIterator.h
+++ b/Engine/Source/Programs/UnrealBuildAccelerator/Common/Public/UbaDirectoryIterator.h
@@ -84,7 +84,7 @@ namespace uba
 			if (res == STATUS_NOT_A_DIRECTORY || res == STATUS_ACCESS_DENIED || res == STATUS_NO_MEDIA_IN_DEVICE || res == STATUS_FVE_LOCKED_VOLUME)
 				return false;
 			#if UBA_DEBUG
-			return logger.Error(L"NtCreateFile for TraverseDir on %s failed with code %x", str.data, (DWORD)res);
+			return logger.Error(L"NtCreateFile for TraverseDir on %s failed with code 0x%x", str.data, (DWORD)res);
 			#else
 			return false;
 			#endif
