[UBA]
* Fixed bug where wrong lock was being used to protect process received files

diff --git a/Engine/Source/Programs/UnrealBuildAccelerator/Common/Private/UbaSessionServer.cpp b/Engine/Source/Programs/UnrealBuildAccelerator/Common/Private/UbaSessionServer.cpp
--- a/Engine/Source/Programs/UnrealBuildAccelerator/Common/Private/UbaSessionServer.cpp
+++ b/Engine/Source/Programs/UnrealBuildAccelerator/Common/Private/UbaSessionServer.cpp
@@ -1580,7 +1580,7 @@ namespace uba
 
 		auto& process = *(RemoteProcess*)h.m_process;
 
-		SCOPED_FUTEX(m_receivedFilesLock, lock);
+		SCOPED_FUTEX(process.m_receivedFilesLock, lock);
 		process.m_receivedFiles.push_back(file);
 		return true;
 	}
