[UBA]
* Added known outputs to enqueuing of processes in scheduler
* Fixed some bugs when cancelling builds

diff --git a/Engine/Source/Programs/UnrealBuildAccelerator/Common/Private/UbaScheduler.cpp b/Engine/Source/Programs/UnrealBuildAccelerator/Common/Private/UbaScheduler.cpp
--- a/Engine/Source/Programs/UnrealBuildAccelerator/Common/Private/UbaScheduler.cpp
+++ b/Engine/Source/Programs/UnrealBuildAccelerator/Common/Private/UbaScheduler.cpp
@@ -56,22 +56,27 @@ namespace uba
 
 	struct Scheduler::ProcessStartInfo2 : ProcessStartInfoHolder
 	{
-		ProcessStartInfo2(const ProcessStartInfo& si, const u8* ki, u32 kic)
+		ProcessStartInfo2(const ProcessStartInfo& si, const u8* ki, u32 kic, const u8* ko, u32 koc)
 		:	ProcessStartInfoHolder(si)
 		, knownInputs(ki)
+		, knownOutputs(ko)
 		, knownInputsCount(kic)
+		, knownOutputsCount(koc)
 		{
 		}
 
 		~ProcessStartInfo2()
 		{
 			delete[] knownInputs;
+			delete[] knownOutputs;
 		}
 
 		CacheFetchInfo* cacheInfo = nullptr;
 
 		const u8* knownInputs;
+		const u8* knownOutputs;
 		u32 knownInputsCount;
+		u32 knownOutputsCount;
 		float weight = 1.0f;
 		u32 cacheBucketId = 0; // Zero means this process will not check cache
 		u32 memoryGroupId = 0; // Zero means no memory group so this process will not be taken into account in memory allocations
@@ -362,8 +367,14 @@ namespace uba
 			ki = new u8[info.knownInputsBytes];
 			memcpy(ki, info.knownInputs, info.knownInputsBytes);
 		}
+		u8* ko = nullptr;
+		if (info.knownOutputsCount)
+		{
+			ko = new u8[info.knownOutputsBytes];
+			memcpy(ko, info.knownOutputs, info.knownOutputsBytes);
+		}
 
-		auto info2 = new ProcessStartInfo2(info.info, ki, info.knownInputsCount);
+		auto info2 = new ProcessStartInfo2(info.info, ki, info.knownInputsCount, ko, info.knownOutputsCount);
 		info2->Expand();
 		info2->weight = info.weight;
 		info2->cacheBucketId = info.cacheBucketId;
@@ -462,7 +473,7 @@ namespace uba
 			auto& dependency = m_processEntries[*i];
 			if (dependency.runStatus == RunStatus_Success)
 				continue;
-			if (dependency.runStatus == RunStatus_Failed || dependency.runStatus == RunStatus_Skipped)
+			if (dependency.runStatus == RunStatus_Failed || dependency.runStatus == RunStatus_Skipped || m_cancelled)
 			{
 				AddToQueue(m_readyToSkip, entryIndex);
 				return;
@@ -593,6 +604,21 @@ namespace uba
 		UBA_ASSERT(entry.dependencyCount == 0);
 		entry = newEntry;
 		entry.firstDependent = firstDependent;
+
+		if (m_cancelled)
+		{
+			entry.runStatus = RunStatus_Skipped;
+			UpdateDependentsNoLock(RunStatus_Skipped, firstDependent);
+			entry.firstDependent = nullptr;
+			auto pi = entry.info;
+			entry.info = nullptr;
+			lock.Leave();
+			ProcessLink::DeleteAll(firstDependent);
+			SkipProcess(*pi);
+			delete pi;
+			return;
+		}
+
 		SetupDependents(entry, index, info);
 		UpdateQueueCounter(1, entry.canExecuteRemotely);
 		m_updateThreadWakeup.Set();
@@ -839,12 +865,14 @@ namespace uba
 		{
 			entry.runStatus = RunStatus_Skipped;
 			auto pi = entry.info;
+			entry.info = nullptr;
 			ProcessLink* firstDependent = entry.firstDependent;
 			entry.firstDependent = nullptr;
 			UpdateDependentsNoLock(RunStatus_Skipped, firstDependent);
 			lock.Leave();
 			ProcessLink::DeleteAll(firstDependent);
 			SkipProcess(*pi);
+			delete pi;
 			return;
 		}
 
@@ -998,6 +1026,9 @@ namespace uba
 			return skipEntry;
 		}
 
+		if (m_cancelled)
+			return ~0u;
+
 		auto returnEntry = [&](u32 index)
 			{
 				if (index == ~0u)
@@ -1005,6 +1036,7 @@ namespace uba
 				auto& entry = m_processEntries[index];
 				UBA_ASSERT(entry.cacheStatus == CacheStatus_Failed);
 				outPrevStatus = entry.runStatus;
+				UBA_ASSERT(!m_cancelled);
 				entry.runStatus = RunStatus_Running;
 				return index;
 			};
@@ -1223,7 +1255,7 @@ namespace uba
 
 			UpdateQueueCounter(-1, canExecuteRemotely);
 
-			if (wasSkipped)
+			if (wasSkipped || m_cancelled)
 			{
 				SkipProcess(*info);
 				continue;
@@ -1500,53 +1532,71 @@ namespace uba
 
 		if (isDone)
 		{
+			ProcessLink::DeleteAll(firstDependent);
+
 			if (exitCode != 0)
 				++m_errorCount;
 
-			if (m_writeTrackedInputs)
+			bool writeTracked = m_writeTrackedInputs;
+			bool writeToCache = m_writeToCache && exitCode == 0 && process.GetExecutionType() != ProcessExecutionType_FromCache;
+
+			if (writeTracked || writeToCache)
 			{
-				static int counter = 0;
-				StringBuffer<> inputsFile;
-				inputsFile.Append(m_session.GetSessionDir()).Append(TCV("inputs")).EnsureEndsWithSlash();
-				static u32 processId = 1; // TODO: This should be done in a better way.. or not at all?
-				GenerateNameForProcess(inputsFile, process.GetStartInfo().arguments, ++processId);
-				inputsFile.Append(TCV(".log"));
-				FileAccessor f(m_session.GetLogger(), inputsFile.data);
-				if (f.CreateWrite())
+				Vector<u8> inputs;
+				Vector<u8> outputs;
+				CreateCacheInputsAndOutputs(inputs, outputs, *si, process);
+
+				if (writeTracked)
 				{
-					Set<TString> files;
-					auto& inputs = process.GetTrackedInputs();
-					BinaryReader reader(inputs.data(), 0, inputs.size());
-					while (reader.GetLeft())
-						files.insert(reader.ReadString());
-					for (auto& input : files)
+					static int counter = 0;
+					StringBuffer<> inputsFile;
+					inputsFile.Append(m_session.GetSessionDir()).Append(TCV("inputs")).EnsureEndsWithSlash();
+					static u32 processId = 1; // TODO: This should be done in a better way.. or not at all?
+					GenerateNameForProcess(inputsFile, process.GetStartInfo().arguments, ++processId);
+					inputsFile.Append(TCV(".log"));
+					FileAccessor f(m_session.GetLogger(), inputsFile.data);
+					auto writeLine = [&](StringView str)
+						{
+							u8 buffer[1024];
+							BinaryWriter writer(buffer, 0, sizeof(buffer));
+							writer.WriteUtf8String(str.data, str.count);
+							writer.WriteByte('\n');
+							f.Write(buffer, writer.GetPosition());
+						};
+					auto writeLines = [&](Vector<u8>& files)
+						{
+							BinaryReader intputReader(files.data(), 0, files.size());
+							while (intputReader.GetLeft())
+							{
+								StringBuffer<> file;
+								intputReader.ReadString(file);
+								writeLine(file);
+							}
+						};
+					if (f.CreateWrite())
 					{
-						u8 buffer[1024];
-						BinaryWriter writer(buffer, 0, sizeof(buffer));
-						writer.WriteUtf8String(input.data(), input.size());
-						writer.WriteByte('\n');
-						f.Write(buffer, writer.GetPosition());
+						writeLine(TCV("Inputs:"));
+						writeLines(inputs);
+						writeLine(TCV(""));
+						writeLine(TCV("Outputs:"));
+						writeLines(outputs);
+						writeLine(TCV(""));
+
+						f.Close();
 					}
-					f.Close();
 				}
-			}
 
-			FinishProcess(ph);
-			ProcessLink::DeleteAll(firstDependent);
-
-			// This does only work if build step actually always write. Many tools read and check if file looks ok and in that case don't write
-			// So this is not a reliable approach..
-			if (m_writeToCache && exitCode == 0 && process.GetExecutionType() != ProcessExecutionType_FromCache)
-			{
-				UBA_ASSERTF(false, TC("Not implemented"));
-				#if 0
-				auto inputs = process.GetTrackedInputs();
-				auto outputs = process.GetTrackedOutputs();
-
-				CacheClient& cacheClient = *m_cacheClients[0];
-				cacheClient.WriteToCache(si->cacheBucketId, *si, inputs.data(), inputs.size(), outputs.data(), outputs.size(), nullptr, 0, process.GetId());
-				#endif
+				// This does only work if build step actually always write. Many tools read and check if file looks ok and in that case don't write
+				// So this is not a reliable approach..
+				if (writeToCache)
+				{
+					CacheClient& cacheClient = *m_cacheClients[0];
+					if (!cacheClient.WriteToCache(si->cacheBucketId, *si, inputs.data(), inputs.size(), outputs.data(), outputs.size(), nullptr, 0, process.GetId()))
+						m_session.GetLogger().Info(TC("Failed to write %s to cache"), si->GetDescription());
+				}
 			}
+			// Need to be called after cache write in order to prevent IsEmpty from returning true
+			FinishProcess(ph);
 
 			delete si;
 		}
@@ -1596,6 +1646,72 @@ namespace uba
 			m_processFinished(handle);
 	}
 
+	void Scheduler::CreateCacheInputsAndOutputs(Vector<u8>& outInputs, Vector<u8>& outOutputs, ProcessStartInfo2& info, Process& process)
+	{
+		auto writeString = [](Vector<u8>& v, u64& p, StringView str)
+			{
+				if (p + 1024 > v.size())
+					v.resize(p + 1024);
+				BinaryWriter w(v.data(), p, v.size());
+				w.WriteString(str);
+				p = w.GetPosition();
+			};
+
+		u64 writePosOutputs = 0;
+		Set<TString> handledFiles;
+		auto writeOutput = [&](TString&& str)
+			{
+				auto res = handledFiles.emplace(std::move(str));
+				if (!res.second)
+					return;
+				writeString(outOutputs, writePosOutputs, *res.first);
+			};
+
+		auto& outputs = process.GetTrackedOutputs();
+		BinaryReader outputsReader(outputs.data(), 0, outputs.size());
+		while (outputsReader.GetLeft())
+			writeOutput(outputsReader.ReadString());
+
+		if (auto ko = (const tchar*)info.knownOutputs)
+			while (*ko)
+			{
+				StringView kov = ToView(ko);
+				writeOutput(kov.ToString());
+				ko += kov.count + 1;
+			}
+
+		outOutputs.resize(writePosOutputs);
+
+		// Adding inputs. Note, inputs can't also be outputs.
+		u64 writePosInputs = 0;
+		auto writeInput = [&](TString&& str)
+			{
+				auto res = handledFiles.emplace(std::move(str));
+				if (!res.second)
+					return;
+				#if PLATFORM_WINDOWS
+				if (StringView(*res.first).Contains(TCV("api-ms-win-")) || StringView(*res.first).Contains(TCV("tzres.dll")))
+					return;
+				#endif
+				writeString(outInputs, writePosInputs, *res.first);
+			};
+
+		auto& inputs = process.GetTrackedInputs();
+		outInputs.reserve(inputs.size());
+		BinaryReader inputsReader(inputs.data(), 0, inputs.size());
+		while (inputsReader.GetLeft())
+			writeInput(inputsReader.ReadString());
+
+		if (auto ki = (const tchar*)info.knownInputs)
+			while (*ki)
+			{
+				StringView kiv = ToView(ki);
+				writeInput(kiv.ToString());
+				ki += kiv.count + 1;
+			}
+
+		outInputs.resize(writePosInputs);
+	}
 
 	void Scheduler::KillNewestLocalProcess(StringView reason)
 	{
diff --git a/Engine/Source/Programs/UnrealBuildAccelerator/Common/Public/UbaScheduler.h b/Engine/Source/Programs/UnrealBuildAccelerator/Common/Public/UbaScheduler.h
--- a/Engine/Source/Programs/UnrealBuildAccelerator/Common/Public/UbaScheduler.h
+++ b/Engine/Source/Programs/UnrealBuildAccelerator/Common/Public/UbaScheduler.h
@@ -67,6 +67,10 @@ namespace uba
 		u32 knownInputsBytes = 0; // knownInputsBytes is the total size in bytes of knownInputs
 		u32 knownInputsCount = 0; // knownInputsCount is the number of strings in the memory block
 
+		const void* knownOutputs = nullptr; // knownOutputs is a memory block with null terminated tchar strings followed by an empty null terminated string to end.
+		u32 knownOutputsBytes = 0; // knownOutputsBytes is the total size in bytes of knownOutputs
+		u32 knownOutputsCount = 0; // knownOutputsCount is the number of strings in the memory block
+
 		const u32* dependencies = nullptr; // An array of u32 holding indicies to processes this process depends on. Index is a rolling number returned by EnqueueProcess
 		u32 dependencyCount = 0; // Number of elements in dependencies
 
@@ -171,8 +175,7 @@ namespace uba
 		void UpdateQueueCounter(int offset, bool canRunRemotely);
 		void UpdateActiveProcessCounter(bool isLocal, int offset);
 		void FinishProcess(const ProcessHandle& handle);
-		void UpdateRunStatus(ProcessEntry& entry, u32 entryIndex, RunStatus newStatus);
-
+		void CreateCacheInputsAndOutputs(Vector<u8>& outInputs, Vector<u8>& outOutputs, ProcessStartInfo2& info, Process& process);
 		void KillNewestLocalProcess(StringView reason);
 		void NativeProcessCreated(ProcessImpl& process);
 		void LocalProcessStart(ProcessStartInfo2& info);
