[UBA]
* Minor opt in CacheClient when running on windows
* Improved some error logging
* Added comment about optimizing sending files to cache server

diff --git a/Engine/Source/Programs/UnrealBuildAccelerator/Common/Private/UbaCacheClient.cpp b/Engine/Source/Programs/UnrealBuildAccelerator/Common/Private/UbaCacheClient.cpp
--- a/Engine/Source/Programs/UnrealBuildAccelerator/Common/Private/UbaCacheClient.cpp
+++ b/Engine/Source/Programs/UnrealBuildAccelerator/Common/Private/UbaCacheClient.cpp
@@ -335,12 +335,12 @@ namespace uba
 			{
 				tchar temp[512];
 				bool res = SearchPathW(NULL, path.data, NULL, 512, temp, NULL);
-				path.Clear().Append(temp);
 				if (!res)
 				{
 					m_logger.Info(TC("Can't find file: %s"), path.data);
 					return false;
 				}
+				path.Clear().Append(temp);
 			}
 			#endif
 
@@ -1415,6 +1415,8 @@ namespace uba
 						return m_logger.Warning(TC("File %s has size 0 and supports placeholders.. so it is likely we are sending a placeholder file. Investigate!!"), path.data);
 				}
 
+				// TODO: Use WorkManager for this!.. flush them in the end
+
 				if (!SendFile(logger, m_client, &m_storage.m_workManager, casKey, fileData, fileSize, path.data))
 					return false;
 
@@ -1817,7 +1819,15 @@ namespace uba
 
 		PreparedDir* dir = nullptr;
 
-		if (!dir)
+		#if PLATFORM_WINDOWS
+		SCOPED_FUTEX_READ(m_directoryPreparserLock, preparserLockRead);
+		auto findIt = m_directoryPreparser.find(pathKey);
+		bool found = findIt != m_directoryPreparser.end();
+		preparserLockRead.Leave();
+		if (found)
+			dir = &findIt->second;
+		else
+		#endif
 		{
 			SCOPED_FUTEX(m_directoryPreparserLock, preparserLock);
 			dir = &m_directoryPreparser.try_emplace(pathKey).first->second;
