[UBA]
* Added feature speculative cache test. This feature allows uba to proceed with testing and even start downloading dependencies directly after a cache test is a hit. This can keep the download pipe saturated.

Example B depends on A. Both are in the cache server

Without speculative cache test: Client tests A, notices it is available, downloads the outputs and tell B it is done. Client tests B, notices it is available and downloads it.

With speculative cache test: Client tests A, notices it is available, register temporary caskeys for output files, tell B it can start testing. B start testing in parallel with A downloading, B notices it is available, register caskeys for output files and start downloading (still in parallel with A). Once A is done it removes the temporary cas keys, B does the same.

This will prevent bubbles in the pipeline and keep network saturated. If there is a huge pdb in the end for the exe and 100% cache hits, then that will end up being downloaded in parallel with other last entries.

diff --git a/Engine/Source/Programs/UnrealBuildAccelerator/Common/Private/UbaCacheClient.cpp b/Engine/Source/Programs/UnrealBuildAccelerator/Common/Private/UbaCacheClient.cpp
--- a/Engine/Source/Programs/UnrealBuildAccelerator/Common/Private/UbaCacheClient.cpp
+++ b/Engine/Source/Programs/UnrealBuildAccelerator/Common/Private/UbaCacheClient.cpp
@@ -234,6 +234,45 @@ namespace uba
 		return true;
 	}
 
+	u64 CacheClient::RegisterSpeculativeOutputs(FetchContext& context)
+	{
+		UBA_ASSERT(context.bucket);
+		auto& bucket = *context.bucket;
+
+		SCOPED_FUTEX(m_speculativeCasKeysLock, lock);
+		u64 id = m_nextSpeculativeOutputId++;
+		auto& keys = m_speculativeOutputsById[id];
+
+		for (u32 offset : context.casKeyOffsets)
+		{
+			StringBuffer<MaxPath> path;
+			CasKey casKey;
+			if (!GetLocalPathAndCasKey(bucket, context.rootPaths, path, casKey, bucket.serverCasKeyTable, bucket.serverPathTable, offset))
+				continue;
+
+			StringBuffer<MaxPath> forKey;
+			forKey.Append(path);
+			if (CaseInsensitiveFs)
+				forKey.MakeLower();
+			StringKey fileNameKey = ToStringKey(forKey);
+
+			m_speculativeCasKeys.emplace(fileNameKey, casKey);
+			keys.push_back(fileNameKey);
+		}
+		return id;
+	}
+
+	void CacheClient::UnregisterSpeculativeOutputs(u64 id)
+	{
+		SCOPED_FUTEX(m_speculativeCasKeysLock, lock);
+		auto it = m_speculativeOutputsById.find(id);
+		if (it == m_speculativeOutputsById.end())
+			return;
+		for (const StringKey& key : it->second)
+			m_speculativeCasKeys.erase(key);
+		m_speculativeOutputsById.erase(it);
+	}
+
 	bool CacheClient::WriteToCache(u32 bucketId, const ProcessStartInfo& info, const u8* inputs, u64 inputsSize, const u8* outputs, u64 outputsSize, const u8* logLines, u64 logLinesSize, u32 processId)
 	{
 		RootPaths rootPaths;
@@ -731,8 +770,17 @@ namespace uba
 					}
 					else
 					{
-						if (!m_storage.StoreCasKey(localCasKey, fileNameKey, path.data, CasKeyZero, false))
-							m_logger.Info(TC("Failed to store input file %s when fetching from cache"), path.data);
+						{
+							SCOPED_FUTEX_READ(m_speculativeCasKeysLock, slock);
+							auto speculativeIt = m_speculativeCasKeys.find(fileNameKey);
+							if (speculativeIt != m_speculativeCasKeys.end())
+								localCasKey = speculativeIt->second;
+						}
+						if (localCasKey == CasKeyZero)
+						{
+							if (!m_storage.StoreCasKey(localCasKey, fileNameKey, path.data, CasKeyZero, false))
+								m_logger.Info(TC("Failed to store input file %s when fetching from cache"), path.data);
+						}
 					}
 					#if UBA_TRACE_FETCH_CACHE
 					storeTime += GetTime() - t;
diff --git a/Engine/Source/Programs/UnrealBuildAccelerator/Common/Private/UbaScheduler.cpp b/Engine/Source/Programs/UnrealBuildAccelerator/Common/Private/UbaScheduler.cpp
--- a/Engine/Source/Programs/UnrealBuildAccelerator/Common/Private/UbaScheduler.cpp
+++ b/Engine/Source/Programs/UnrealBuildAccelerator/Common/Private/UbaScheduler.cpp
@@ -22,6 +22,7 @@
 #define UBA_SCHEDULER_PROCESS_QUEUES \
 			UBA_SCHEDULER_PROCESS_QUEUE(TestCache) \
 			UBA_SCHEDULER_PROCESS_QUEUE(DownloadCache) \
+			UBA_SCHEDULER_PROCESS_QUEUE(CompleteFromCache) \
 			UBA_SCHEDULER_PROCESS_QUEUE(RunLocal) \
 			UBA_SCHEDULER_PROCESS_QUEUE(RunRemote) \
 			UBA_SCHEDULER_PROCESS_QUEUE(RunCrossArchitecture) \
@@ -50,6 +51,7 @@ namespace uba
 		CacheFetchInfo(const ProcessStartInfo& i, Session& s, RootsHandle rh, MessagePriority mp) : CacheClient::FetchContext(i, s, rh, mp) {}
 		CacheResult result;
 		CacheClient* client;
+		u64 speculativeOutputId = 0;
 	};
 
 	////////////////////////////////////////////////////////////////////////////////////////////////////
@@ -65,11 +67,7 @@ namespace uba
 		{
 		}
 
-		~ProcessStartInfo2()
-		{
-			delete[] knownInputs;
-			delete[] knownOutputs;
-		}
+		~ProcessStartInfo2() { delete[] knownInputs; delete[] knownOutputs; }
 
 		CacheFetchInfo* cacheInfo = nullptr;
 
@@ -172,6 +170,7 @@ namespace uba
 
 			table->GetValueAsBool(writeTrackedInputs, TC("WriteTrackedInputs"));
 			table->GetValueAsBool(deferRemoteCapableProcesses, TC("DeferRemoteCapableProcesses"));
+			table->GetValueAsBool(enableSpeculativeCacheTest, TC("EnableSpeculativeCacheTest"));
 			table->GetValueAsBool(traceSchedulerQueues, TC("TraceSchedulerQueues"));
 
 			u32 temp = 0;
@@ -225,6 +224,7 @@ namespace uba
 	,	m_useThreadToExitRemoteProcess(info.useThreadToExitRemoteProcess)
 	,	m_writeTrackedInputs(info.writeTrackedInputs)
 	,	m_deferRemoteCapableProcesses(info.deferRemoteCapableProcesses)
+	,	m_enableSpeculativeCacheTest(info.enableSpeculativeCacheTest)
 	,	m_traceSchedulerQueues(info.traceSchedulerQueues)
 	,	m_processConfigs(info.processConfigs)
 	,	m_writeToCache(info.writeToCache && info.cacheClientCount)
@@ -410,6 +410,7 @@ namespace uba
 		entry.info = info2;
 		entry.firstDependent = nullptr;
 		entry.dependencyCount = 0;
+		entry.speculativeDependencyCount = 0;
 		entry.runStatus = RunStatus_QueuedForRun;
 		entry.cacheStatus = useCache ? CacheStatus_QueuedForTest : CacheStatus_Failed;
 		entry.canDetour = canDetour;
@@ -468,6 +469,7 @@ namespace uba
 	{
 		// Go through all the processes and see if we have dependencies and if we should just skip this
 		u32 dependencyCount = 0;
+		u32 speculativeDependencyCount = 0;
 		for (const u32* i=info.dependencies, *e=i+info.dependencyCount; i!=e; ++i)
 		{
 			auto& dependency = m_processEntries[*i];
@@ -478,8 +480,11 @@ namespace uba
 				AddToQueue(m_readyToSkip, entryIndex);
 				return;
 			}
-			else
-				++dependencyCount;
+			++dependencyCount;
+			// Speculative count only includes deps that haven't yet confirmed a cache hit.
+			// Deps at QueuedForDownload or beyond already fired UpdateSpeculativeDependentsOnCacheHit.
+			if (dependency.cacheStatus < CacheStatus_QueuedForDownload)
+				++speculativeDependencyCount;
 		}
 
 		if (!dependencyCount)
@@ -492,6 +497,7 @@ namespace uba
 		}
 
 		entry.dependencyCount = dependencyCount;
+		entry.speculativeDependencyCount = speculativeDependencyCount;
 
 		for (const u32* i=info.dependencies, *e=i+info.dependencyCount; i!=e; ++i)
 		{
@@ -500,32 +506,117 @@ namespace uba
 				continue;
 			dependency.firstDependent = new ProcessLink { dependency.firstDependent, entryIndex };
 		}
+
+		// If all deps already confirmed cache hits, speculatively test now even though deps are still downloading.
+		if (m_enableSpeculativeCacheTest && !speculativeDependencyCount && entry.cacheStatus == CacheStatus_QueuedForTest)
+		{
+			AddToQueue(m_readyToTestCache, entryIndex);
+			entry.cacheStatus = CacheStatus_QueuedForTestSpeculative;
+		}
 	}
 
-	void Scheduler::UpdateDependentsNoLock(RunStatus newRunStatus, ProcessLink* firstDependent)
+	void Scheduler::UpdateDependentsNoLock(RunStatus newRunStatus, ProcessLink* firstDependent, bool hadCacheHit)
 	{
 		bool shouldSkip = newRunStatus == RunStatus_Failed || newRunStatus == RunStatus_Skipped;
 		for (ProcessLink* it = firstDependent; it; it = it->next)
 		{
 			u32 dependentIndex = it->entryIndex;
 			auto& dependent = m_processEntries[dependentIndex];
+
+			// If completing dep never confirmed a cache hit, decrement the dependent's speculative count.
+			// If it did confirm a hit, UpdateSpeculativeDependentsOnCacheHit already decremented it.
+			if (m_enableSpeculativeCacheTest && !hadCacheHit && dependent.speculativeDependencyCount > 0)
+			{
+				if (--dependent.speculativeDependencyCount == 0
+					&& dependent.cacheStatus == CacheStatus_QueuedForTest
+					&& !dependent.shouldSkip && !shouldSkip)
+				{
+					AddToQueue(m_readyToTestCache, dependentIndex);
+					dependent.cacheStatus = CacheStatus_QueuedForTestSpeculative;
+				}
+			}
+
 			if (--dependent.dependencyCount)
 			{
 				if (shouldSkip)
 					dependent.shouldSkip = true;
 				continue;
 			}
-			if (newRunStatus == RunStatus_Success && !dependent.shouldSkip)
+
+			// dependencyCount hit zero
+			if (!shouldSkip && !dependent.shouldSkip && newRunStatus == RunStatus_Success)
 			{
-				if (dependent.cacheStatus == CacheStatus_QueuedForTest)
+				switch (dependent.cacheStatus)
+				{
+				case CacheStatus_QueuedForTest:
 					AddToQueue(m_readyToTestCache, dependentIndex);
-				else
+					break;
+				case CacheStatus_QueuedForTestSpeculative:
+					// Already in m_readyToTestCache; will be processed normally now that deps are done.
+					break;
+				case CacheStatus_Testing:
+				case CacheStatus_QueuedForDownload:
+				case CacheStatus_Downloading:
+					// Speculative work in flight — worker checks dependencyCount on completion.
+					break;
+				case CacheStatus_PendingCompletion:
+					// Speculative download finished while deps were pending; fire completion now.
+					AddToQueue(m_readyToCompleteFromCache, dependentIndex);
+					break;
+				default: // CacheStatus_Suspended (non-cacheable) or CacheStatus_Failed (cache miss)
 					AddToQueue(GetReadyToRunQueue(dependent), dependentIndex);
+					break;
+				}
 				continue;
 			}
-			if (dependent.runStatus == RunStatus_QueuedForRun)
+
+			// Dep failed/skipped, or this dependent itself should be skipped
+			if (dependent.cacheStatus == CacheStatus_Testing
+				|| dependent.cacheStatus == CacheStatus_QueuedForDownload
+				|| dependent.cacheStatus == CacheStatus_Downloading)
+			{
+				// Worker in flight — signal it to skip on completion.
+				dependent.shouldSkip = true;
+			}
+			else if (dependent.cacheStatus == CacheStatus_QueuedForTestSpeculative)
+			{
+				// In test queue but not yet started — remove and skip directly.
+				m_readyToTestCache.erase(dependentIndex);
+				AddToQueue(m_readyToSkip, dependentIndex);
+			}
+			else if (dependent.runStatus == RunStatus_QueuedForRun)
+			{
+				AddToQueue(m_readyToSkip, dependentIndex);
+			}
+			else if (dependent.cacheStatus == CacheStatus_PendingCompletion)
+			{
+				// Speculatively downloaded but a dep failed — treat as skipped.
+				dependent.runStatus = RunStatus_Skipped;
 				AddToQueue(m_readyToSkip, dependentIndex);
+			}
+		}
+	}
 
+	void Scheduler::UpdateSpeculativeDependentsOnCacheHit(ProcessLink* firstDependent)
+	{
+		// Called under m_processEntriesLock when an entry confirms a cache hit (-> QueuedForDownload).
+		// Decrements speculativeDependencyCount for each dependent; promotes to cache test when it hits zero.
+		for (ProcessLink* it = firstDependent; it; it = it->next)
+		{
+			u32 dependentIndex = it->entryIndex;
+			auto& dependent = m_processEntries[dependentIndex];
+			if (dependent.runStatus >= RunStatus_Success)
+				continue; // already terminal
+			if (dependent.speculativeDependencyCount == 0)
+				continue; // already at zero (already promoted, or never had speculative deps)
+			if (m_enableSpeculativeCacheTest
+				&& --dependent.speculativeDependencyCount == 0
+				&& dependent.cacheStatus == CacheStatus_QueuedForTest
+				&& !dependent.shouldSkip)
+			{
+				AddToQueue(m_readyToTestCache, dependentIndex);
+				dependent.cacheStatus = CacheStatus_QueuedForTestSpeculative;
+			}
 		}
 	}
 
@@ -535,13 +626,29 @@ namespace uba
 		SCOPED_FUTEX(m_processEntriesLock, lock);
 		for (auto& entry : m_processEntries)
 		{
-			if (entry.runStatus == RunStatus_QueuedForRun && (entry.cacheStatus == CacheStatus_QueuedForTest || entry.cacheStatus == CacheStatus_Failed))
+			if (entry.runStatus != RunStatus_QueuedForRun)
+				continue;
+
+			if (entry.cacheStatus == CacheStatus_Testing
+				|| entry.cacheStatus == CacheStatus_Downloading
+				|| entry.cacheStatus == CacheStatus_QueuedForDownload)
+			{
+				// Worker thread is in-flight. Set shouldSkip so the worker detects it on completion.
+				entry.shouldSkip = true;
+				continue;
+			}
+
+			if (entry.cacheStatus == CacheStatus_PendingCompletion)
 			{
-				entry.cacheStatus = CacheStatus_Failed;
-				entry.runStatus = RunStatus_Skipped;
-				// UpdateDependentsNoLock(RunStatus_Skipped, firstDependent); // No need to do this, all following entries are cancelled
-				skipped.push_back(entry.info);
+				UBA_ASSERT(entry.info->cacheInfo);
+				delete entry.info->cacheInfo;
+				entry.info->cacheInfo = nullptr;
 			}
+
+			entry.cacheStatus = CacheStatus_Failed;
+			entry.runStatus = RunStatus_Skipped;
+			// UpdateDependentsNoLock(RunStatus_Skipped, firstDependent); // No need to do this, all following entries are cancelled
+			skipped.push_back(entry.info);
 		}
 		ClearQueue(m_readyToTestCache);
 		ClearQueue(m_readyToDownloadCache);
@@ -550,6 +657,7 @@ namespace uba
 		ClearQueue(m_readyToRunCrossArchitecture);
 		ClearQueue(m_readyToRunCrossPlatform);
 		ClearQueue(m_readyToSkip);
+		ClearQueue(m_readyToCompleteFromCache);
 
 		lock.Leave();
 		for (auto pi : skipped)
@@ -588,6 +696,7 @@ namespace uba
 		u32 index = u32(m_processEntries.size());
 		auto& entry = m_processEntries.emplace_back();
 		entry.dependencyCount = 0;
+		entry.speculativeDependencyCount = 0;
 		entry.firstDependent = nullptr;
 		entry.runStatus = RunStatus_Suspended;
 		entry.cacheStatus = CacheStatus_Suspended;
@@ -608,7 +717,7 @@ namespace uba
 		if (m_cancelled)
 		{
 			entry.runStatus = RunStatus_Skipped;
-			UpdateDependentsNoLock(RunStatus_Skipped, firstDependent);
+			UpdateDependentsNoLock(RunStatus_Skipped, firstDependent, false); // newly created entry, no cache hit yet
 			entry.firstDependent = nullptr;
 			auto pi = entry.info;
 			entry.info = nullptr;
@@ -736,6 +845,9 @@ namespace uba
 
 			m_session.FlushDeadProcesses();
 
+			while (DrainCompleteFromCache())
+				;
+
 			while (RunCacheQuery())
 				;
 
@@ -841,7 +953,7 @@ namespace uba
 			FinishProcess(ph);
 
 			entry.firstDependent = nullptr;
-			UpdateDependentsNoLock(RunStatus_Skipped, firstDependent);
+			UpdateDependentsNoLock(RunStatus_Skipped, firstDependent, false); // running process never confirmed a cache hit
 			lock.Leave();
 			ProcessLink::DeleteAll(firstDependent);
 		}
@@ -859,16 +971,22 @@ namespace uba
 
 		--m_activeCacheQueries;
 
+		// Record whether the hit was previously confirmed (Downloading) before overwriting cacheStatus.
+		// If it was Testing→Failed the hit was never confirmed; speculative counts must be decremented.
+		// If it was Downloading→Failed the hit was confirmed; UpdateSpeculativeDependentsOnCacheHit
+		// already ran so we must not double-decrement.
+		bool hadCacheHit = (entry.cacheStatus == CacheStatus_Downloading);
+
 		entry.cacheStatus = CacheStatus_Failed;
 
-		if (m_cancelled)
+		if (m_cancelled || entry.shouldSkip)
 		{
 			entry.runStatus = RunStatus_Skipped;
 			auto pi = entry.info;
 			entry.info = nullptr;
 			ProcessLink* firstDependent = entry.firstDependent;
 			entry.firstDependent = nullptr;
-			UpdateDependentsNoLock(RunStatus_Skipped, firstDependent);
+			UpdateDependentsNoLock(RunStatus_Skipped, firstDependent, hadCacheHit);
 			lock.Leave();
 			ProcessLink::DeleteAll(firstDependent);
 			SkipProcess(*pi);
@@ -877,7 +995,13 @@ namespace uba
 		}
 
 		UBA_ASSERT(entry.runStatus == RunStatus_QueuedForRun);
-		AddToQueue(GetReadyToRunQueue(entry), processIndex);
+
+		// Only add to run queue if all dependencies have completed.
+		// With speculative cache testing, dependencyCount can be > 0 here; UpdateDependentsNoLock
+		// will promote to the run queue when the last dep completes (cacheStatus == Failed case).
+		if (entry.dependencyCount == 0)
+			AddToQueue(GetReadyToRunQueue(entry), processIndex);
+
 		lock.Leave();
 
 		m_updateThreadWakeup.Set();
@@ -1020,7 +1144,7 @@ namespace uba
 			auto& entry = m_processEntries[skipEntry];
 			outPrevStatus = entry.runStatus;
 			entry.runStatus = RunStatus_Skipped;
-			UpdateDependentsNoLock(RunStatus_Skipped, entry.firstDependent);
+			UpdateDependentsNoLock(RunStatus_Skipped, entry.firstDependent, false);
 			ProcessLink::DeleteAll(entry.firstDependent);
 			entry.firstDependent = nullptr;
 			return skipEntry;
@@ -1102,6 +1226,62 @@ namespace uba
 		return returnEntry(entry);
 	}
 
+	bool Scheduler::DrainCompleteFromCache()
+	{
+		SCOPED_FUTEX(m_processEntriesLock, entriesLock);
+		u32 idx = PopFromQueue(m_readyToCompleteFromCache);
+		if (idx == ~0u)
+			return false;
+
+		auto& entry = m_processEntries[idx];
+		UBA_ASSERT(entry.cacheStatus == CacheStatus_PendingCompletion);
+		UBA_ASSERT(entry.runStatus == RunStatus_QueuedForRun);
+		UBA_ASSERT(entry.dependencyCount == 0);
+		UBA_ASSERT(entry.info->cacheInfo);
+
+		entry.cacheStatus = CacheStatus_Success;
+		entry.runStatus = RunStatus_Success;
+		ProcessLink* firstDependent = entry.firstDependent;
+		entry.firstDependent = nullptr;
+		ProcessStartInfo2* si = entry.info;
+		entry.info = nullptr;
+		
+		entriesLock.Leave();
+
+		auto cacheInfo = si->cacheInfo;
+		si->cacheInfo = nullptr;
+		auto process = new CachedProcess(*si);
+		process->logLines.swap(cacheInfo->result.logLines);
+		delete cacheInfo;
+		ProcessHandle ph(process);
+		ProcessExitedResponse exitedResponse = ProcessExitedResponse_None;
+		if (auto func = si->exitedFunc)
+			func(si->userData, ph, exitedResponse);
+		UBA_ASSERT(exitedResponse == ProcessExitedResponse_None);
+
+		if (!m_cancelled)
+		{
+			SCOPED_FUTEX(m_processEntriesLock, lock);
+			UpdateDependentsNoLock(RunStatus_Success, firstDependent, true); // hadCacheHit=true
+		}
+
+		ProcessLink::DeleteAll(firstDependent);
+		UpdateActiveProcessCounter(true, -1);
+		m_updateThreadWakeup.Set();
+
+		if (m_writeTrackedInputs)
+		{
+			Vector<u8> inputs;
+			Vector<u8> outputs;
+			CreateCacheInputsAndOutputs(inputs, outputs, *si, *process);
+			WriteTrackedInputs(inputs,outputs, *process);
+		}
+
+		FinishProcess(ph);
+		delete si;
+		return true;
+	}
+
 	bool Scheduler::RunCacheQuery()
 	{
 		SCOPED_FUTEX(m_processEntriesLock, lock);
@@ -1118,7 +1298,7 @@ namespace uba
 		auto& entry = m_processEntries[entryToProcess];
 		auto info = entry.info;
 
-		if (entry.cacheStatus == CacheStatus_QueuedForTest)
+		if (entry.cacheStatus == CacheStatus_QueuedForTest || entry.cacheStatus == CacheStatus_QueuedForTestSpeculative)
 		{
 			entry.cacheStatus = CacheStatus_Testing;
 			++m_activeCacheQueries;
@@ -1146,9 +1326,16 @@ namespace uba
 						isHit = true;
 
 						SCOPED_FUTEX(m_processEntriesLock, lock);
-						m_processEntries[entryToProcess].cacheStatus = CacheStatus_QueuedForDownload;
+						auto& hitEntry = m_processEntries[entryToProcess];
+						hitEntry.cacheStatus = CacheStatus_QueuedForDownload;
 						AddToQueue(m_readyToDownloadCache, entryToProcess);
 						--m_activeCacheQueries;
+						// Register A's expected output CAS keys so speculative dependents
+						// see them instead of stale on-disk values during their cache test.
+						if (m_enableSpeculativeCacheTest)
+							cacheInfo->speculativeOutputId = cacheInfo->client->RegisterSpeculativeOutputs(*cacheInfo);
+						// Allow dependents to start speculative cache tests now that this entry has a confirmed hit.
+						UpdateSpeculativeDependentsOnCacheHit(hitEntry.firstDependent);
 						break;
 					}
 
@@ -1182,10 +1369,55 @@ namespace uba
 
 					bool success = cacheInfo->client->FetchFilesFromCache(cacheInfo->result, *cacheInfo, m_session.WritePlaceholders());
 					--m_activeCacheDownloads;
+					if (cacheInfo->speculativeOutputId)
+						cacheInfo->client->UnregisterSpeculativeOutputs(cacheInfo->speculativeOutputId);
 
 					if (success)
 					{
 						UpdateQueueCounter(-1, canExecuteRemotely); // Queue counter is only reduced when process is picked up for Running.. which means need to reduce it on cache success too
+
+						{
+							SCOPED_FUTEX(m_processEntriesLock, lock);
+							auto& dlEntry = m_processEntries[entryToProcess];
+
+							if (m_cancelled || dlEntry.shouldSkip)
+							{
+								// Build is being cancelled or a dep failed; skip the process.
+								--m_activeCacheQueries;
+								dlEntry.cacheStatus = CacheStatus_Success;
+								dlEntry.runStatus = RunStatus_Skipped;
+								ProcessLink* firstDependent = dlEntry.firstDependent;
+								dlEntry.firstDependent = nullptr;
+								dlEntry.info = nullptr;
+								UpdateDependentsNoLock(RunStatus_Skipped, firstDependent, true);
+								lock.Leave();
+								ProcessLink::DeleteAll(firstDependent);
+								SkipProcess(*info);
+								delete cacheInfo;
+								delete info;
+								++m_cacheHitCount;
+								UpdateStatusHitMissCount();
+								return;
+							}
+
+							if (dlEntry.dependencyCount > 0)
+							{
+								// Speculative path: dependencies are still in flight.
+								// Park cacheInfo (which holds the log lines) in info; DrainCompleteFromCache
+								// will create the CachedProcess and clean up when the last dep finishes.
+								--m_activeCacheQueries;
+								dlEntry.cacheStatus = CacheStatus_PendingCompletion;
+								lock.Leave();
+								m_updateThreadWakeup.Set();
+								++m_cacheHitCount;
+								UpdateStatusHitMissCount();
+								return;
+							}
+
+							// dependencyCount == 0: fall through to the existing ExitProcess path.
+							// ExitProcess will handle cacheStatus, m_activeCacheQueries, and completion.
+						}
+
 						auto process = new CachedProcess(si);
 						process->logLines.swap(cacheInfo->result.logLines);
 						ProcessHandle ph(process);
@@ -1494,10 +1726,13 @@ namespace uba
 		ProcessLink* firstDependent = entry.firstDependent;
 		if (isDone)
 		{
+			// hadCacheHit: if cacheStatus is Downloading or beyond, the hit was already confirmed and
+			// UpdateSpeculativeDependentsOnCacheHit already fired — don't double-decrement.
+			bool hadCacheHit = (entry.cacheStatus >= CacheStatus_QueuedForDownload);
 			entry.runStatus = exitCode == 0 ? RunStatus_Success : RunStatus_Failed;
 			entry.firstDependent = nullptr;
 			entry.info = nullptr;
-			UpdateDependentsNoLock(entry.runStatus, firstDependent);
+			UpdateDependentsNoLock(entry.runStatus, firstDependent, hadCacheHit);
 		}
 		else
 		{
@@ -1547,44 +1782,7 @@ namespace uba
 				CreateCacheInputsAndOutputs(inputs, outputs, *si, process);
 
 				if (writeTracked)
-				{
-					static int counter = 0;
-					StringBuffer<> inputsFile;
-					inputsFile.Append(m_session.GetSessionDir()).Append(TCV("inputs")).EnsureEndsWithSlash();
-					static u32 processId = 1; // TODO: This should be done in a better way.. or not at all?
-					GenerateNameForProcess(inputsFile, process.GetStartInfo().arguments, ++processId);
-					inputsFile.Append(TCV(".log"));
-					FileAccessor f(m_session.GetLogger(), inputsFile.data);
-					auto writeLine = [&](StringView str)
-						{
-							u8 buffer[1024];
-							BinaryWriter writer(buffer, 0, sizeof(buffer));
-							writer.WriteUtf8String(str.data, str.count);
-							writer.WriteByte('\n');
-							f.Write(buffer, writer.GetPosition());
-						};
-					auto writeLines = [&](Vector<u8>& files)
-						{
-							BinaryReader intputReader(files.data(), 0, files.size());
-							while (intputReader.GetLeft())
-							{
-								StringBuffer<> file;
-								intputReader.ReadString(file);
-								writeLine(file);
-							}
-						};
-					if (f.CreateWrite())
-					{
-						writeLine(TCV("Inputs:"));
-						writeLines(inputs);
-						writeLine(TCV(""));
-						writeLine(TCV("Outputs:"));
-						writeLines(outputs);
-						writeLine(TCV(""));
-
-						f.Close();
-					}
-				}
+					WriteTrackedInputs(inputs,outputs, process);
 
 				// This does only work if build step actually always write. Many tools read and check if file looks ok and in that case don't write
 				// So this is not a reliable approach..
@@ -1713,6 +1911,46 @@ namespace uba
 		outInputs.resize(writePosInputs);
 	}
 
+	void Scheduler::WriteTrackedInputs(const Vector<u8>& inputs, const Vector<u8>& outputs, Process& process)
+	{
+		static int counter = 0;
+		StringBuffer<> inputsFile;
+		inputsFile.Append(m_session.GetSessionDir()).Append(TCV("inputs")).EnsureEndsWithSlash();
+		static u32 processId = 1; // TODO: This should be done in a better way.. or not at all?
+		GenerateNameForProcess(inputsFile, process.GetStartInfo().arguments, ++processId);
+		inputsFile.Append(TCV(".log"));
+		FileAccessor f(m_session.GetLogger(), inputsFile.data);
+		auto writeLine = [&](StringView str)
+			{
+				u8 buffer[1024];
+				BinaryWriter writer(buffer, 0, sizeof(buffer));
+				writer.WriteUtf8String(str.data, str.count);
+				writer.WriteByte('\n');
+				f.Write(buffer, writer.GetPosition());
+			};
+		auto writeLines = [&](const Vector<u8>& files)
+			{
+				BinaryReader intputReader(files.data(), 0, files.size());
+				while (intputReader.GetLeft())
+				{
+					StringBuffer<> file;
+					intputReader.ReadString(file);
+					writeLine(file);
+				}
+			};
+		if (f.CreateWrite())
+		{
+			writeLine(TCV("Inputs:"));
+			writeLines(inputs);
+			writeLine(TCV(""));
+			writeLine(TCV("Outputs:"));
+			writeLines(outputs);
+			writeLine(TCV(""));
+
+			f.Close();
+		}
+	}
+
 	void Scheduler::KillNewestLocalProcess(StringView reason)
 	{
 		ProcessHandle ph;
diff --git a/Engine/Source/Programs/UnrealBuildAccelerator/Common/Public/UbaCacheClient.h b/Engine/Source/Programs/UnrealBuildAccelerator/Common/Public/UbaCacheClient.h
--- a/Engine/Source/Programs/UnrealBuildAccelerator/Common/Public/UbaCacheClient.h
+++ b/Engine/Source/Programs/UnrealBuildAccelerator/Common/Public/UbaCacheClient.h
@@ -106,6 +106,9 @@ namespace uba
 		bool FetchEntryFromCache(CacheResult& outResult, FetchContext& context, u32 bucketId);
 		bool FetchFilesFromCache(CacheResult& outResult, FetchContext& context, bool writePlaceholders);
 
+		u64 RegisterSpeculativeOutputs(FetchContext& context);
+		void UnregisterSpeculativeOutputs(u64 id);
+
 	private:
 		struct FetchScope;
 
@@ -159,6 +162,11 @@ namespace uba
 		struct PathHash { TString path; CasKey hash; };
 		Vector<PathHash> m_pathHashes;
 
+		Futex m_speculativeCasKeysLock;
+		UnorderedMap<StringKey, CasKey> m_speculativeCasKeys;
+		UnorderedMap<u64, Vector<StringKey>> m_speculativeOutputsById;
+		u64 m_nextSpeculativeOutputId = 1;
+
 		#if UBA_CACHE_USE_SEND_MAX_BUFFER_CACHE
 		BufferCache m_sendBufferSlots;
 		#endif
diff --git a/Engine/Source/Programs/UnrealBuildAccelerator/Common/Public/UbaScheduler.h b/Engine/Source/Programs/UnrealBuildAccelerator/Common/Public/UbaScheduler.h
--- a/Engine/Source/Programs/UnrealBuildAccelerator/Common/Public/UbaScheduler.h
+++ b/Engine/Source/Programs/UnrealBuildAccelerator/Common/Public/UbaScheduler.h
@@ -40,6 +40,7 @@ namespace uba
 		bool memWatchdog = false; // Will watch memory and throttle/kill processes if needed
 		bool writeTrackedInputs = false; // Set to true to write out log files with tracked inputs
 		bool deferRemoteCapableProcesses = false; // Set to true to have local session first pick local-only processes
+		bool enableSpeculativeCacheTest = false; // Set to true to allow cache testing of a process before its dependencies have finished downloading. Requires CacheClient support for speculative CAS keys to be correct.
 		bool traceSchedulerQueues = false; // Show in visualizer how many processes that are queued up for state change
 		u32 maxRacingPercent = 0u; // If processor weights go below this percent, then we will start filling up with racing processes
 		u8 memTracingLevel = 0;
@@ -135,10 +136,12 @@ namespace uba
 		enum CacheStatus : u8
 		{
 			CacheStatus_Suspended,
-			CacheStatus_QueuedForTest,
-			CacheStatus_Testing,
-			CacheStatus_QueuedForDownload,
-			CacheStatus_Downloading,
+			CacheStatus_QueuedForTest,             // not yet in any queue; blocked on deps or awaiting initial scheduling
+			CacheStatus_QueuedForTestSpeculative,  // in m_readyToTestCache; deps still pending (speculative start)
+			CacheStatus_Testing,                   // test worker in flight
+			CacheStatus_QueuedForDownload,         // hit confirmed; in m_readyToDownloadCache
+			CacheStatus_Downloading,               // download worker in flight
+			CacheStatus_PendingCompletion,         // download done; deps still pending (will fire when last dep completes)
 			CacheStatus_Success,
 			CacheStatus_Failed,
 		};
@@ -157,7 +160,9 @@ namespace uba
 		u32 GetQueueCount(ProcessQueue& queue);
 		void ClearQueue(ProcessQueue& queue);
 		void SetupDependents(ProcessEntry& entry, u32 entryIndex, const EnqueueProcessInfo& info);
-		void UpdateDependentsNoLock(RunStatus newRunStatus, ProcessLink* firstDependent);
+		void UpdateDependentsNoLock(RunStatus newRunStatus, ProcessLink* firstDependent, bool hadCacheHit);
+		void UpdateSpeculativeDependentsOnCacheHit(ProcessLink* firstDependent);
+		bool DrainCompleteFromCache();
 		void SkipAllQueued();
 		bool ProcessReturned(Process& process, bool isLocal);
 		void HandleCacheMissed(u32 processIndex);
@@ -176,6 +181,7 @@ namespace uba
 		void UpdateActiveProcessCounter(bool isLocal, int offset);
 		void FinishProcess(const ProcessHandle& handle);
 		void CreateCacheInputsAndOutputs(Vector<u8>& outInputs, Vector<u8>& outOutputs, ProcessStartInfo2& info, Process& process);
+		void WriteTrackedInputs(const Vector<u8>& inputs, const Vector<u8>& outputs, Process& process);
 		void KillNewestLocalProcess(StringView reason);
 		void NativeProcessCreated(ProcessImpl& process);
 		void LocalProcessStart(ProcessStartInfo2& info);
@@ -190,6 +196,7 @@ namespace uba
 			ProcessStartInfo2* info;
 			ProcessLink* firstDependent;
 			u32 dependencyCount;
+			u32 speculativeDependencyCount;
 			CacheStatus cacheStatus;
 			RunStatus runStatus;
 			bool shouldSkip;
@@ -216,6 +223,7 @@ namespace uba
 		ProcessQueue m_readyToRunCrossArchitecture; // Can run remotely and also cross architecture
 		ProcessQueue m_readyToRunCrossPlatform; // Can run remotely and also cross platform
 		ProcessQueue m_readyToSkip;
+		ProcessQueue m_readyToCompleteFromCache; // Cache download done but deps were still pending; fired when last dep completes
 
 		Function<void(const ProcessHandle&)> m_processFinished;
 
@@ -230,6 +238,7 @@ namespace uba
 		bool m_allowDisableRemoteExecution = false;
 		bool m_writeTrackedInputs = false;
 		bool m_deferRemoteCapableProcesses = false;
+		bool m_enableSpeculativeCacheTest = false;
 		bool m_traceSchedulerQueues = false;
 		bool m_cancelled = false;
 		bool m_hadRemoteRequest = false;
