[UBA]
* Moved code ensuring that enough data is fetched of caskey and path tables for outputs on cache test hit. This makes sure that speculative testing never reads outside tables

diff --git a/Engine/Source/Programs/UnrealBuildAccelerator/Common/Private/UbaCacheClient.cpp b/Engine/Source/Programs/UnrealBuildAccelerator/Common/Private/UbaCacheClient.cpp
--- a/Engine/Source/Programs/UnrealBuildAccelerator/Common/Private/UbaCacheClient.cpp
+++ b/Engine/Source/Programs/UnrealBuildAccelerator/Common/Private/UbaCacheClient.cpp
@@ -248,8 +248,10 @@ namespace uba
 			StringBuffer<MaxPath> path;
 			CasKey casKey;
 			if (!GetLocalPathAndCasKey(bucket, context.rootPaths, path, casKey, bucket.serverCasKeyTable, bucket.serverPathTable, offset))
+			{
+				UBA_ASSERT(false);
 				continue;
-
+			}
 			StringBuffer<MaxPath> forKey;
 			forKey.Append(path);
 			if (CaseInsensitiveFs)
@@ -809,8 +811,8 @@ namespace uba
 		reader.Skip(sharedLogLinesSize);
 
 		// Calculate max caskey offset and fetch tables (faster to do in one call)
-		u64 maxCasKeyOffset = 0;
 		{
+			u64 maxCasKeyOffset = 0;
 			BinaryReader sharedReader(sharedPos, 0, sharedSize);
 			while (sharedReader.GetLeft())
 				maxCasKeyOffset = Max(maxCasKeyOffset, sharedReader.Read7BitEncoded());
@@ -994,6 +996,26 @@ namespace uba
 					context.casKeyOffsets.push_back(u32(outputsReader.Read7BitEncoded()));
 				context.bucket = &bucket;
 
+				// Make sure caskey/path tables has enough data for outputs
+				if (!context.casKeyOffsets.empty())
+				{
+					u64 maxCasKeyOffset = 0;
+					for (u64 offset : context.casKeyOffsets)
+						maxCasKeyOffset = Max(maxCasKeyOffset, offset);
+					if (!FetchCasTable(tws, bucket, cacheStats, u32(maxCasKeyOffset)))
+						return false;
+
+					u64 maxPathTableOffset = 0;
+					BinaryReader casKeyReader(bucket.serverCasKeyTable.GetMemory(), 0, bucket.serverCasKeyTableSize);
+					for (u64 offset : context.casKeyOffsets)
+					{
+						casKeyReader.SetPosition(offset);
+						maxPathTableOffset = Max(maxPathTableOffset, casKeyReader.Read7BitEncoded());
+					}
+					if (!FetchPathTable(tws, bucket, cacheStats, u32(maxPathTableOffset)))
+						return false;
+				}
+
 				outResult.hit = true;
 				success = true;
 				return true;
@@ -1008,9 +1030,6 @@ namespace uba
 
 	bool CacheClient::FetchFilesFromCache(CacheResult& outResult, FetchContext& context, bool writePlaceholders)
 	{
-		UBA_ASSERT(context.bucket);
-		auto& bucket = *context.bucket;
-
 		if (context.casKeyOffsets.empty())
 			return true;
 
@@ -1024,25 +1043,6 @@ namespace uba
 		u8 memory[1024];
 		FetchScope scope(m_session, success, memory, context.info.GetDescription());
 
-		// Make sure caskey/path tables has enough data
-		{
-			u64 maxCasKeyOffset = 0;
-			for (u64 offset : context.casKeyOffsets)
-				maxCasKeyOffset = Max(maxCasKeyOffset, offset);
-			if (!FetchCasTable(tws, bucket, scope.cacheStats, u32(maxCasKeyOffset)))
-				return false;
-
-			u64 maxPathTableOffset = 0;
-			BinaryReader casKeyReader(bucket.serverCasKeyTable.GetMemory(), 0, bucket.serverCasKeyTableSize);
-			for (u64 offset : context.casKeyOffsets)
-			{
-				casKeyReader.SetPosition(offset);
-				maxPathTableOffset = Max(maxPathTableOffset, casKeyReader.Read7BitEncoded());
-			}
-			if (!FetchPathTable(tws, bucket, scope.cacheStats, u32(maxPathTableOffset)))
-				return false;
-		}
-
 		// Go wide on multiple files
 		Atomic<bool> fetchSuccess = true;
 		m_storage.m_workManager.ParallelFor(u32(context.casKeyOffsets.size() - 1), context.casKeyOffsets, [&](const WorkContext&, auto& it)
