[UBA]
* Added option to enable threaded send on NetworkClient

diff --git a/Engine/Source/Programs/UnrealBuildAccelerator/Common/Private/UbaNetworkClient.cpp b/Engine/Source/Programs/UnrealBuildAccelerator/Common/Private/UbaNetworkClient.cpp
--- a/Engine/Source/Programs/UnrealBuildAccelerator/Common/Private/UbaNetworkClient.cpp
+++ b/Engine/Source/Programs/UnrealBuildAccelerator/Common/Private/UbaNetworkClient.cpp
@@ -20,6 +20,7 @@ namespace uba
 		const ConfigTable& table = *tablePtr;
 		table.GetValueAsU32(recvBufferSize, TC("RecvBufferSize"));
 		table.GetValueAsU32(sendBufferSize, TC("SendBufferSize"));
+		table.GetValueAsBool(useThreadedSend, TC("UseThreadedSend"));
 	}
 
 	struct NetworkClient::Connection
@@ -55,6 +56,7 @@ namespace uba
 		m_sendSize = fixedSendSize;
 		m_receiveTimeoutSeconds = info.receiveTimeoutSeconds;
 		m_useMessagePriority = info.useMessagePriority;
+		m_useThreadedSend = info.useThreadedSend;
 		m_recvBufferSize = info.recvBufferSize;
 		m_sendBufferSize = info.sendBufferSize;
 
@@ -128,6 +130,9 @@ namespace uba
 				backend.SetSendBufferSize(backendConnection, m_sendBufferSize);
 		}
 
+		if (m_useThreadedSend)
+			backend.SetThreadedSend(backendConnection);
+
 		RecvContext rc(*this, backend, backendConnection);
 
 		// The only way out of this function is to get a call to one of the below callbacks since exitScopeEvent must be set.
diff --git a/Engine/Source/Programs/UnrealBuildAccelerator/Common/Public/UbaNetworkClient.h b/Engine/Source/Programs/UnrealBuildAccelerator/Common/Public/UbaNetworkClient.h
--- a/Engine/Source/Programs/UnrealBuildAccelerator/Common/Public/UbaNetworkClient.h
+++ b/Engine/Source/Programs/UnrealBuildAccelerator/Common/Public/UbaNetworkClient.h
@@ -27,6 +27,7 @@ namespace uba
 		u32 receiveTimeoutSeconds = DefaultNetworkReceiveTimeoutSeconds;
 		u32 workerCount = ~0u; // ~0u means logical processor count will be used
 		bool useMessagePriority = false;
+		bool useThreadedSend = false;
 		const u8* cryptoKey128 = nullptr;
 
 		u32 recvBufferSize = 0;
@@ -111,6 +112,7 @@ namespace uba
 		Event m_isOrWasConnected;
 		Atomic<u32> m_connectionCount = 0;
 		bool m_useMessagePriority;
+		bool m_useThreadedSend;
 
 		Futex m_onConnectedFunctionsLock;
 		Vector<OnConnectedFunction> m_onConnectedFunctions;
