[UBA]
* Added detouring of __p__wpgmptr to make sure cvtres.exe outputs proper virtual paths when run remotely

diff --git a/Engine/Source/Programs/UnrealBuildAccelerator/Detours/Private/Windows/UbaDetoursFunctionsUcrtBase.inl b/Engine/Source/Programs/UnrealBuildAccelerator/Detours/Private/Windows/UbaDetoursFunctionsUcrtBase.inl
--- a/Engine/Source/Programs/UnrealBuildAccelerator/Detours/Private/Windows/UbaDetoursFunctionsUcrtBase.inl
+++ b/Engine/Source/Programs/UnrealBuildAccelerator/Detours/Private/Windows/UbaDetoursFunctionsUcrtBase.inl
@@ -118,6 +118,16 @@ errno_t Detoured__get_wpgmptr(wchar_t** pValue)
 	return 0;
 }
 
+static WCHAR* g_fake_wpgmptr = g_virtualApplication.data;
+
+WCHAR** Detoured___p__wpgmptr()
+{
+	DETOURED_CALL(__p__wpgmptr);
+	DEBUG_LOG_DETOURED(L"__p__wpgmptr", L"(%ls)", g_virtualApplication.data);
+	return &g_fake_wpgmptr;
+}
+
+
 #if UBA_DEBUG_LOG_ENABLED
 const wchar_t* WaccessResultToString(errno_t res)
 {
diff --git a/Engine/Source/Programs/UnrealBuildAccelerator/Detours/Private/Windows/UbaDetoursFunctionsWin.h b/Engine/Source/Programs/UnrealBuildAccelerator/Detours/Private/Windows/UbaDetoursFunctionsWin.h
--- a/Engine/Source/Programs/UnrealBuildAccelerator/Detours/Private/Windows/UbaDetoursFunctionsWin.h
+++ b/Engine/Source/Programs/UnrealBuildAccelerator/Detours/Private/Windows/UbaDetoursFunctionsWin.h
@@ -184,6 +184,7 @@
 	DETOURED_FUNCTION(_wfullpath) \
 	DETOURED_FUNCTION(_fullpath) \
 	DETOURED_FUNCTION(_get_wpgmptr) \
+	DETOURED_FUNCTION(__p__wpgmptr) \
 	DETOURED_FUNCTION(_waccess_s) \
 	DETOURED_FUNCTION(_wspawnl) \
 	DETOURED_FUNCTION(_get_osfhandle) \
