[UBA]
* Added option for time64 to always return 0 and changed so for cvtres.exe it always return 0. This to make output obj file be determinstic (there are no options to do it from the outside)

diff --git a/Engine/Source/Programs/UnrealBuildAccelerator/Core/Private/UbaApplicationRules.cpp b/Engine/Source/Programs/UnrealBuildAccelerator/Core/Private/UbaApplicationRules.cpp
--- a/Engine/Source/Programs/UnrealBuildAccelerator/Core/Private/UbaApplicationRules.cpp
+++ b/Engine/Source/Programs/UnrealBuildAccelerator/Core/Private/UbaApplicationRules.cpp
@@ -240,14 +240,21 @@ namespace uba
 	{
 		using Super = ApplicationRulesVcLink;
 
-		virtual bool CloseFileMappingAfterUse(StringView fileName) const
+		virtual bool CloseFileMappingAfterUse(StringView fileName) const override
 		{
 			return true;
 		}
 
 		virtual bool SupportsPlaceholder(StringView fileName) const override { return SupportsLinkerPlaceholder(fileName); }
+
+		virtual bool Time64AlwaysReturnZero() const override
+		{
+			return true;
+		}
 	};
 
+	class ApplicationCvtresExe : public ApplicationRulesLinkExe { virtual bool Time64AlwaysReturnZero() const override { return true; } };
+
 	class ApplicationRulesLldLinkExe : public ApplicationRulesVcLink
 	{
 		using Super = ApplicationRulesVcLink;
@@ -797,7 +804,7 @@ namespace uba
 			{ ApplicationHash(TC("cl.exe")),					new ApplicationRulesClExe() },	 // Must be index 2
 			{ ApplicationHash(TC("link.exe")),					new ApplicationRulesLinkExe() }, // Must be index 3
 			{ ApplicationHash(TC("lib.exe")),					new ApplicationRulesVcLink() },
-			{ ApplicationHash(TC("cvtres.exe")),				new ApplicationRulesLinkExe() },
+			{ ApplicationHash(TC("cvtres.exe")),				new ApplicationCvtresExe() },
 			{ ApplicationHash(TC("mt.exe")),					new ApplicationRulesVcLink() },
 			{ ApplicationHash(TC("rc.exe")),					new ApplicationRulesVcLink() },
 			{ ApplicationHash(TC("lld-link.exe")),				new ApplicationRulesLldLinkExe() },
diff --git a/Engine/Source/Programs/UnrealBuildAccelerator/Core/Public/UbaApplicationRules.h b/Engine/Source/Programs/UnrealBuildAccelerator/Core/Public/UbaApplicationRules.h
--- a/Engine/Source/Programs/UnrealBuildAccelerator/Core/Public/UbaApplicationRules.h
+++ b/Engine/Source/Programs/UnrealBuildAccelerator/Core/Public/UbaApplicationRules.h
@@ -210,6 +210,11 @@ namespace uba
 			return false;
 		}
 
+		virtual bool Time64AlwaysReturnZero() const
+		{
+			return false;
+		}
+
 		u32 index = ~0u;
 	};
 
diff --git a/Engine/Source/Programs/UnrealBuildAccelerator/Detours/Private/Windows/UbaDetoursFunctionsUcrtBase.inl b/Engine/Source/Programs/UnrealBuildAccelerator/Detours/Private/Windows/UbaDetoursFunctionsUcrtBase.inl
--- a/Engine/Source/Programs/UnrealBuildAccelerator/Detours/Private/Windows/UbaDetoursFunctionsUcrtBase.inl
+++ b/Engine/Source/Programs/UnrealBuildAccelerator/Detours/Private/Windows/UbaDetoursFunctionsUcrtBase.inl
@@ -127,6 +127,17 @@ WCHAR** Detoured___p__wpgmptr()
 	return &g_fake_wpgmptr;
 }
 
+__time64_t Detoured__time64(__time64_t* destTime)
+{
+	DETOURED_CALL(_time64);
+	if (g_rules->Time64AlwaysReturnZero())
+	{
+		if (destTime)
+			*destTime = 0;
+		return 0;
+	}
+	return True__time64(destTime);
+}
 
 #if UBA_DEBUG_LOG_ENABLED
 const wchar_t* WaccessResultToString(errno_t res)
diff --git a/Engine/Source/Programs/UnrealBuildAccelerator/Detours/Private/Windows/UbaDetoursFunctionsWin.h b/Engine/Source/Programs/UnrealBuildAccelerator/Detours/Private/Windows/UbaDetoursFunctionsWin.h
--- a/Engine/Source/Programs/UnrealBuildAccelerator/Detours/Private/Windows/UbaDetoursFunctionsWin.h
+++ b/Engine/Source/Programs/UnrealBuildAccelerator/Detours/Private/Windows/UbaDetoursFunctionsWin.h
@@ -185,6 +185,7 @@
 	DETOURED_FUNCTION(_fullpath) \
 	DETOURED_FUNCTION(_get_wpgmptr) \
 	DETOURED_FUNCTION(__p__wpgmptr) \
+	DETOURED_FUNCTION(_time64) \
 	DETOURED_FUNCTION(_waccess_s) \
 	DETOURED_FUNCTION(_wspawnl) \
 	DETOURED_FUNCTION(_get_osfhandle) \
@@ -531,6 +532,7 @@ extern "C" {
 	HRESULT SHGetFolderPathW(HWND hwnd, int csidl, HANDLE hToken, DWORD dwFlags, LPWSTR pszPath);
 	enum FILE_INFO_BY_NAME_CLASS{};
 	BOOL WINAPI GetFileInformationByName(PCWSTR FileName, FILE_INFO_BY_NAME_CLASS FileInformationClass, PVOID FileInfoBuffer, ULONG FileInfoBufferSize);
+	__time64_t WINAPI _time64(__time64_t* destTime);
 
 }
 
