[UBA]
* Fixed so cache test queue always have room in the cache work manager

diff --git a/Engine/Source/Programs/UnrealBuildAccelerator/Common/Private/UbaScheduler.cpp b/Engine/Source/Programs/UnrealBuildAccelerator/Common/Private/UbaScheduler.cpp
--- a/Engine/Source/Programs/UnrealBuildAccelerator/Common/Private/UbaScheduler.cpp
+++ b/Engine/Source/Programs/UnrealBuildAccelerator/Common/Private/UbaScheduler.cpp
@@ -969,6 +969,9 @@ namespace uba
 		auto& entry = m_processEntries[processIndex];
 		UBA_ASSERTF((entry.cacheStatus == CacheStatus_Testing || entry.cacheStatus == CacheStatus_Downloading) && entry.runStatus == RunStatus_QueuedForRun, TC("Unexpected entry state. Run: %u Cache: %u"), u32(entry.runStatus), u32(entry.cacheStatus));
 
+		if (entry.cacheStatus == CacheStatus_Downloading)
+			--m_activeCacheDownloads;
+
 		--m_activeCacheQueries;
 
 		// Record whether the hit was previously confirmed (Downloading) before overwriting cacheStatus.
@@ -1186,7 +1189,7 @@ namespace uba
 				if (ProcessEntry* entry = PeekFromQueue(m_readyToRunLocal)) // Check local only first
 					if (entry->info->weight <= weightLeft)
 						return PopFromQueue(m_readyToRunLocal);
-				if (!m_forceRemote || (m_hadRemoteRequest && m_hadCrossSystemRequest)) // If force remote but we've had cross system requests we let local machine pop non-cross system queue
+				if (!m_forceRemote || (!m_hadRemoteRequest && m_hadCrossSystemRequest)) // If force remote but we've had cross system requests we let local machine pop non-cross system queue
 					if (ProcessEntry* entry = PeekFromQueue(m_readyToRunRemote))
 						if (entry->info->weight <= weightLeft)
 							return PopFromQueue(m_readyToRunRemote);
@@ -1289,11 +1292,27 @@ namespace uba
 		if (m_activeCacheQueries >= m_maxParallelCacheQueries)
 			return false;
 
-		u32 entryToProcess = PopFromQueue(m_readyToTestCache);
-		if (entryToProcess == ~0u)
+		u32 entryToProcess;
+
+		// Some hard coded constants here.
+		// First, we don't want to look ahead too far since it will take a while and we want to get bandwidth maxed out earlier..
+		// So after 3000 ready to download we start prioritizing actual downloading over testing. Downside is that it might take longer to find cache misses, but worth it
+		if (m_readyToDownloadCache.size() > 3000)
 			entryToProcess = PopFromQueue(m_readyToDownloadCache);
-		if (entryToProcess == ~0u)
-			return false;
+		else
+		{
+			entryToProcess = PopFromQueue(m_readyToTestCache);
+			if (entryToProcess == ~0u)
+			{
+				// We always reserve two entries in the worker for testing when speculative testing is enabled. Because sometimes there are bottlenecks where testing will happen sequentially
+				// and if we don't reserve these slots we might catch up with testing which we don't want in order to keep bandwidth maxed out
+				if (m_enableSpeculativeCacheTest && m_activeCacheDownloads >= m_maxParallelCacheQueries - 2)
+					return false;
+				entryToProcess = PopFromQueue(m_readyToDownloadCache);
+			}
+			if (entryToProcess == ~0u)
+				return false;
+		}
 
 		auto& entry = m_processEntries[entryToProcess];
 		auto info = entry.info;
@@ -1350,16 +1369,15 @@ namespace uba
 					++m_cacheMissCount;
 					UpdateStatusHitMissCount();
 
-				}, 1, TC("UbaSchedTest"));
+				}, 1, TC("UbaSchedTest"), ColorWork, WorkPriority_High);
 		}
 		else if (entry.cacheStatus == CacheStatus_QueuedForDownload)
 		{
 			entry.cacheStatus = CacheStatus_Downloading;
 			++m_activeCacheQueries;
 			bool canExecuteRemotely = entry.canExecuteRemotely;
-			lock.Leave();
-
 			++m_activeCacheDownloads;
+			lock.Leave();
 
 			m_cacheWorkManager->AddWork([this, info, entryToProcess, canExecuteRemotely](const WorkContext&)
 				{
@@ -1368,7 +1386,6 @@ namespace uba
 					ProcessStartInfo2& si = *info;
 
 					bool success = cacheInfo->client->FetchFilesFromCache(cacheInfo->result, *cacheInfo, m_session.WritePlaceholders());
-					--m_activeCacheDownloads;
 					if (cacheInfo->speculativeOutputId)
 						cacheInfo->client->UnregisterSpeculativeOutputs(cacheInfo->speculativeOutputId);
 
@@ -1376,48 +1393,49 @@ namespace uba
 					{
 						UpdateQueueCounter(-1, canExecuteRemotely); // Queue counter is only reduced when process is picked up for Running.. which means need to reduce it on cache success too
 
-						{
-							SCOPED_FUTEX(m_processEntriesLock, lock);
-							auto& dlEntry = m_processEntries[entryToProcess];
+						SCOPED_FUTEX(m_processEntriesLock, lock);
+						auto& entry = m_processEntries[entryToProcess];
 
-							if (m_cancelled || dlEntry.shouldSkip)
-							{
-								// Build is being cancelled or a dep failed; skip the process.
-								--m_activeCacheQueries;
-								dlEntry.cacheStatus = CacheStatus_Success;
-								dlEntry.runStatus = RunStatus_Skipped;
-								ProcessLink* firstDependent = dlEntry.firstDependent;
-								dlEntry.firstDependent = nullptr;
-								dlEntry.info = nullptr;
-								UpdateDependentsNoLock(RunStatus_Skipped, firstDependent, true);
-								lock.Leave();
-								ProcessLink::DeleteAll(firstDependent);
-								SkipProcess(*info);
-								delete cacheInfo;
-								delete info;
-								++m_cacheHitCount;
-								UpdateStatusHitMissCount();
-								return;
-							}
+						--m_activeCacheDownloads;
 
-							if (dlEntry.dependencyCount > 0)
-							{
-								// Speculative path: dependencies are still in flight.
-								// Park cacheInfo (which holds the log lines) in info; DrainCompleteFromCache
-								// will create the CachedProcess and clean up when the last dep finishes.
-								--m_activeCacheQueries;
-								dlEntry.cacheStatus = CacheStatus_PendingCompletion;
-								lock.Leave();
-								m_updateThreadWakeup.Set();
-								++m_cacheHitCount;
-								UpdateStatusHitMissCount();
-								return;
-							}
+						if (m_cancelled || entry.shouldSkip)
+						{
+							// Build is being cancelled or a dep failed; skip the process.
+							--m_activeCacheQueries;
+							entry.cacheStatus = CacheStatus_Success;
+							entry.runStatus = RunStatus_Skipped;
+							ProcessLink* firstDependent = entry.firstDependent;
+							entry.firstDependent = nullptr;
+							entry.info = nullptr;
+							UpdateDependentsNoLock(RunStatus_Skipped, firstDependent, true);
+							lock.Leave();
+							ProcessLink::DeleteAll(firstDependent);
+							SkipProcess(*info);
+							delete cacheInfo;
+							delete info;
+							++m_cacheHitCount;
+							UpdateStatusHitMissCount();
+							return;
+						}
 
-							// dependencyCount == 0: fall through to the existing ExitProcess path.
-							// ExitProcess will handle cacheStatus, m_activeCacheQueries, and completion.
+						if (entry.dependencyCount > 0)
+						{
+							// Speculative path: dependencies are still in flight.
+							// Park cacheInfo (which holds the log lines) in info; DrainCompleteFromCache
+							// will create the CachedProcess and clean up when the last dep finishes.
+							--m_activeCacheQueries;
+							entry.cacheStatus = CacheStatus_PendingCompletion;
+							lock.Leave();
+							m_updateThreadWakeup.Set();
+							++m_cacheHitCount;
+							UpdateStatusHitMissCount();
+							return;
 						}
 
+						// dependencyCount == 0: fall through to the existing ExitProcess path.
+						// ExitProcess will handle cacheStatus, m_activeCacheQueries, and completion.
+						lock.Leave();
+
 						auto process = new CachedProcess(si);
 						process->logLines.swap(cacheInfo->result.logLines);
 						ProcessHandle ph(process);
diff --git a/Engine/Source/Programs/UnrealBuildAccelerator/Common/Public/UbaScheduler.h b/Engine/Source/Programs/UnrealBuildAccelerator/Common/Public/UbaScheduler.h
--- a/Engine/Source/Programs/UnrealBuildAccelerator/Common/Public/UbaScheduler.h
+++ b/Engine/Source/Programs/UnrealBuildAccelerator/Common/Public/UbaScheduler.h
@@ -248,7 +248,7 @@ namespace uba
 		Atomic<float> m_activeLocalProcessWeight = 0.0f;
 		Atomic<float> m_activeRacingProcessWeight = 0.0f;
 		u32 m_activeCacheQueries = 0;
-		Atomic<u32> m_activeCacheDownloads = 0;
+		u32 m_activeCacheDownloads = 0;
 		Atomic<u32> m_getNextQueueCount = 0;
 
 		Atomic<u32> m_totalProcesses = 0; // Process count registed in to scheduler
