[UBA]
* Added CrossArchAndPlatform queue where processes that can run both cross architecture and cross platform is enqueued
* Improved stats tracking code in Scheduler. Added GetStats function that return all the useful information needed to be able to make good coordinator decisions

diff --git a/Engine/Source/Programs/UnrealBuildAccelerator/Common/Private/UbaScheduler.cpp b/Engine/Source/Programs/UnrealBuildAccelerator/Common/Private/UbaScheduler.cpp
--- a/Engine/Source/Programs/UnrealBuildAccelerator/Common/Private/UbaScheduler.cpp
+++ b/Engine/Source/Programs/UnrealBuildAccelerator/Common/Private/UbaScheduler.cpp
@@ -1,24 +1,22 @@
 // Copyright Epic Games, Inc. All Rights Reserved.
 
 #include "UbaScheduler.h"
-
-#if PLATFORM_WINDOWS
-#include <winternl.h>
-#endif
-
 #include "UbaApplicationRules.h"
 #include "UbaCacheClient.h"
 #include "UbaConfig.h"
 #include "UbaEnvironment.h"
 #include "UbaFileAccessor.h"
 #include "UbaNetworkMessage.h"
-#include "UbaNetworkServer.h"
 #include "UbaProcess.h"
 #include "UbaProcessStartInfoHolder.h"
 #include "UbaRootPaths.h"
 #include "UbaSessionServer.h"
 #include "UbaStringBuffer.h"
 
+#if PLATFORM_WINDOWS
+#include <winternl.h>
+#endif
+
 #define UBA_SCHEDULER_PROCESS_QUEUES \
 			UBA_SCHEDULER_PROCESS_QUEUE(TestCache) \
 			UBA_SCHEDULER_PROCESS_QUEUE(DownloadCache) \
@@ -27,6 +25,7 @@
 			UBA_SCHEDULER_PROCESS_QUEUE(RunRemote) \
 			UBA_SCHEDULER_PROCESS_QUEUE(RunCrossArchitecture) \
 			UBA_SCHEDULER_PROCESS_QUEUE(RunCrossPlatform) \
+			UBA_SCHEDULER_PROCESS_QUEUE(RunCrossArchAndPlatform) \
 			UBA_SCHEDULER_PROCESS_QUEUE(Skip) \
 
 namespace uba
@@ -93,7 +92,7 @@ namespace uba
 		float weight = 1.0f;
 		bool wasReturned = false;
 		bool hasCheckedWasReturned = false;
-		bool isLocal = true;
+		Category runningCategory = Category_Local;
 	};
 
 	////////////////////////////////////////////////////////////////////////////////////////////////////
@@ -171,7 +170,7 @@ namespace uba
 			table->GetValueAsBool(writeTrackedInputs, TC("WriteTrackedInputs"));
 			table->GetValueAsBool(deferRemoteCapableProcesses, TC("DeferRemoteCapableProcesses"));
 			table->GetValueAsBool(enableSpeculativeCacheTest, TC("EnableSpeculativeCacheTest"));
-			table->GetValueAsBool(traceSchedulerQueues, TC("TraceSchedulerQueues"));
+			table->GetValueAsBool(traceDetails, TC("TraceDetails"));
 
 			u32 temp = 0;
 			if (table->GetValueAsU32(temp, TC("MemTracingLevel")))
@@ -216,7 +215,7 @@ namespace uba
 	:	m_session(info.session)
 	,	m_maxLocalProcessors(info.maxLocalProcessors != ~0u ? info.maxLocalProcessors : GetLogicalProcessorCount())
 	,	m_maxRacingWeight(float(info.maxRacingPercent)*float(m_maxLocalProcessors)/100.0f)
-	,	m_maxParallelCacheQueries(info.maxParallelCacheQueries)
+	,	m_maxParallelCacheQueries(Max(info.maxParallelCacheQueries, 4u))
 	,	m_updateThreadWakeup(EventResetType_Auto)
 	,	m_enableProcessReuse(info.enableProcessReuse)
 	,	m_forceRemote(info.forceRemote)
@@ -225,7 +224,7 @@ namespace uba
 	,	m_writeTrackedInputs(info.writeTrackedInputs)
 	,	m_deferRemoteCapableProcesses(info.deferRemoteCapableProcesses)
 	,	m_enableSpeculativeCacheTest(info.enableSpeculativeCacheTest)
-	,	m_traceSchedulerQueues(info.traceSchedulerQueues)
+	,	m_traceDetails(info.traceDetails)
 	,	m_processConfigs(info.processConfigs)
 	,	m_writeToCache(info.writeToCache && info.cacheClientCount)
 	{
@@ -406,6 +405,17 @@ namespace uba
 				processConfig->GetValueAsBool(canDetour, TC("CanDetour"));
 			}
 		}
+
+		auto GetCapabilityCategory = [](bool canExecuteRemotely, bool canCrossArchitecture, bool canCrossPlatform) -> Category
+		{
+			if (!canExecuteRemotely) return Category_Local;
+			if (canCrossArchitecture && canCrossPlatform) return Category_RemoteCrossArchAndPlatform;
+			if (canCrossArchitecture) return Category_RemoteCrossArch;
+			if (canCrossPlatform) return Category_RemoteCrossPlatform;
+			return Category_Remote;
+		};
+
+
 		ProcessEntry entry;
 		entry.info = info2;
 		entry.firstDependent = nullptr;
@@ -413,23 +423,28 @@ namespace uba
 		entry.speculativeDependencyCount = 0;
 		entry.runStatus = RunStatus_QueuedForRun;
 		entry.cacheStatus = useCache ? CacheStatus_QueuedForTest : CacheStatus_Failed;
+		entry.statsState = StatsState_None;
+		entry.category = GetCapabilityCategory(canExecuteRemotely, canCrossArchitecture, canCrossPlatform);
 		entry.canDetour = canDetour;
-		entry.canExecuteRemotely = canExecuteRemotely;
-		entry.canCrossArchitecture = canCrossArchitecture;
-		entry.canCrossPlatform = canCrossPlatform;
 		entry.shouldSkip = false;
 		return entry;
 	}
 
 	Scheduler::ProcessQueue& Scheduler::GetReadyToRunQueue(ProcessEntry& entry)
 	{
-		if (!entry.canDetour || !entry.canExecuteRemotely)
+		if (!entry.canDetour || entry.category == Category_Local)
 			return m_readyToRunLocal;
-		if (entry.canCrossPlatform)
+		switch (entry.category)
+		{
+		case Category_RemoteCrossArchAndPlatform:
+			return m_readyToRunCrossArchAndPlatform;
+		case Category_RemoteCrossPlatform:
 			return m_readyToRunCrossPlatform;
-		if (entry.canCrossArchitecture)
+		case Category_RemoteCrossArch:
 			return m_readyToRunCrossArchitecture;
-		return m_readyToRunRemote;
+		default:
+			return m_readyToRunRemote;
+		}
 	}
 
 	void Scheduler::AddToQueue(ProcessQueue& queue, u32 entryIndex)
@@ -478,7 +493,7 @@ namespace uba
 			if (dependency.runStatus == RunStatus_Failed || dependency.runStatus == RunStatus_Skipped || m_cancelled)
 			{
 				AddToQueue(m_readyToSkip, entryIndex);
-				return;
+				return; // Not counted in stats — will be skipped immediately
 			}
 			++dependencyCount;
 			// Speculative count only includes deps that haven't yet confirmed a cache hit.
@@ -487,15 +502,27 @@ namespace uba
 				++speculativeDependencyCount;
 		}
 
+		float weight = entry.info->weight;
+
 		if (!dependencyCount)
 		{
 			if (entry.cacheStatus == CacheStatus_QueuedForTest)
+			{
+				// Cache outcome not yet known — enters pending state until miss confirmed
+				TransitionStats(entry, StatsState_Pending, weight);
 				AddToQueue(m_readyToTestCache, entryIndex);
+			}
 			else
+			{
+				// No cache query — ready immediately
+				TransitionStats(entry, StatsState_Ready, weight);
 				AddToQueue(GetReadyToRunQueue(entry), entryIndex);
+			}
 			return;
 		}
 
+		// Has unsatisfied deps — process enters the left state
+		TransitionStats(entry, StatsState_Pending, weight);
 		entry.dependencyCount = dependencyCount;
 		entry.speculativeDependencyCount = speculativeDependencyCount;
 
@@ -543,27 +570,26 @@ namespace uba
 				continue;
 			}
 
+			float weight = dependent.info ? dependent.info->weight : 0.0f;
+
 			// dependencyCount hit zero
 			if (!shouldSkip && !dependent.shouldSkip && newRunStatus == RunStatus_Success)
 			{
 				switch (dependent.cacheStatus)
 				{
 				case CacheStatus_QueuedForTest:
-					AddToQueue(m_readyToTestCache, dependentIndex);
+					AddToQueue(m_readyToTestCache, dependentIndex); // stays pending
 					break;
 				case CacheStatus_QueuedForTestSpeculative:
-					// Already in m_readyToTestCache; will be processed normally now that deps are done.
-					break;
 				case CacheStatus_Testing:
 				case CacheStatus_QueuedForDownload:
 				case CacheStatus_Downloading:
-					// Speculative work in flight — worker checks dependencyCount on completion.
-					break;
+					break; // in-flight, stays pending
 				case CacheStatus_PendingCompletion:
-					// Speculative download finished while deps were pending; fire completion now.
-					AddToQueue(m_readyToCompleteFromCache, dependentIndex);
+					AddToQueue(m_readyToCompleteFromCache, dependentIndex); // stays pending
 					break;
-				default: // CacheStatus_Suspended (non-cacheable) or CacheStatus_Failed (cache miss)
+				default: // CacheStatus_Suspended (non-cacheable) or CacheStatus_Failed (cache miss confirmed)
+					TransitionStats(dependent, StatsState_Ready, weight); // pending → ready
 					AddToQueue(GetReadyToRunQueue(dependent), dependentIndex);
 					break;
 				}
@@ -575,22 +601,26 @@ namespace uba
 				|| dependent.cacheStatus == CacheStatus_QueuedForDownload
 				|| dependent.cacheStatus == CacheStatus_Downloading)
 			{
-				// Worker in flight — signal it to skip on completion.
+				// Worker in flight — signal it to skip on completion. left→none happens when worker finishes.
+				TransitionStats(dependent, StatsState_None, weight); // left → none (worker will handle skip)
 				dependent.shouldSkip = true;
 			}
 			else if (dependent.cacheStatus == CacheStatus_QueuedForTestSpeculative)
 			{
 				// In test queue but not yet started — remove and skip directly.
+				TransitionStats(dependent, StatsState_None, weight); // left → none
 				m_readyToTestCache.erase(dependentIndex);
 				AddToQueue(m_readyToSkip, dependentIndex);
 			}
 			else if (dependent.runStatus == RunStatus_QueuedForRun)
 			{
+				TransitionStats(dependent, StatsState_None, weight); // left → none
 				AddToQueue(m_readyToSkip, dependentIndex);
 			}
 			else if (dependent.cacheStatus == CacheStatus_PendingCompletion)
 			{
 				// Speculatively downloaded but a dep failed — treat as skipped.
+				TransitionStats(dependent, StatsState_None, weight); // left → none
 				dependent.runStatus = RunStatus_Skipped;
 				AddToQueue(m_readyToSkip, dependentIndex);
 			}
@@ -645,6 +675,7 @@ namespace uba
 				entry.info->cacheInfo = nullptr;
 			}
 
+			TransitionStats(entry, StatsState_None, entry.info ? entry.info->weight : 0.0f);
 			entry.cacheStatus = CacheStatus_Failed;
 			entry.runStatus = RunStatus_Skipped;
 			// UpdateDependentsNoLock(RunStatus_Skipped, firstDependent); // No need to do this, all following entries are cancelled
@@ -656,6 +687,7 @@ namespace uba
 		ClearQueue(m_readyToRunRemote);
 		ClearQueue(m_readyToRunCrossArchitecture);
 		ClearQueue(m_readyToRunCrossPlatform);
+		ClearQueue(m_readyToRunCrossArchAndPlatform);
 		ClearQueue(m_readyToSkip);
 		ClearQueue(m_readyToCompleteFromCache);
 
@@ -684,7 +716,7 @@ namespace uba
 		SetupDependents(m_processEntries.emplace_back(entry), index, info);
 		lock.Leave();
 		++m_totalProcesses;
-		UpdateQueueCounter(1, entry.canExecuteRemotely);
+		UpdateQueueCounter();
 		m_updateThreadWakeup.Set();
 		return index;
 	}
@@ -701,6 +733,7 @@ namespace uba
 		entry.runStatus = RunStatus_Suspended;
 		entry.cacheStatus = CacheStatus_Suspended;
 		++m_totalProcesses;
+		UpdateQueueCounter();
 		return index;
 	}
 
@@ -729,7 +762,6 @@ namespace uba
 		}
 
 		SetupDependents(entry, index, info);
-		UpdateQueueCounter(1, entry.canExecuteRemotely);
 		m_updateThreadWakeup.Set();
 	}
 
@@ -740,10 +772,11 @@ namespace uba
 
 	void Scheduler::GetStats(u32& outQueued, u32& outActiveLocal, u32& outActiveRemote, u32& outFinished)
 	{
-		outActiveLocal = m_activeLocalProcesses;
-		outActiveRemote = m_activeRemoteProcesses;
-		outFinished = m_finishedProcesses;
-		outQueued = m_queuedProcesses;
+		SCOPED_FUTEX_READ(m_processEntriesLock, lock);
+		outActiveLocal  = m_stats.runningLocalCount;
+		outActiveRemote = m_stats.runningRemoteCount + m_stats.runningRemoteCrossArchCount + m_stats.runningRemoteCrossPlatformCount + m_stats.runningRemoteCrossArchAndPlatformCount;
+		outFinished     = m_finishedProcesses;
+		outQueued       = m_totalProcesses - outFinished - outActiveLocal - outActiveRemote;
 	}
 
 	void Scheduler::GetResults(u32& outSucceeded, u32& outFailed, u32& outSkipped)
@@ -791,12 +824,19 @@ namespace uba
 		m_processFinished = processFinished;
 	}
 
+	void Scheduler::GetStats(SchedulerStats& outStats)
+	{
+		SCOPED_FUTEX_READ(m_processEntriesLock, lock);
+		outStats = m_stats;
+	}
+
 	u32 Scheduler::GetProcessCountThatCanRunRemotelyNow()
 	{
 		if (m_session.IsRemoteExecutionDisabled())
 			return 0;
 		SCOPED_FUTEX_READ(m_processEntriesLock, lock);
-		return GetQueueCount(m_readyToRunRemote) + GetQueueCount(m_readyToRunCrossArchitecture) + GetQueueCount(m_readyToRunCrossPlatform);
+		return	m_stats.activeRemoteCount + m_stats.activeRemoteCrossArchCount + m_stats.activeRemoteCrossPlatformCount + m_stats.activeRemoteCrossArchAndPlatformCount +
+				m_stats.readyRemoteCount + m_stats.readyRemoteCrossArchCount + m_stats.readyRemoteCrossPlatformCount + m_stats.readyRemoteCrossArchAndPlatformCount;
 	}
 
 	void Scheduler::GetProcessWeightThatCanRunRemotelyNow(float& outTotal, float& outCrossArchitecture, float& outCrossPlatform)
@@ -807,41 +847,40 @@ namespace uba
 		if (m_session.IsRemoteExecutionDisabled())
 			return;
 		SCOPED_FUTEX_READ(m_processEntriesLock, lock);
-		outCrossArchitecture = float(GetQueueCount(m_readyToRunCrossArchitecture));
-		outCrossPlatform = float(GetQueueCount(m_readyToRunCrossPlatform));
-		outTotal = float(GetQueueCount(m_readyToRunRemote)) + outCrossArchitecture + outCrossPlatform;
+		outCrossArchitecture = m_stats.activeRemoteCrossArchWeight + m_stats.readyRemoteCrossArchWeight;
+		outCrossPlatform     = m_stats.activeRemoteCrossPlatformWeight + m_stats.readyRemoteCrossPlatformWeight;
+		outTotal             = m_stats.activeRemoteWeight + m_stats.readyRemoteWeight + outCrossArchitecture + outCrossPlatform
+		                     + m_stats.activeRemoteCrossArchAndPlatformWeight + m_stats.readyRemoteCrossArchAndPlatformWeight;
 	}
 
 	void Scheduler::ThreadLoop()
 	{
+		auto& trace = m_session.GetTrace();
 		#define UBA_SCHEDULER_PROCESS_QUEUE(queue) u32 queue##Len = 0;
 		UBA_SCHEDULER_PROCESS_QUEUES
 		#undef UBA_SCHEDULER_PROCESS_QUEUE
-
-		if (m_traceSchedulerQueues)
-			m_session.GetTrace().StatusUpdate(2, 1, TCV("Queues"), LogEntryType_Info);
-
+		SchedulerStats lastStats = {};
+		bool queuesChanged = true;
+		const u32 startRow = 20;
+
+		if (m_traceDetails)
+		{
+			u32 row = startRow;
+			trace.StatusUpdate(row++, 1, TCV("Running"),  LogEntryType_Info);
+			trace.StatusUpdate(row++, 1, TCV("Queues"),  LogEntryType_Info);
+			trace.StatusUpdate(row, 1, TCV("Category"),  LogEntryType_Info);
+			trace.StatusUpdate(row, 13, TCV("Pending"),  LogEntryType_Info);
+			trace.StatusUpdate(row, 17, TCV("Ready"),  LogEntryType_Info);
+			trace.StatusUpdate(row, 21, TCV("Active"),  LogEntryType_Info);
+			row++;
+			#define UBA_PROCESS_CATEGORY(cat) trace.StatusUpdate(row++, 1, TCV("  " #cat),  LogEntryType_Info);
+			UBA_PROCESS_CATEGORIES
+			#undef UBA_PROCESS_CATEGORY
+		}
 		while (m_loop)
 		{
 			if (!m_updateThreadWakeup.IsSet())
 				break;
-			
-			if (m_traceSchedulerQueues)
-			{
-				bool changed = false;
-				SCOPED_FUTEX_READ(m_processEntriesLock, lock)
-				#define UBA_SCHEDULER_PROCESS_QUEUE(queue) u32 queue##LenNew = u32(m_readyTo##queue.size()); changed |= queue##LenNew != queue##Len; queue##Len = queue##LenNew;
-				UBA_SCHEDULER_PROCESS_QUEUES
-				#undef UBA_SCHEDULER_PROCESS_QUEUE
-				if (changed)
-				{
-					StringBuffer<> temp;
-					#define UBA_SCHEDULER_PROCESS_QUEUE(queue) temp.Append(TC("") #queue ": ").AppendValue(queue##Len).Append(TCV("  "));
-					UBA_SCHEDULER_PROCESS_QUEUES
-					#undef UBA_SCHEDULER_PROCESS_QUEUE
-					m_session.GetTrace().StatusUpdate(2, 6, temp, LogEntryType_Info);
-				}
-			}				
 
 			m_session.FlushDeadProcesses();
 
@@ -853,6 +892,48 @@ namespace uba
 
 			while (RunQueuedProcess(true, false, false))
 				;
+
+			if (m_traceDetails)
+			{
+				SchedulerStats stats;
+				{
+					SCOPED_FUTEX_READ(m_processEntriesLock, lock)
+					stats = m_stats;
+					#define UBA_SCHEDULER_PROCESS_QUEUE(queue) u32 queue##LenNew = u32(m_readyTo##queue.size()); queuesChanged |= queue##LenNew != queue##Len; queue##Len = queue##LenNew;
+					UBA_SCHEDULER_PROCESS_QUEUES
+					#undef UBA_SCHEDULER_PROCESS_QUEUE
+				}
+				StringBuffer<> temp;
+				if (queuesChanged)
+				{
+					queuesChanged = false;
+					#define UBA_SCHEDULER_PROCESS_QUEUE(queue) temp.Append(TC("") #queue ": ").AppendValue(queue##Len).Append(TCV("  "));
+					UBA_SCHEDULER_PROCESS_QUEUES
+					#undef UBA_SCHEDULER_PROCESS_QUEUE
+					trace.StatusUpdate(startRow+1, 6, temp, LogEntryType_Info);
+				}
+
+				if (memcmp(&stats, &lastStats, sizeof(stats)) != 0)
+				{
+					u32 row = startRow;
+					lastStats = stats;
+
+					temp.Clear();
+					temp.Appendf(TC("TestCache: %u "), stats.activeCacheQueries - stats.activeCacheDownloads);
+					temp.Appendf(TC("DownloadCache: %u "), stats.activeCacheDownloads);
+					#define UBA_PROCESS_CATEGORY(cat) temp.Appendf(TC(#cat ": %u "), stats.running##cat##Count);
+					UBA_PROCESS_CATEGORIES
+					#undef UBA_PROCESS_CATEGORY
+					trace.StatusUpdate(row,  6, temp, LogEntryType_Info);
+					row += 3;
+
+					auto traceValue = [&](u32 row, u32 col, u32 value) { temp.Clear().AppendValue(value); trace.StatusUpdate(row,  13 + col*4, temp, LogEntryType_Info); };
+					#define UBA_PROCESS_CATEGORY(cat) traceValue(row, 0, stats.pending##cat##Count);traceValue(row, 1, stats.ready##cat##Count);traceValue(row, 2, stats.active##cat##Count); ++row;
+					UBA_PROCESS_CATEGORIES
+					#undef UBA_PROCESS_CATEGORY
+				}
+			}				
+
 		}
 	}
 
@@ -920,10 +1001,10 @@ namespace uba
 			}
 		}
 
-		// We let the process that matches ei.isLocal do the return of process
+		// We let the process that matches ei.runningCategory do the return of process
 		// .. if remote process cancels first it will give over ownership to local race process
-		// .. if race process cancels first ei.isLocal will still be false and remote will clean up
-		if (ei.isLocal != isLocal)
+		// .. if race process cancels first ei.runningCategory will still be remote and remote will clean up
+		if ((ei.runningCategory == Category_Local) != isLocal)
 			return true;
 
 		UBA_ASSERT(ei.processCount == 0);
@@ -937,18 +1018,21 @@ namespace uba
 		if (entry.runStatus != RunStatus_Running)
 			return m_session.GetLogger().Error(TC("This should not happen (%u)"), entry.runStatus);
 
+		StatsRunning(m_stats, ei.runningCategory, -1, ei.weight);
+
 		ProcessLink* firstDependent = entry.firstDependent;
 		if (!m_cancelled)
 		{
 			entry.runStatus = RunStatus_QueuedForRun;
+			TransitionStats(entry, StatsState_Ready, ei.weight);
 			AddToQueue(GetReadyToRunQueue(entry), processIndex);
-			UpdateQueueCounter(1, entry.canExecuteRemotely);
 			lock.Leave();
 		}
 		else
 		{
 			m_session.GetLogger().Error(TC("Here to check this path works"));
 			entry.runStatus = RunStatus_Skipped;
+			TransitionStats(entry, StatsState_None, ei.weight);
 			ProcessHandle ph(&process);
 			FinishProcess(ph);
 
@@ -958,7 +1042,6 @@ namespace uba
 			ProcessLink::DeleteAll(firstDependent);
 		}
 
-		UpdateActiveProcessCounter(isLocal, -1);
 		m_updateThreadWakeup.Set();
 		return true;
 	}
@@ -970,9 +1053,9 @@ namespace uba
 		UBA_ASSERTF((entry.cacheStatus == CacheStatus_Testing || entry.cacheStatus == CacheStatus_Downloading) && entry.runStatus == RunStatus_QueuedForRun, TC("Unexpected entry state. Run: %u Cache: %u"), u32(entry.runStatus), u32(entry.cacheStatus));
 
 		if (entry.cacheStatus == CacheStatus_Downloading)
-			--m_activeCacheDownloads;
+			--m_stats.activeCacheDownloads;
 
-		--m_activeCacheQueries;
+		--m_stats.activeCacheQueries;
 
 		// Record whether the hit was previously confirmed (Downloading) before overwriting cacheStatus.
 		// If it was Testing→Failed the hit was never confirmed; speculative counts must be decremented.
@@ -986,6 +1069,7 @@ namespace uba
 		{
 			entry.runStatus = RunStatus_Skipped;
 			auto pi = entry.info;
+			TransitionStats(entry, StatsState_None, pi->weight);
 			entry.info = nullptr;
 			ProcessLink* firstDependent = entry.firstDependent;
 			entry.firstDependent = nullptr;
@@ -1003,7 +1087,10 @@ namespace uba
 		// With speculative cache testing, dependencyCount can be > 0 here; UpdateDependentsNoLock
 		// will promote to the run queue when the last dep completes (cacheStatus == Failed case).
 		if (entry.dependencyCount == 0)
+		{
+			TransitionStats(entry, StatsState_Ready, entry.info->weight); // pending → ready (cache miss confirmed)
 			AddToQueue(GetReadyToRunQueue(entry), processIndex);
+		}
 
 		lock.Leave();
 
@@ -1018,7 +1105,14 @@ namespace uba
 			return;
 		if (m_session.IsRemoteExecutionDisabled())
 			return;
-		u32 count = m_queuedProcessesCanExecuteRemotely;
+		u32 count;
+		{
+			SCOPED_FUTEX_READ(m_processEntriesLock, lock);
+			count =	m_stats.activeRemoteCount    + m_stats.readyRemoteCount    + m_stats.pendingRemoteCount +
+					m_stats.activeRemoteCrossArchCount + m_stats.readyRemoteCrossArchCount + m_stats.pendingRemoteCrossArchCount +
+					m_stats.activeRemoteCrossPlatformCount + m_stats.readyRemoteCrossPlatformCount + m_stats.pendingRemoteCrossPlatformCount +
+					m_stats.activeRemoteCrossArchAndPlatformCount + m_stats.readyRemoteCrossArchAndPlatformCount + m_stats.pendingRemoteCrossArchAndPlatformCount;
+		}
 		if (count < m_maxLocalProcessors)
 			m_session.DisableRemoteExecution();
 		else
@@ -1057,12 +1151,14 @@ namespace uba
 				}
 
 				// Turn it into real process adding weight
-				info->isLocal = true;
-				
+				const Category prevRunningCategory = info->runningCategory;
+				info->runningCategory = Category_Local;
+
 				m_activeLocalProcessWeight += weight;
 
-				UpdateActiveProcessCounter(true, 1);
-				UpdateActiveProcessCounter(false, -1);
+				SCOPED_FUTEX(m_processEntriesLock, entriesLock);
+				StatsRunning(m_stats, prevRunningCategory, -1, weight);
+				StatsRunning(m_stats, Category_Local, 1, weight);
 			}
 
 			m_activeRacingProcessWeight -= weight;
@@ -1085,7 +1181,7 @@ namespace uba
 			UBA_ASSERT(info->processIndex == ~0u);
 			
 			// Have to return the weight here if local. This means that this is a failed reuse and process has already been "exited" but we've kept process weight alive
-			if (info->isLocal)
+			if (info->runningCategory == Category_Local)
 			{
 				m_activeLocalProcessWeight -= info->weight;
 				UBA_ASSERTF(m_activeLocalProcessWeight >= 0, TC("Active process weight is below zero (%f)"), m_activeLocalProcessWeight.load());
@@ -1095,7 +1191,7 @@ namespace uba
 			return;
 		}
 
-		if (!m_useThreadToExitRemoteProcess || info->isLocal)
+		if (!m_useThreadToExitRemoteProcess || info->runningCategory == Category_Local)
 		{
 			#if 0
 			static int counter = 0;
@@ -1170,18 +1266,30 @@ namespace uba
 
 		if (!isLocal)
 		{
+			if (isCrossArchitecture && isCrossPlatform)
+				return returnEntry(PopFromQueue(m_readyToRunCrossArchAndPlatform));
+
 			if (isCrossArchitecture)
-				return returnEntry(PopFromQueue(m_readyToRunCrossArchitecture));
+			{
+				u32 index = returnEntry(PopFromQueue(m_readyToRunCrossArchitecture));
+				if (index != ~0u) return index;
+				return returnEntry(PopFromQueue(m_readyToRunCrossArchAndPlatform));
+			}
+
 			if (isCrossPlatform)
-				return returnEntry(PopFromQueue(m_readyToRunCrossPlatform));
+			{
+				u32 index = returnEntry(PopFromQueue(m_readyToRunCrossPlatform));
+				if (index != ~0u) return index;
+				return returnEntry(PopFromQueue(m_readyToRunCrossArchAndPlatform));
+			}
 
 			u32 index = returnEntry(PopFromQueue(m_readyToRunRemote));
-			if (index != ~0u)
-				return index;
+			if (index != ~0u) return index;
 			index = returnEntry(PopFromQueue(m_readyToRunCrossArchitecture));
-			if (index != ~0u)
-				return index;
-			return returnEntry(PopFromQueue(m_readyToRunCrossPlatform));
+			if (index != ~0u) return index;
+			index = returnEntry(PopFromQueue(m_readyToRunCrossPlatform));
+			if (index != ~0u) return index;
+			return returnEntry(PopFromQueue(m_readyToRunCrossArchAndPlatform));
 		}
 
 		auto getEntry = [&](float weightLeft)
@@ -1201,6 +1309,9 @@ namespace uba
 				if (ProcessEntry* entry = PeekFromQueue(m_readyToRunCrossPlatform))
 					if (entry->info->weight <= weightLeft)
 						return PopFromQueue(m_readyToRunCrossPlatform);
+				if (ProcessEntry* entry = PeekFromQueue(m_readyToRunCrossArchAndPlatform))
+					if (entry->info->weight <= weightLeft)
+						return PopFromQueue(m_readyToRunCrossArchAndPlatform);
 				return ~0u;
 			};
 
@@ -1209,12 +1320,12 @@ namespace uba
 
 		float activeWeight = m_activeLocalProcessWeight;
 		float maxWeight = float(m_maxLocalProcessors) + 0.9f;
-		if (activeWeight > 1.0f && (activeWeight >= maxWeight || !LocalProcessOkToSpawn()))
+		if (activeWeight > 1.0f && (activeWeight >= maxWeight || !LocalProcessOkToSpawn(m_stats.runningLocalCount)))
 			return ~0u;
 		if (activeWeight == 0)
 			maxWeight = 100000.0f; // If we are not running anything we must pick up the last process regardless of weight
 
-		// If there are queued getnext, then we try to hand over running queued processes
+		// If there are ready getnext, then we try to hand over running ready processes
 		// Setting the event will wake up the getnext thread.. and If the getnext thread is timing out we will retry
 		if (m_getNextQueueCount)
 		{
@@ -1247,8 +1358,9 @@ namespace uba
 		ProcessLink* firstDependent = entry.firstDependent;
 		entry.firstDependent = nullptr;
 		ProcessStartInfo2* si = entry.info;
+		TransitionStats(entry, StatsState_None, si->weight);
 		entry.info = nullptr;
-		
+
 		entriesLock.Leave();
 
 		auto cacheInfo = si->cacheInfo;
@@ -1269,7 +1381,6 @@ namespace uba
 		}
 
 		ProcessLink::DeleteAll(firstDependent);
-		UpdateActiveProcessCounter(true, -1);
 		m_updateThreadWakeup.Set();
 
 		if (m_writeTrackedInputs)
@@ -1289,24 +1400,24 @@ namespace uba
 	{
 		SCOPED_FUTEX(m_processEntriesLock, lock);
 		
-		if (m_activeCacheQueries >= m_maxParallelCacheQueries)
+		if (m_stats.activeCacheQueries >= m_maxParallelCacheQueries)
 			return false;
 
-		u32 entryToProcess;
+		u32 entryToProcess = ~0u;
 
 		// Some hard coded constants here.
-		// First, we don't want to look ahead too far since it will take a while and we want to get bandwidth maxed out earlier..
-		// So after 3000 ready to download we start prioritizing actual downloading over testing. Downside is that it might take longer to find cache misses, but worth it
-		if (m_readyToDownloadCache.size() > 3000)
+		// We want to some downloading while testing ahead.. so we always try to use half of the slots for downloading
+		if (m_stats.activeCacheDownloads < m_maxParallelCacheQueries/2)
 			entryToProcess = PopFromQueue(m_readyToDownloadCache);
-		else
+
+		if (entryToProcess == ~0u)
 		{
 			entryToProcess = PopFromQueue(m_readyToTestCache);
 			if (entryToProcess == ~0u)
 			{
 				// We always reserve two entries in the worker for testing when speculative testing is enabled. Because sometimes there are bottlenecks where testing will happen sequentially
 				// and if we don't reserve these slots we might catch up with testing which we don't want in order to keep bandwidth maxed out
-				if (m_enableSpeculativeCacheTest && m_activeCacheDownloads >= m_maxParallelCacheQueries - 2)
+				if (m_enableSpeculativeCacheTest && m_stats.activeCacheDownloads >= m_maxParallelCacheQueries - 2)
 					return false;
 				entryToProcess = PopFromQueue(m_readyToDownloadCache);
 			}
@@ -1320,17 +1431,16 @@ namespace uba
 		if (entry.cacheStatus == CacheStatus_QueuedForTest || entry.cacheStatus == CacheStatus_QueuedForTestSpeculative)
 		{
 			entry.cacheStatus = CacheStatus_Testing;
-			++m_activeCacheQueries;
+			++m_stats.activeCacheQueries;
 			lock.Leave();
 
 			m_cacheWorkManager->AddWork([this, info, entryToProcess](const WorkContext&)
 				{
 					ProcessStartInfo2& si = *info;
 
-					MessagePriority priority = m_activeCacheDownloads != 0 ? HasPriority : NormalPriority;
+					MessagePriority priority = m_stats.activeCacheDownloads != 0 ? HasPriority : NormalPriority;
 					auto cacheInfo = new CacheFetchInfo(si, m_session, si.rootsHandle, priority);
 
-					bool isHit = false;
 					for (auto cacheClient : m_cacheClients)
 					{
 						if (!cacheClient->FetchEntryFromCache(cacheInfo->result, *cacheInfo, si.cacheBucketId) || !cacheInfo->result.hit)
@@ -1342,24 +1452,19 @@ namespace uba
 						if (m_cancelled)
 							break;
 
-						isHit = true;
-
 						SCOPED_FUTEX(m_processEntriesLock, lock);
 						auto& hitEntry = m_processEntries[entryToProcess];
 						hitEntry.cacheStatus = CacheStatus_QueuedForDownload;
 						AddToQueue(m_readyToDownloadCache, entryToProcess);
-						--m_activeCacheQueries;
+						--m_stats.activeCacheQueries;
 						// Register A's expected output CAS keys so speculative dependents
 						// see them instead of stale on-disk values during their cache test.
 						if (m_enableSpeculativeCacheTest)
 							cacheInfo->speculativeOutputId = cacheInfo->client->RegisterSpeculativeOutputs(*cacheInfo);
 						// Allow dependents to start speculative cache tests now that this entry has a confirmed hit.
 						UpdateSpeculativeDependentsOnCacheHit(hitEntry.firstDependent);
-						break;
-					}
+						lock.Leave();
 
-					if (isHit)
-					{
 						m_updateThreadWakeup.Set();
 						return;
 					}
@@ -1374,12 +1479,11 @@ namespace uba
 		else if (entry.cacheStatus == CacheStatus_QueuedForDownload)
 		{
 			entry.cacheStatus = CacheStatus_Downloading;
-			++m_activeCacheQueries;
-			bool canExecuteRemotely = entry.canExecuteRemotely;
-			++m_activeCacheDownloads;
+			++m_stats.activeCacheQueries;
+			++m_stats.activeCacheDownloads;
 			lock.Leave();
 
-			m_cacheWorkManager->AddWork([this, info, entryToProcess, canExecuteRemotely](const WorkContext&)
+			m_cacheWorkManager->AddWork([this, info, entryToProcess](const WorkContext&)
 				{
 					auto cacheInfo = info->cacheInfo;
 
@@ -1391,21 +1495,20 @@ namespace uba
 
 					if (success)
 					{
-						UpdateQueueCounter(-1, canExecuteRemotely); // Queue counter is only reduced when process is picked up for Running.. which means need to reduce it on cache success too
-
 						SCOPED_FUTEX(m_processEntriesLock, lock);
 						auto& entry = m_processEntries[entryToProcess];
 
-						--m_activeCacheDownloads;
+						--m_stats.activeCacheDownloads;
 
 						if (m_cancelled || entry.shouldSkip)
 						{
 							// Build is being cancelled or a dep failed; skip the process.
-							--m_activeCacheQueries;
+							--m_stats.activeCacheQueries;
 							entry.cacheStatus = CacheStatus_Success;
 							entry.runStatus = RunStatus_Skipped;
 							ProcessLink* firstDependent = entry.firstDependent;
 							entry.firstDependent = nullptr;
+							TransitionStats(entry, StatsState_None, entry.info->weight);
 							entry.info = nullptr;
 							UpdateDependentsNoLock(RunStatus_Skipped, firstDependent, true);
 							lock.Leave();
@@ -1423,7 +1526,7 @@ namespace uba
 							// Speculative path: dependencies are still in flight.
 							// Park cacheInfo (which holds the log lines) in info; DrainCompleteFromCache
 							// will create the CachedProcess and clean up when the last dep finishes.
-							--m_activeCacheQueries;
+							--m_stats.activeCacheQueries;
 							entry.cacheStatus = CacheStatus_PendingCompletion;
 							lock.Leave();
 							m_updateThreadWakeup.Set();
@@ -1433,7 +1536,7 @@ namespace uba
 						}
 
 						// dependencyCount == 0: fall through to the existing ExitProcess path.
-						// ExitProcess will handle cacheStatus, m_activeCacheQueries, and completion.
+						// ExitProcess will handle cacheStatus, m_stats.activeCacheQueries, and completion.
 						lock.Leave();
 
 						auto process = new CachedProcess(si);
@@ -1442,7 +1545,6 @@ namespace uba
 						ExitProcessInfo exitInfo;
 						exitInfo.scheduler = this;
 						exitInfo.startInfo = info;
-						exitInfo.isLocal = true;
 						exitInfo.processIndex = entryToProcess;
 						ExitProcess(exitInfo, *process, 0, true, false);
 						++m_cacheHitCount;
@@ -1495,15 +1597,34 @@ namespace uba
 					return false;
 				return RaceRemoteProcess();
 			}
+
+			auto GetRunningCategory = [](bool isLocal, bool isCrossArchitecture, bool isCrossPlatform)
+			{
+				if (isLocal) return Category_Local;
+				if (isCrossArchitecture && isCrossPlatform) return Category_RemoteCrossArchAndPlatform;
+				if (isCrossArchitecture) return Category_RemoteCrossArch;
+				if (isCrossPlatform) return Category_RemoteCrossPlatform;
+				return Category_Remote;
+			};
+
+
 			auto& processEntry = m_processEntries[indexToRun];
 			auto info = processEntry.info;
 			bool canDetour = processEntry.canDetour && !m_forceNative;
 			bool wasSkipped = processEntry.runStatus == RunStatus_Skipped;
-			bool canExecuteRemotely = processEntry.canExecuteRemotely;
 			UBA_ASSERT(processEntry.cacheStatus == CacheStatus_Failed);
-			lock.Leave();
+			const Category runningCategory = GetRunningCategory(isLocal, isCrossArchitecture, isCrossPlatform);
 
-			UpdateQueueCounter(-1, canExecuteRemotely);
+			if (!wasSkipped && !m_cancelled)
+			{
+				TransitionStats(processEntry, StatsState_Active, info->weight);
+				StatsRunning(m_stats, runningCategory, 1, info->weight);
+			}
+			else
+			{
+				TransitionStats(processEntry, StatsState_None, info->weight);
+			}
+			lock.Leave();
 
 			if (wasSkipped || m_cancelled)
 			{
@@ -1514,12 +1635,10 @@ namespace uba
 			if (isLocal)
 				LocalProcessStart(*info);
 
-			UpdateActiveProcessCounter(isLocal, 1);
-	
 			auto exitInfo = new ExitProcessInfo();
 			exitInfo->scheduler = this;
 			exitInfo->startInfo = info;
-			exitInfo->isLocal = isLocal;
+			exitInfo->runningCategory = runningCategory;
 			exitInfo->processIndex = indexToRun;
 			exitInfo->weight = info->weight;
 
@@ -1536,21 +1655,32 @@ namespace uba
 			if (isLocal)
 				m_session.RunProcess(si, true, canDetour);
 			else
-				m_session.RunProcessRemote(si, info->weight, info->knownInputs, info->knownInputsCount, processEntry.canCrossArchitecture, processEntry.canCrossPlatform);
+			{
+				bool canCrossArch = processEntry.category == Category_RemoteCrossArch || processEntry.category == Category_RemoteCrossArchAndPlatform;
+				bool canCrossPlatform = processEntry.category == Category_RemoteCrossPlatform || processEntry.category == Category_RemoteCrossArchAndPlatform;
+				m_session.RunProcessRemote(si, info->weight, info->knownInputs, info->knownInputsCount, canCrossArch, canCrossPlatform);
+			}
 			return true;
 		}
 	}
 
 	bool Scheduler::RaceRemoteProcess()
 	{
-		if (!m_activeRemoteProcesses || m_maxRacingWeight == 0)
+		if (m_maxRacingWeight == 0)
 			return false;
+		u32 runningLocalCount;
+		{
+			SCOPED_FUTEX_READ(m_processEntriesLock, lock);
+			if (!m_stats.runningRemoteCount && !m_stats.runningRemoteCrossArchCount && !m_stats.runningRemoteCrossPlatformCount && !m_stats.runningRemoteCrossArchAndPlatformCount)
+				return false;
+			runningLocalCount = m_stats.runningLocalCount;
+		}
 
 		float totalWeight = m_activeLocalProcessWeight + m_activeRacingProcessWeight;
 		if (totalWeight >= m_maxRacingWeight)
 			return false;
 
-		if (!LocalProcessOkToSpawn())
+		if (!LocalProcessOkToSpawn(runningLocalCount))
 			return false;
 
 		SCOPED_FUTEX(m_racingProcessesLock, lock);
@@ -1665,19 +1795,19 @@ namespace uba
 			auto& processEntry = m_processEntries[indexToRun];
 			auto newInfo = processEntry.info;
 			bool wasSkipped = processEntry.runStatus == RunStatus_Skipped;
-			bool canExecuteRemotely = processEntry.canExecuteRemotely;
+			if (!wasSkipped)
+			{
+				TransitionStats(processEntry, StatsState_Active, newInfo->weight);
+				StatsRunning(m_stats, ei->runningCategory, 1, newInfo->weight);
+			}
 			lock.Leave();
 
-			UpdateQueueCounter(-1, canExecuteRemotely);
-
 			if (wasSkipped)
 			{
 				SkipProcess(*newInfo);
 				continue;
 			}
 
-			UpdateActiveProcessCounter(isLocal, 1);
-
 			auto& si = *newInfo;
 
 			UBA_ASSERT(ei->weight == si.weight); // TOOD: Support different weights?
@@ -1749,6 +1879,9 @@ namespace uba
 			bool hadCacheHit = (entry.cacheStatus >= CacheStatus_QueuedForDownload);
 			entry.runStatus = exitCode == 0 ? RunStatus_Success : RunStatus_Failed;
 			entry.firstDependent = nullptr;
+			if (!fromCache)
+				StatsRunning(m_stats, info.runningCategory, -1, si->weight);
+			TransitionStats(entry, StatsState_None, si->weight);
 			entry.info = nullptr;
 			UpdateDependentsNoLock(entry.runStatus, firstDependent, hadCacheHit);
 		}
@@ -1757,19 +1890,22 @@ namespace uba
 			UBA_ASSERT(!fromCache);
 			UBA_ASSERT(entry.cacheStatus == CacheStatus_Failed);
 
-			entry.canExecuteRemotely = false;
 			entry.canDetour = exitedResponse != ProcessExitedResponse_RerunNative;
 			entry.runStatus = RunStatus_QueuedForRun;
 			AddToQueue(GetReadyToRunQueue(entry), entryIndex);
+			StatsRunning(m_stats, info.runningCategory, -1, si->weight);
+			TransitionStats(entry, StatsState_None, si->weight);  // Active → None (old category)
+			entry.category = Category_Local;               // now can only run locally
+			TransitionStats(entry, StatsState_Ready, si->weight); // None → Queued (new category)
 		}
 
-		if (info.isLocal)
+		if (info.runningCategory == Category_Local)
 		{
 			if (fromCache)
 			{
 				UBA_ASSERT(isDone);
 				entry.cacheStatus = CacheStatus_Success;
-				--m_activeCacheQueries;
+				--m_stats.activeCacheQueries;
 			}
 			else if (!forReuse)
 			{
@@ -1780,7 +1916,6 @@ namespace uba
 
 		lock.Leave(); // entry reference can be bad after this since it can realloc
 
-		UpdateActiveProcessCounter(info.isLocal, -1);
 		m_updateThreadWakeup.Set();
 
 		if (isDone)
@@ -1788,7 +1923,10 @@ namespace uba
 			ProcessLink::DeleteAll(firstDependent);
 
 			if (exitCode != 0)
+			{
 				++m_errorCount;
+				UpdateQueueCounter();
+			}
 
 			bool writeTracked = m_writeTrackedInputs;
 			bool writeToCache = m_writeToCache && exitCode == 0 && process.GetExecutionType() != ProcessExecutionType_FromCache;
@@ -1838,26 +1976,50 @@ namespace uba
 		m_session.GetTrace().StatusUpdate(1, 6, str, LogEntryType_Info);
 	}
 
-	void Scheduler::UpdateQueueCounter(int offset, bool canExecuteRemotely)
+	void Scheduler::UpdateQueueCounter()
 	{
-		m_queuedProcesses += u32(offset);
-		if (canExecuteRemotely)
-			m_queuedProcessesCanExecuteRemotely += u32(offset);
 		m_session.UpdateProgress(m_totalProcesses, m_finishedProcesses, m_errorCount);
 	}
 
-	void Scheduler::UpdateActiveProcessCounter(bool isLocal, int offset)
+	#define UBA_STATS_SWITCH(field) \
+		switch (cat) { \
+		case Category_Local:                      s.field##LocalCount += u32(d); s.field##LocalWeight += w; return; \
+		case Category_Remote:                     s.field##RemoteCount += u32(d); s.field##RemoteWeight += w; return; \
+		case Category_RemoteCrossArch:            s.field##RemoteCrossArchCount += u32(d); s.field##RemoteCrossArchWeight += w; return; \
+		case Category_RemoteCrossPlatform:        s.field##RemoteCrossPlatformCount += u32(d); s.field##RemoteCrossPlatformWeight += w; return; \
+		case Category_RemoteCrossArchAndPlatform: s.field##RemoteCrossArchAndPlatformCount += u32(d); s.field##RemoteCrossArchAndPlatformWeight += w; return; \
+		default: return; }
+
+	void Scheduler::StatsActive(SchedulerStats& s, Category cat, int d, float w) { UBA_STATS_SWITCH(active)   }
+	void Scheduler::StatsReady(SchedulerStats& s, Category cat, int d, float w) { UBA_STATS_SWITCH(ready)    }
+	void Scheduler::StatsPending(SchedulerStats& s, Category cat, int d, float w) { UBA_STATS_SWITCH(pending) }
+	void Scheduler::StatsRunning(SchedulerStats& s, Category cat, int d, float w) { UBA_STATS_SWITCH(running)  }
+	#undef UBA_STATS_SWITCH
+
+
+	void Scheduler::TransitionStats(ProcessEntry& entry, StatsState newState, float weight)
 	{
-		if (isLocal)
-			m_activeLocalProcesses += u32(offset);
-		else
-			m_activeRemoteProcesses += u32(offset);
+		switch (entry.statsState)
+		{
+		case StatsState_Pending: StatsPending(m_stats, entry.category, -1, -weight); break;
+		case StatsState_Ready:   StatsReady  (m_stats, entry.category, -1, -weight); break;
+		case StatsState_Active:  StatsActive (m_stats, entry.category, -1, -weight); break;
+		default: break;
+		}
+		entry.statsState = newState;
+		switch (newState)
+		{
+		case StatsState_Pending: StatsPending(m_stats, entry.category, 1, weight); break;
+		case StatsState_Ready:   StatsReady  (m_stats, entry.category, 1, weight); break;
+		case StatsState_Active:  StatsActive (m_stats, entry.category, 1, weight); break;
+		default: break;
+		}
 	}
 
 	void Scheduler::FinishProcess(const ProcessHandle& handle)
 	{
 		++m_finishedProcesses;
-		m_session.UpdateProgress(m_totalProcesses, m_finishedProcesses, m_errorCount);
+		UpdateQueueCounter();
 		if (m_processFinished)
 			m_processFinished(handle);
 	}
@@ -2057,7 +2219,7 @@ namespace uba
 		++stats.activeProcessCount;
 	}
 
-	bool Scheduler::LocalProcessOkToSpawn()
+	bool Scheduler::LocalProcessOkToSpawn(u32 runningLocalCount)
 	{
 		if (!m_memWaitThreshold)
 			return true;
@@ -2072,7 +2234,7 @@ namespace uba
 			return true;
 		}
 
-		if (!m_activeLocalProcesses) // Always let one run so it can finish eventually
+		if (!runningLocalCount) // Always let one run so it can finish eventually
 			return true;
 
 		if (m_memDelaySpawning)
@@ -2459,8 +2621,11 @@ namespace uba
 
 				while (memUsed > memKillThreshold)
 				{
-					if (m_activeLocalProcesses <= 1)
-						break;
+					{
+						SCOPED_FUTEX_READ(m_processEntriesLock, lock);
+						if (m_stats.runningLocalCount <= 1)
+							break;
+					}
 					m_lastKillTime = GetTime();
 
 					str.Clear().Appendf(TC("Low on memory (%s/%s). Kill threshold is %s"), BytesToText(memUsed).str, BytesToText(memTotal).str, BytesToText(memKillThreshold).str);
diff --git a/Engine/Source/Programs/UnrealBuildAccelerator/Common/Public/UbaScheduler.h b/Engine/Source/Programs/UnrealBuildAccelerator/Common/Public/UbaScheduler.h
--- a/Engine/Source/Programs/UnrealBuildAccelerator/Common/Public/UbaScheduler.h
+++ b/Engine/Source/Programs/UnrealBuildAccelerator/Common/Public/UbaScheduler.h
@@ -41,7 +41,7 @@ namespace uba
 		bool writeTrackedInputs = false; // Set to true to write out log files with tracked inputs
 		bool deferRemoteCapableProcesses = false; // Set to true to have local session first pick local-only processes
 		bool enableSpeculativeCacheTest = false; // Set to true to allow cache testing of a process before its dependencies have finished downloading. Requires CacheClient support for speculative CAS keys to be correct.
-		bool traceSchedulerQueues = false; // Show in visualizer how many processes that are queued up for state change
+		bool traceDetails = false; // Show details of scheduler states in visualizer 
 		u32 maxRacingPercent = 0u; // If processor weights go below this percent, then we will start filling up with racing processes
 		u8 memTracingLevel = 0;
 		u8 memStartWaitPercent = 85; // When memory usage goes above this percent, no new processes will be spawned until back below
@@ -83,6 +83,39 @@ namespace uba
 
 	////////////////////////////////////////////////////////////////////////////////////////////////////
 
+	#define UBA_PROCESS_CATEGORIES \
+		UBA_PROCESS_CATEGORY(Local) \
+		UBA_PROCESS_CATEGORY(Remote) \
+		UBA_PROCESS_CATEGORY(RemoteCrossArch) \
+		UBA_PROCESS_CATEGORY(RemoteCrossPlatform) \
+		UBA_PROCESS_CATEGORY(RemoteCrossArchAndPlatform) \
+
+	// SchedulerStats contains following variables. Count is number of processes, weight is the accumulated weight of same processes
+	//
+	//   active   - Processes currently executing. Note: categorised by capability, not where they are actually running.
+	//   ready    - Cache miss confirmed (or no cache); available to run immediately if a slot is free.
+	//   pending  - Not yet confirmed as needing a worker: either blocked on dependencies or somewhere in the cache pipeline.
+	//   running  - Where processes are actually executing (by helper location, not capability).
+
+	struct SchedulerStats
+	{
+		#define UBA_PROCESS_CATEGORY(cat) \
+			u32 active##cat##Count = 0; \
+			float active##cat##Weight = 0; \
+			u32 ready##cat##Count = 0; \
+			float ready##cat##Weight = 0; \
+			u32 pending##cat##Count = 0; \
+			float pending##cat##Weight = 0; \
+			u32 running##cat##Count = 0; \
+			float running##cat##Weight = 0;
+		UBA_PROCESS_CATEGORIES
+		#undef UBA_PROCESS_CATEGORY
+		u32 activeCacheQueries = 0;
+		u32 activeCacheDownloads = 0;
+	};
+
+	////////////////////////////////////////////////////////////////////////////////////////////////////
+
 	class Scheduler
 	{
 	public:
@@ -101,7 +134,8 @@ namespace uba
 		void ResumeQueuedProcess(u32 index, const EnqueueProcessInfo& info); // Resume process queued with EnqueueSuspended
 
 		u32 GetTotalProcesses();
-		void GetStats(u32& outQueued, u32& outActiveLocal, u32& outActiveRemote, u32& outFinished); // This is not "threadsafe".. which means that you are not guaranteed to get the exact correct state.
+		void GetStats(u32& outQueued, u32& outActiveLocal, u32& outActiveRemote, u32& outFinished);
+		void GetStats(SchedulerStats& outStats);
 		void GetResults(u32& outSucceeded, u32& outFailed, u32& outSkipped);
 
 		bool IsEmpty(); // Returns true if scheduler is entirely empty.. as in no processes are left in the system
@@ -133,6 +167,8 @@ namespace uba
 			RunStatus_Skipped,
 		};
 
+		enum StatsState : u8 { StatsState_None, StatsState_Pending, StatsState_Ready, StatsState_Active };
+
 		enum CacheStatus : u8
 		{
 			CacheStatus_Suspended,
@@ -146,6 +182,13 @@ namespace uba
 			CacheStatus_Failed,
 		};
 
+		enum Category : u8
+		{
+			#define UBA_PROCESS_CATEGORY(cat) Category_##cat,
+			UBA_PROCESS_CATEGORIES
+			#undef UBA_PROCESS_CATEGORY
+		};
+
 		UBA_NOINLINE void ThreadLoop();
 
 		struct ProcessEntry;
@@ -177,15 +220,19 @@ namespace uba
 		void ExitProcess(ExitProcessInfo& info, Process& process, u32 exitCode, bool fromCache, bool forReuse);
 		void SkipProcess(ProcessStartInfo2& info);
 		void UpdateStatusHitMissCount();
-		void UpdateQueueCounter(int offset, bool canRunRemotely);
-		void UpdateActiveProcessCounter(bool isLocal, int offset);
+		void UpdateQueueCounter();
+		void StatsActive(SchedulerStats& s, Category cat, int d, float w);
+		void StatsReady(SchedulerStats& s, Category cat, int d, float w);
+		void StatsPending(SchedulerStats& s, Category cat, int d, float w);
+		void StatsRunning(SchedulerStats& s, Category cat, int d, float w);
+		void TransitionStats(ProcessEntry& entry, StatsState newState, float weight);
 		void FinishProcess(const ProcessHandle& handle);
 		void CreateCacheInputsAndOutputs(Vector<u8>& outInputs, Vector<u8>& outOutputs, ProcessStartInfo2& info, Process& process);
 		void WriteTrackedInputs(const Vector<u8>& inputs, const Vector<u8>& outputs, Process& process);
 		void KillNewestLocalProcess(StringView reason);
 		void NativeProcessCreated(ProcessImpl& process);
 		void LocalProcessStart(ProcessStartInfo2& info);
-		bool LocalProcessOkToSpawn();
+		bool LocalProcessOkToSpawn(u32 runningLocalCount);
 		void LocalProcessExit(ProcessStartInfo2& info, const ProcessHandle& handle);
 		void LoadMemTrackTable();
 		void SaveMemTrackTable();
@@ -199,11 +246,10 @@ namespace uba
 			u32 speculativeDependencyCount;
 			CacheStatus cacheStatus;
 			RunStatus runStatus;
+			StatsState statsState;
+			Category category;
 			bool shouldSkip;
 			bool canDetour;
-			bool canExecuteRemotely;
-			bool canCrossArchitecture;
-			bool canCrossPlatform;
 		};
 
 		struct ProcessQueue : Set<u32> {};
@@ -220,8 +266,9 @@ namespace uba
 		ProcessQueue m_readyToDownloadCache;
 		ProcessQueue m_readyToRunLocal; // Can only run locally
 		ProcessQueue m_readyToRunRemote; // Can run remotely on same architecture
-		ProcessQueue m_readyToRunCrossArchitecture; // Can run remotely and also cross architecture
-		ProcessQueue m_readyToRunCrossPlatform; // Can run remotely and also cross platform
+		ProcessQueue m_readyToRunCrossArchitecture;    // Can run remotely on a different architecture
+		ProcessQueue m_readyToRunCrossPlatform;        // Can run remotely on a different platform
+		ProcessQueue m_readyToRunCrossArchAndPlatform; // Can run remotely on a different architecture AND platform
 		ProcessQueue m_readyToSkip;
 		ProcessQueue m_readyToCompleteFromCache; // Cache download done but deps were still pending; fired when last dep completes
 
@@ -239,7 +286,7 @@ namespace uba
 		bool m_writeTrackedInputs = false;
 		bool m_deferRemoteCapableProcesses = false;
 		bool m_enableSpeculativeCacheTest = false;
-		bool m_traceSchedulerQueues = false;
+		bool m_traceDetails = false;
 		bool m_cancelled = false;
 		bool m_hadRemoteRequest = false;
 		bool m_hadCrossSystemRequest = false;
@@ -247,16 +294,12 @@ namespace uba
 
 		Atomic<float> m_activeLocalProcessWeight = 0.0f;
 		Atomic<float> m_activeRacingProcessWeight = 0.0f;
-		u32 m_activeCacheQueries = 0;
-		u32 m_activeCacheDownloads = 0;
 		Atomic<u32> m_getNextQueueCount = 0;
 
-		Atomic<u32> m_totalProcesses = 0; // Process count registed in to scheduler
-		Atomic<u32> m_queuedProcesses = 0; // Number of processes that has not started yet
-		Atomic<u32> m_activeLocalProcesses = 0; // Number of processes active on local machine
-		Atomic<u32> m_activeRemoteProcesses = 0; // Number of processes active on remote helpers
+		SchedulerStats m_stats;
+
+		Atomic<u32> m_totalProcesses = 0; // Process count registered in to scheduler
 		Atomic<u32> m_finishedProcesses = 0; // Processes finished (includes skipped)
-		Atomic<u32> m_queuedProcessesCanExecuteRemotely = 0; 
 		Atomic<u32> m_errorCount = 0;
 		Atomic<u32> m_cacheHitCount = 0;
 		Atomic<u32> m_cacheMissCount = 0;
