[UBA]
* Fixed header compile errors

diff --git a/Engine/Source/Programs/UnrealBuildAccelerator/Common/Private/UbaEnvironment.cpp b/Engine/Source/Programs/UnrealBuildAccelerator/Common/Private/UbaEnvironment.cpp
--- a/Engine/Source/Programs/UnrealBuildAccelerator/Common/Private/UbaEnvironment.cpp
+++ b/Engine/Source/Programs/UnrealBuildAccelerator/Common/Private/UbaEnvironment.cpp
@@ -613,7 +613,7 @@ namespace uba
 			}();
 		return sv;
 	}
-	HMODULE GetUbaWineModule()
+	void* GetUbaWineModule()
 	{
 		static HMODULE ubaWineModule = []()
 			{
diff --git a/Engine/Source/Programs/UnrealBuildAccelerator/Common/Private/UbaNetworkBackendTcp.cpp b/Engine/Source/Programs/UnrealBuildAccelerator/Common/Private/UbaNetworkBackendTcp.cpp
--- a/Engine/Source/Programs/UnrealBuildAccelerator/Common/Private/UbaNetworkBackendTcp.cpp
+++ b/Engine/Source/Programs/UnrealBuildAccelerator/Common/Private/UbaNetworkBackendTcp.cpp
@@ -234,8 +234,8 @@ namespace uba
 
 		if (IsRunningWine())
 		{
-			if (HMODULE wineDll = GetUbaWineModule())
-				SocketInitWine(logger, wineDll);
+			if (void* wineDll = GetUbaWineModule())
+				SocketInitWine(logger, (HMODULE)wineDll);
 			else
 				logger.Warning(TC("Failed to load UbaWine.dll.so (%s)"), LastErrorToText().data);
 
@@ -1690,7 +1690,7 @@ namespace uba
 
 		#if PLATFORM_WINDOWS
 		if (!GetTcpAutoTuning && IsRunningWine())
-			GetTcpAutoTuning = (decltype(GetTcpAutoTuning))GetProcAddress(GetUbaWineModule(), "GetTcpAutoTuning");
+			GetTcpAutoTuning = (decltype(GetTcpAutoTuning))GetProcAddress((HMODULE)GetUbaWineModule(), "GetTcpAutoTuning");
 		#else
 		GetTcpAutoTuning = UnixGetTcpAutoTuning;
 		#endif
@@ -2197,7 +2197,7 @@ namespace uba
 		#if 0
 		if (IsRunningWine())
 		{
-			static HMODULE wineDll = GetUbaWineModule();
+			static HMODULE wineDll = (HMODULE)GetUbaWineModule();
 			if (!wineDll)
 				return false;
 			using GetLinuxTcpInfoFunc = int WINAPI(Socket, void*, int*);
diff --git a/Engine/Source/Programs/UnrealBuildAccelerator/Common/Private/UbaStorage.cpp b/Engine/Source/Programs/UnrealBuildAccelerator/Common/Private/UbaStorage.cpp
--- a/Engine/Source/Programs/UnrealBuildAccelerator/Common/Private/UbaStorage.cpp
+++ b/Engine/Source/Programs/UnrealBuildAccelerator/Common/Private/UbaStorage.cpp
@@ -1047,8 +1047,8 @@ namespace uba
 		if (info.sideStepWineForSharedMemory && IsRunningWine())
 		{
 			static FileMappingBackend wineBackend;
-			if (HMODULE h = GetUbaWineModule())
-				if (FileMapping_GetWineBackend(wineBackend, logger, h))
+			if (void* h = GetUbaWineModule())
+				if (FileMapping_GetWineBackend(wineBackend, logger, (HMODULE)h))
 				{
 					logger.Debug(TC("Using wine filemapping backend"));
 					return wineBackend;
diff --git a/Engine/Source/Programs/UnrealBuildAccelerator/Common/Public/UbaEnvironment.h b/Engine/Source/Programs/UnrealBuildAccelerator/Common/Public/UbaEnvironment.h
--- a/Engine/Source/Programs/UnrealBuildAccelerator/Common/Public/UbaEnvironment.h
+++ b/Engine/Source/Programs/UnrealBuildAccelerator/Common/Public/UbaEnvironment.h
@@ -35,6 +35,6 @@ namespace uba
 
 	#if PLATFORM_WINDOWS
 	StringView GetUbaWineModulePath();
-	HMODULE GetUbaWineModule();
+	void* GetUbaWineModule();
 	#endif
 }
\ No newline at end of file
