[UbaCoordinatorHorde]
* Fixed so cache server can be requested in c++ code paths

#Synced-CL 52497433

diff --git a/Engine/Source/Developer/UbaCoordinatorHorde/Private/UbaCoordinatorHorde.cpp b/Engine/Source/Developer/UbaCoordinatorHorde/Private/UbaCoordinatorHorde.cpp
--- a/Engine/Source/Developer/UbaCoordinatorHorde/Private/UbaCoordinatorHorde.cpp
+++ b/Engine/Source/Developer/UbaCoordinatorHorde/Private/UbaCoordinatorHorde.cpp
@@ -56,6 +56,23 @@ public:
 		#endif
 	}
 
+	virtual bool RequestCacheServer(uba::tchar* outEndpoint, uba::u32 outEndpointCapacity, uba::Guid& outSessionKey, bool& outWriteAccess) override
+	{
+		#if PLATFORM_WINDOWS
+		FString Endpoint;
+		FString SessionKey;
+		bool bWriteAccess;
+		if (!m_manager.RequestCacheServer(Endpoint, SessionKey, bWriteAccess))
+			return false;
+		memcpy(outEndpoint, *Endpoint, (Endpoint.Len()+1)*sizeof(uba::tchar));
+		FString::ToHexBlob(SessionKey, (uint8*)&outSessionKey, sizeof(uba::Guid));
+		outWriteAccess = bWriteAccess;
+		return true;
+		#else
+		return false;
+		#endif
+	}
+
 	FUbaHordeAgentManager m_manager;
 };
 
@@ -87,12 +104,7 @@ extern "C"
 		FHttpModule::Get();
 		FDesktopPlatformModule::TryGet();
 
-		auto coordinator = new CoordinatorHorde(info.workDir, info.binariesDir);
-		auto& m = coordinator->m_manager;
-		//m.SetPool(info.pool);
-		m.SetTargetCoreCount(info.maxCoreCount);
-
-		return coordinator;
+		return new CoordinatorHorde(info.workDir, info.binariesDir);
 	}
 
 	UBA_EXPORT void UbaDestroyCoordinator(uba::Coordinator* coordinator)
@@ -117,4 +129,4 @@ class FUbaCoordinatorHordeModule : public IModuleInterface
 
 IMPLEMENT_MODULE(FUbaCoordinatorHordeModule, UbaCoordinatorHorde);
 
-#endif
\ No newline at end of file
+#endif
diff --git a/Engine/Source/Programs/UnrealBuildAccelerator/Cli/Private/UbaCli.cpp b/Engine/Source/Programs/UnrealBuildAccelerator/Cli/Private/UbaCli.cpp
--- a/Engine/Source/Programs/UnrealBuildAccelerator/Cli/Private/UbaCli.cpp
+++ b/Engine/Source/Programs/UnrealBuildAccelerator/Cli/Private/UbaCli.cpp
@@ -1464,12 +1464,11 @@ namespace uba
 					cinfo.workDir = coordinatorWorkDir.data;
 					cinfo.binariesDir = binariesDir.data;
 
-					// TODO: This is very horde specific.. maybe all these parameters should be a string or something
-					cinfo.pool = coordinatorPool.c_str();
-					cinfo.maxCoreCount = 1;
 					cinfo.logging = false;
-					if (!coordinator.Create(logger, coordinatorName.c_str(), cinfo, networkBackend, networkServer, scheduler))
+					if (!coordinator.Create(logger, coordinatorName.c_str(), cinfo, networkBackend, networkServer, sessionServer.GetTrace()))
 						return false;
+					coordinator.SetScheduler(scheduler);
+					coordinator.RequestHelpers(1);
 					return true;
 				};
 			destroyCoordinator = [&]() { coordinator.Destroy(); };
diff --git a/Engine/Source/Programs/UnrealBuildAccelerator/Common/Public/UbaCoordinator.h b/Engine/Source/Programs/UnrealBuildAccelerator/Common/Public/UbaCoordinator.h
--- a/Engine/Source/Programs/UnrealBuildAccelerator/Common/Public/UbaCoordinator.h
+++ b/Engine/Source/Programs/UnrealBuildAccelerator/Common/Public/UbaCoordinator.h
@@ -8,11 +8,7 @@ namespace uba
 	{
 		const tchar* workDir = nullptr;
 		const tchar* binariesDir = nullptr;
-
-		// TODO: This is very horde specific.. maybe all these parameters should be a string or something
-		const tchar* pool = nullptr;
 		u32 maxCoreCount = 0;
-
 		bool logging = false;
 	};
 
@@ -25,6 +21,9 @@ namespace uba
 
 		using UpdateStatusCallback = void(void* userData, const tchar* status);
 		virtual void SetUpdateStatusCallback(UpdateStatusCallback* callback, void* userData) {}
+
+		virtual bool RequestCacheServer(tchar* outEndpoint, u32 outEndpointCapacity, Guid& outSessionKey, bool& outWriteAccess) { return false; }
+
 	};
 
 	// This is how the function signatures creating/destroying the coordinator module needs to exported
diff --git a/Engine/Source/Programs/UnrealBuildAccelerator/Common/Public/UbaCoordinatorWrapper.h b/Engine/Source/Programs/UnrealBuildAccelerator/Common/Public/UbaCoordinatorWrapper.h
--- a/Engine/Source/Programs/UnrealBuildAccelerator/Common/Public/UbaCoordinatorWrapper.h
+++ b/Engine/Source/Programs/UnrealBuildAccelerator/Common/Public/UbaCoordinatorWrapper.h
@@ -24,7 +24,7 @@ namespace uba
 	class CoordinatorWrapper
 	{
 	public:
-		bool Create(Logger& logger, const tchar* coordinatorType, const CoordinatorCreateInfo& info, NetworkBackend& networkBackend, NetworkServer& networkServer, Scheduler* scheduler = nullptr)
+		bool Create(Logger& logger, const tchar* coordinatorType, const CoordinatorCreateInfo& info, NetworkBackend& networkBackend, NetworkServer& networkServer, Trace& trace, u32 traceRow = 10)
 		{
 			UbaCreateCoordinatorFunc* createCoordinator = nullptr;
 
@@ -59,26 +59,11 @@ namespace uba
 			if (!m_coordinator)
 				return false;
 
-			m_loopCoordinator.Create(EventResetType_Manual);
-			m_networkBackend = &networkBackend;
-			m_networkServer = &networkServer;
-			m_scheduler = scheduler;
-
-			if (m_scheduler)
-				m_scheduler->GetSession().GetTrace().StatusUpdate(0, 1, ToView(coordinatorType), LogEntryType_Info, {});
-
-			m_coordinatorThread.Start([this, mcc = info.maxCoreCount]() { ThreadUpdate(mcc); return 0; }, TC("UbaCoordWrap"));
-
-			return true;
-		}
-
-		void ThreadUpdate(u32 maxCoreCount)
-		{
+			m_traceRow = traceRow;
 			m_coordinator->SetUpdateStatusCallback([](void* userData, const tchar* status)
 				{
 					auto& cw = *(CoordinatorWrapper*)userData;
-					if (cw.m_scheduler)
-						cw.m_scheduler->GetSession().GetTrace().StatusUpdate(0, 6, ToView(status), LogEntryType_Info, {});
+					cw.m_trace->StatusUpdate(cw.m_traceRow, 6, ToView(status), LogEntryType_Info, {});
 				}, this);
 
 			m_coordinator->SetAddClientCallback([](void* userData, const tchar* ip, u16 port, const tchar* crypto)
@@ -87,23 +72,66 @@ namespace uba
 					return cw.m_networkServer->AddClient(*cw.m_networkBackend, ip, port);
 				}, this);
 
+			m_networkBackend = &networkBackend;
+			m_networkServer = &networkServer;
+			m_trace = &trace;
+			m_trace->StatusUpdate(m_traceRow, 1, ToView(coordinatorType), LogEntryType_Info, {});
+
+			if (info.maxCoreCount)
+				RequestHelpers(info.maxCoreCount);
+			return true;
+		}
+
+		bool IsCreated()
+		{
+			return m_coordinator != nullptr;
+		}
+
+		void SetScheduler(Scheduler* scheduler)
+		{
+			m_scheduler = scheduler;
+		}
+
+		bool RequestCacheServer(StringBufferBase& outEndpoint, Guid& outSessionKey, bool& outWriteAccess)
+		{
+			UBA_ASSERT(m_coordinator);
+			if (!m_coordinator->RequestCacheServer(outEndpoint.data, outEndpoint.capacity, outSessionKey, outWriteAccess))
+				return false;
+			outEndpoint.count = TStrlen(outEndpoint.data);
+			return true;
+		}
+
+		bool RequestHelpers(u32 maxCoreCount)
+		{
+			if (m_loopCoordinator.IsCreated())
+				return false;
+			m_loopCoordinator.Create(EventResetType_Manual);
+			m_coordinatorThread.Start([this, mcc = maxCoreCount]() { ThreadUpdate(mcc); return 0; }, TC("UbaCoordWrap"));
+			return true;
+		}
+
+		void ThreadUpdate(u32 maxCoreCount)
+		{
 			do
 			{
 				u32 coreCount = maxCoreCount;
 				if (m_scheduler)
 					coreCount = Min(m_scheduler->GetProcessCountThatCanRunRemotelyNow(), maxCoreCount);
 
-				m_coordinator->SetTargetCoreCount(coreCount);
+				m_coordinator->SetTargetCoreCount(coreCount); // Need to be called all the time to request more helpers
 			}
-			while (!m_loopCoordinator.IsSet(3000));
+			while (!m_loopCoordinator.IsSet(2000));
 		}
 
 		void Destroy()
 		{
 			if (!m_coordinator)
 				return;
-			m_loopCoordinator.Set();
-			m_coordinatorThread.Wait();
+			if (m_loopCoordinator.IsCreated())
+			{
+				m_loopCoordinator.Set();
+				m_coordinatorThread.Wait();
+			}
 			m_destroyCoordinator(m_coordinator);
 			m_coordinator = nullptr;
 		}
@@ -111,9 +139,11 @@ namespace uba
 		Coordinator* m_coordinator = nullptr;
 		NetworkBackend* m_networkBackend = nullptr;
 		NetworkServer* m_networkServer = nullptr;
+		Trace* m_trace;
 		Scheduler* m_scheduler = nullptr;
 		UbaDestroyCoordinatorFunc* m_destroyCoordinator = nullptr;
 		Event m_loopCoordinator;
 		Thread m_coordinatorThread;
+		u32 m_traceRow = 0;
 	};
 }
\ No newline at end of file
diff --git a/Engine/Source/Programs/UnrealBuildAccelerator/Host/Private/UbaExports.cpp b/Engine/Source/Programs/UnrealBuildAccelerator/Host/Private/UbaExports.cpp
--- a/Engine/Source/Programs/UnrealBuildAccelerator/Host/Private/UbaExports.cpp
+++ b/Engine/Source/Programs/UnrealBuildAccelerator/Host/Private/UbaExports.cpp
@@ -1241,11 +1241,8 @@ extern "C"
 						cinfo.workDir = coordinatorWorkDir.data;
 						cinfo.binariesDir = binariesDir.data;
 
-						coordinatorTable->GetValueAsString(cinfo.pool, TC("Pool"));
-						UBA_ASSERT(cinfo.pool);
-
-						cinfo.maxCoreCount = 500;
-						coordinatorTable->GetValueAsU32(cinfo.maxCoreCount, TC("MaxCoreCount"));
+						u32 maxCoreCount = 0;
+						coordinatorTable->GetValueAsU32(maxCoreCount, TC("MaxCoreCount"));
 
 						cinfo.logging = false;
 						coordinatorTable->GetValueAsBool(cinfo.logging, TC("Log"));
@@ -1254,8 +1251,10 @@ extern "C"
 						if (coordinatorTable->GetValueAsString(uri, TC("Uri")))
 							uba::SetEnvironmentVariableW(TC("UE_HORDE_URL"), uri);
 
-						if (!ubaInstance->coordinator.Create(logger, coordinatorName, cinfo, *networkServer->backend, *networkServer, scheduler))
+						if (!ubaInstance->coordinator.Create(logger, coordinatorName, cinfo, *networkServer->backend, *networkServer, sessionServer->GetTrace(), 0))
 							return;
+						ubaInstance->coordinator.SetScheduler(scheduler);
+						ubaInstance->coordinator.RequestHelpers(maxCoreCount);
 					}
 				}
 			};
