[UBA]
* Fixed bug where ArchAndPlatform queue was not included in arch and platform returned from GetProcessWeightThatCanRunRemotelyNow

diff --git a/Engine/Source/Programs/UnrealBuildAccelerator/Common/Private/UbaScheduler.cpp b/Engine/Source/Programs/UnrealBuildAccelerator/Common/Private/UbaScheduler.cpp
--- a/Engine/Source/Programs/UnrealBuildAccelerator/Common/Private/UbaScheduler.cpp
+++ b/Engine/Source/Programs/UnrealBuildAccelerator/Common/Private/UbaScheduler.cpp
@@ -847,10 +847,13 @@ namespace uba
 		if (m_session.IsRemoteExecutionDisabled())
 			return;
 		SCOPED_FUTEX_READ(m_processEntriesLock, lock);
+		float archAndPlatformWeight = m_stats.activeRemoteCrossArchAndPlatformWeight + m_stats.readyRemoteCrossArchAndPlatformWeight;
 		outCrossArchitecture = m_stats.activeRemoteCrossArchWeight + m_stats.readyRemoteCrossArchWeight;
 		outCrossPlatform     = m_stats.activeRemoteCrossPlatformWeight + m_stats.readyRemoteCrossPlatformWeight;
 		outTotal             = m_stats.activeRemoteWeight + m_stats.readyRemoteWeight + outCrossArchitecture + outCrossPlatform
-		                     + m_stats.activeRemoteCrossArchAndPlatformWeight + m_stats.readyRemoteCrossArchAndPlatformWeight;
+		                     + archAndPlatformWeight;
+		outCrossArchitecture += archAndPlatformWeight;
+		outCrossPlatform     += archAndPlatformWeight;
 	}
 
 	void Scheduler::ThreadLoop()
