[UBA]
* Fixed compile error in unit tests

diff --git a/Engine/Source/Programs/UnrealBuildAccelerator/Test/Private/UbaTestStress.cpp b/Engine/Source/Programs/UnrealBuildAccelerator/Test/Private/UbaTestStress.cpp
--- a/Engine/Source/Programs/UnrealBuildAccelerator/Test/Private/UbaTestStress.cpp
+++ b/Engine/Source/Programs/UnrealBuildAccelerator/Test/Private/UbaTestStress.cpp
@@ -68,12 +68,9 @@ namespace uba
 			cinfo.workDir = coordinatorWorkDir.data;
 			cinfo.binariesDir = binariesDir.data;
 
-			cinfo.pool = TC("PoolToUse");
-
-
 			cinfo.maxCoreCount = 400;
 			cinfo.logging = true;
-			CHECK_TRUE(coordinator.Create(logger, TC("Horde"), cinfo, networkBackend, networkServer));
+			CHECK_TRUE(coordinator.Create(logger, TC("Horde"), cinfo, networkBackend, networkServer, sessionServer.GetTrace()));
 		}
 		auto cg = MakeGuard([&]() { coordinator.Destroy(); });
 
