[UBA]
* Fixed so mac is using POSIX_SPAWN_CLOEXEC_DEFAULT for spawning child processes in order to work around potential problems with hangs caused by race conditions
* Fixed a couple more "leaks" of file descriptors to child processes

diff --git a/Engine/Source/Programs/UnrealBuildAccelerator/Common/Private/UbaProcess.cpp b/Engine/Source/Programs/UnrealBuildAccelerator/Common/Private/UbaProcess.cpp
--- a/Engine/Source/Programs/UnrealBuildAccelerator/Common/Private/UbaProcess.cpp
+++ b/Engine/Source/Programs/UnrealBuildAccelerator/Common/Private/UbaProcess.cpp
@@ -2296,6 +2296,16 @@ namespace uba
 			sigemptyset(&emptymask);
 			res = posix_spawnattr_setsigmask(&attr, &emptymask);
 			UBA_ASSERTF(res == 0, TC("posix_spawnattr_setsigmask (%s)"), strerror(errno));
+			short spawnFlags = POSIX_SPAWN_SETPGROUP | POSIX_SPAWN_SETSIGMASK;
+			
+			#if PLATFORM_MAC
+			// POSIX_SPAWN_CLOEXEC_DEFAULT closes all fds in the child that aren't explicitly
+			// set up via fileActions. This prevents a race on Mac where another thread can fork
+			// between pipe() and fcntl(FD_CLOEXEC), causing the child to inherit stray pipe fds
+			// and the parent to hang forever in poll() waiting for an EOF that never comes.
+			spawnFlags |= POSIX_SPAWN_CLOEXEC_DEFAULT;
+			#endif
+			
 			res = posix_spawnattr_setflags(&attr, POSIX_SPAWN_SETPGROUP | POSIX_SPAWN_SETSIGMASK);
 			UBA_ASSERTF(res == 0, TC("posix_spawnattr_setflags (%s)"), strerror(errno));
 			res = posix_spawnattr_setpgroup(&attr, getpgrp());
@@ -2427,12 +2437,6 @@ namespace uba
 					return UBA_EXIT_CODE(18);
 				}
 				
-				#if PLATFORM_MAC
-				auto set_cloexec = [](int fd){ int f = fcntl(fd, F_GETFD); if (f != -1) fcntl(fd, F_SETFD, f | FD_CLOEXEC); };
-				set_cloexec(outPipe[0]); set_cloexec(outPipe[1]);
-				set_cloexec(errPipe[0]); set_cloexec(errPipe[1]);
-				#endif
-
 				res = posix_spawn_file_actions_addclose(&fileActions, outPipe[0]);
 				UBA_ASSERTF(!res, "posix_spawn_file_actions_addclose outPipe[0] failed: %i", res);
 
diff --git a/Engine/Source/Programs/UnrealBuildAccelerator/Common/Public/UbaSocket.inl b/Engine/Source/Programs/UnrealBuildAccelerator/Common/Public/UbaSocket.inl
--- a/Engine/Source/Programs/UnrealBuildAccelerator/Common/Public/UbaSocket.inl
+++ b/Engine/Source/Programs/UnrealBuildAccelerator/Common/Public/UbaSocket.inl
@@ -8,6 +8,8 @@ Socket SocketCreate(int family, int type, int protocol)
 {
 #if PLATFORM_WINDOWS
 	return FromPlatformSocket(WSASocketW(family, type, protocol, NULL, 0, WSA_FLAG_OVERLAPPED));
+#elif PLATFORM_LINUX
+	return FromPlatformSocket(::socket(family, type | SOCK_CLOEXEC, protocol));
 #else
 	return FromPlatformSocket(::socket(family, type, protocol));
 #endif
@@ -222,7 +224,11 @@ int SocketPoll(Socket s, int events, int timeoutMs, int* outRevents)
 Socket SocketAccept(Socket s, sockaddr* outAddr, int addrLen)
 {
 	socklen_t remoteSockAddrLen = addrLen;
+#if PLATFORM_LINUX
+	return FromPlatformSocket(accept4(ToPlatformSocket(s), outAddr, &remoteSockAddrLen, SOCK_CLOEXEC));
+#else
 	return FromPlatformSocket(accept(ToPlatformSocket(s), outAddr, &remoteSockAddrLen));
+#endif
 }
 
 int SocketRecv(Socket s, void* data, int len)
diff --git a/Engine/Source/Programs/UnrealBuildAccelerator/Detours/Private/PosixOS/UbaDetoursFunctionsPosix.cpp b/Engine/Source/Programs/UnrealBuildAccelerator/Detours/Private/PosixOS/UbaDetoursFunctionsPosix.cpp
--- a/Engine/Source/Programs/UnrealBuildAccelerator/Detours/Private/PosixOS/UbaDetoursFunctionsPosix.cpp
+++ b/Engine/Source/Programs/UnrealBuildAccelerator/Detours/Private/PosixOS/UbaDetoursFunctionsPosix.cpp
@@ -1287,10 +1287,25 @@ UBA_EXPORT ssize_t UBA_WRAPPER(write)(int fd, const void* buf, size_t count)
 //	return -1;
 //}
 
+void DuplicateFd(int oldFd, int newFd)
+{
+	SCOPED_WRITE_LOCK(g_fileHandlesLock, lock);
+	auto findIt = g_fileHandles.find(oldFd);
+	if (findIt == g_fileHandles.end())
+		return;
+	DetouredHandle& h = findIt->second;
+	FileObject* fo = h.fileObject;
+	++fo->refCount;
+	g_fileHandles[newFd].fileObject = fo;
+}
+
+
 UBA_EXPORT int UBA_WRAPPER(dup)(int oldfd)
 {
 	UBA_INIT_DETOUR(dup, oldfd);
 	auto res = TRUE_WRAPPER(dup)(oldfd);
+	if (res != -1)
+		DuplicateFd(oldfd, res);
 	DEBUG_LOG_TRUE("dup", "(%i) -> %i", oldfd, res);
 	return res;
 }
@@ -1299,20 +1314,8 @@ UBA_EXPORT int UBA_WRAPPER(dup2)(int oldfd, int newfd)
 {
 	UBA_INIT_DETOUR(dup2, oldfd, newfd);
 	auto res = TRUE_WRAPPER(dup2)(oldfd, newfd);
-
 	if (res != -1)
-	{
-		SCOPED_WRITE_LOCK(g_fileHandlesLock, lock);
-		auto findIt = g_fileHandles.find(oldfd);
-		if (findIt != g_fileHandles.end())
-		{
-			DetouredHandle& h = findIt->second;
-			FileObject* fo = h.fileObject;
-			++fo->refCount;
-			g_fileHandles[newfd].fileObject = fo;
-		}
-	}
-
+		DuplicateFd(oldfd, newfd);
 	DEBUG_LOG_TRUE("dup2", "(%i, %i) -> %i", oldfd, newfd, res);
 	return res;
 }
diff --git a/Engine/Source/Programs/UnrealBuildAccelerator/Detours/Private/PosixOS/UbaDetoursMainPosix.cpp b/Engine/Source/Programs/UnrealBuildAccelerator/Detours/Private/PosixOS/UbaDetoursMainPosix.cpp
--- a/Engine/Source/Programs/UnrealBuildAccelerator/Detours/Private/PosixOS/UbaDetoursMainPosix.cpp
+++ b/Engine/Source/Programs/UnrealBuildAccelerator/Detours/Private/PosixOS/UbaDetoursMainPosix.cpp
@@ -79,6 +79,7 @@ static void __attribute__((constructor(102))) PreInitCtor()
 		printf("UbaDetours: Failed to mmap fd: %u\n", g_comFd);
 		return;
 	}
+	fcntl(g_comFd, F_SETFD, FD_CLOEXEC); // Don't let grandchild processes inherit this fd
 
 	g_messageMappingMem = rptr;
 
