[UBA]
* Fixed non-unity compile error in debug in UbaFileAccessor.cpp

diff --git a/Engine/Source/Programs/UnrealBuildAccelerator/Common/Private/UbaFileAccessor.cpp b/Engine/Source/Programs/UnrealBuildAccelerator/Common/Private/UbaFileAccessor.cpp
--- a/Engine/Source/Programs/UnrealBuildAccelerator/Common/Private/UbaFileAccessor.cpp
+++ b/Engine/Source/Programs/UnrealBuildAccelerator/Common/Private/UbaFileAccessor.cpp
@@ -1,6 +1,7 @@
 // Copyright Epic Games, Inc. All Rights Reserved.
 
 #include "UbaFileAccessor.h"
+#include "UbaDefinitions.h"
 #include "UbaEvent.h"
 #include "UbaFile.h"
 #include "UbaLogger.h"
