[UBA]
* Changed so handles coming from other process can be detouring range without being detoured. This because cooking our biggest game cause the host process to reach detouring handle range

diff --git a/Engine/Source/Programs/UnrealBuildAccelerator/Detours/Private/Windows/UbaDetoursFunctionsKernelBase.inl b/Engine/Source/Programs/UnrealBuildAccelerator/Detours/Private/Windows/UbaDetoursFunctionsKernelBase.inl
--- a/Engine/Source/Programs/UnrealBuildAccelerator/Detours/Private/Windows/UbaDetoursFunctionsKernelBase.inl
+++ b/Engine/Source/Programs/UnrealBuildAccelerator/Detours/Private/Windows/UbaDetoursFunctionsKernelBase.inl
@@ -182,8 +182,15 @@ BOOL Detoured_DuplicateHandle(HANDLE hSourceProcessHandle, HANDLE hSourceHandle,
 		DEBUG_LOG_TRUE(L"DuplicateHandle",L"%llu (%llu) to %llu (%llu) -> %ls", uintptr_t(hSourceHandle), uintptr_t(hSourceProcessHandle), lpTargetHandle ? uintptr_t(*lpTargetHandle) : 0ull, uintptr_t(hTargetProcessHandle), ToString(res));
 		return res;
 	}
+	if (hSourceProcessHandle != hTargetProcessHandle)
+	{
+		// This is not good. Means that the hosting process likely has hit the detour handle range. Large cooks can cause this
+		//UBA_ASSERTF(false, TC("Duplicating detoured handles from one process to the other not supported..(SourceProcess: %llu SourceHandle: %llu)"), u64(hSourceProcessHandle), u64(hSourceHandle));
+		auto res = True_DuplicateHandle(hSourceProcessHandle, hSourceHandle, hTargetProcessHandle, lpTargetHandle, dwDesiredAccess, bInheritHandle, dwOptions);
+		DEBUG_LOG_TRUE(L"DuplicateHandle",L"%llu (%llu) to %llu (%llu) -> %ls", uintptr_t(hSourceHandle), uintptr_t(hSourceProcessHandle), lpTargetHandle ? uintptr_t(*lpTargetHandle) : 0ull, uintptr_t(hTargetProcessHandle), ToString(res));
+		return res;
+	}
 
-	UBA_ASSERT(hSourceProcessHandle == hTargetProcessHandle);
 	UBA_ASSERT(!(dwOptions & DUPLICATE_CLOSE_SOURCE));
 	auto& dh = asDetouredHandle(hSourceHandle);
 
