[UBA]
* Fixed handle leak if WSAEventSelect fails
* Fixed crash if no coordinator was set and RequestHelpers is called
* Added detouring of posix execv

diff --git a/Engine/Source/Programs/UnrealBuildAccelerator/Common/Private/UbaNetworkBackendTcp.cpp b/Engine/Source/Programs/UnrealBuildAccelerator/Common/Private/UbaNetworkBackendTcp.cpp
--- a/Engine/Source/Programs/UnrealBuildAccelerator/Common/Private/UbaNetworkBackendTcp.cpp
+++ b/Engine/Source/Programs/UnrealBuildAccelerator/Common/Private/UbaNetworkBackendTcp.cpp
@@ -188,7 +188,10 @@ namespace uba
 			if (ev == WSA_INVALID_EVENT)
 				return;
 			if (WSAEventSelect(ToPlatformSocket(s), ev, FD_WRITE | FD_CLOSE) == SOCKET_ERROR)
+			{
+				WSACloseEvent(wsaEvent);
 				return;
+			}
 			wsaEvent = ev;
 			#endif
 		}
diff --git a/Engine/Source/Programs/UnrealBuildAccelerator/Common/Public/UbaCoordinatorWrapper.h b/Engine/Source/Programs/UnrealBuildAccelerator/Common/Public/UbaCoordinatorWrapper.h
--- a/Engine/Source/Programs/UnrealBuildAccelerator/Common/Public/UbaCoordinatorWrapper.h
+++ b/Engine/Source/Programs/UnrealBuildAccelerator/Common/Public/UbaCoordinatorWrapper.h
@@ -103,6 +103,8 @@ namespace uba
 
 		bool RequestHelpers(u32 maxCoreCount)
 		{
+			if (!m_coordinator)
+				return false;
 			if (m_loopCoordinator.IsCreated())
 				return false;
 			m_loopCoordinator.Create(EventResetType_Manual);
diff --git a/Engine/Source/Programs/UnrealBuildAccelerator/Detours/Private/PosixOS/UbaDetoursFunctionsPosix.cpp b/Engine/Source/Programs/UnrealBuildAccelerator/Detours/Private/PosixOS/UbaDetoursFunctionsPosix.cpp
--- a/Engine/Source/Programs/UnrealBuildAccelerator/Detours/Private/PosixOS/UbaDetoursFunctionsPosix.cpp
+++ b/Engine/Source/Programs/UnrealBuildAccelerator/Detours/Private/PosixOS/UbaDetoursFunctionsPosix.cpp
@@ -2502,13 +2502,6 @@ UBA_EXPORT int UBA_WRAPPER(dladdr)(const void *addr, Dl_info *info)
 	return res;
 }
 
-UBA_EXPORT int UBA_WRAPPER(execv)(const char* path, char* const argv[])
-{
-	UBA_INIT_DETOUR(execv, path, argv);
-	DEBUG_LOG_TRUE("execv", "%s", path);
-	return TRUE_WRAPPER(execv)(path, argv);
-}
-
 int Internal_execve(const char* pathname, char* const _Nullable argv[], char* const _Nullable envp[])
 {
 	UBA_INIT_DETOUR(execve, pathname, argv, envp);
@@ -2561,6 +2554,11 @@ int Internal_execve(const char* pathname, char* const _Nullable argv[], char* co
 	//return -1;
 }
 
+UBA_EXPORT int UBA_WRAPPER(execv)(const char* path, char* const argv[])
+{
+	return Internal_execve(path, argv, environ);
+}
+
 UBA_EXPORT int UBA_WRAPPER(execve)(const char* pathname, char* const _Nullable argv[], char* const _Nullable envp[])
 {
 	return Internal_execve(pathname, argv, envp);
