[UBA]
* Fixed assert deadlocking issue on posix
* Fixed so exec* handles shell scripts
* Removed comment that was wrong

diff --git a/Engine/Source/Programs/UnrealBuildAccelerator/Detours/Private/PosixOS/UbaDetoursFunctionsPosix.cpp b/Engine/Source/Programs/UnrealBuildAccelerator/Detours/Private/PosixOS/UbaDetoursFunctionsPosix.cpp
--- a/Engine/Source/Programs/UnrealBuildAccelerator/Detours/Private/PosixOS/UbaDetoursFunctionsPosix.cpp
+++ b/Engine/Source/Programs/UnrealBuildAccelerator/Detours/Private/PosixOS/UbaDetoursFunctionsPosix.cpp
@@ -44,6 +44,7 @@ namespace uba
 	u32 g_processId;
 	pid_t g_pid;
 	void Deinit();
+	bool TakeLockForRpc();
 
 	SharedMemoryAllocatorHandle g_sharedMemoryAllocator;
 	u8* g_sharedMemory;
@@ -2214,6 +2215,22 @@ int shared_posix_spawn(pid_t* pid, const char* path, const posix_spawn_file_acti
 
 		DevirtualizePath(result);
 		path = result.data;
+
+		// If the executable is a shell script, wrap it with /bin/sh so the kernel
+		// does not need to parse the shebang — and so UBA's virtual FS interception
+		// does not prevent the script from being opened by the kernel's exec path.
+		u32 plen = result.count;
+		if (plen >= 3 && result.data[plen-3] == '.' && result.data[plen-2] == 's' && result.data[plen-1] == 'h')
+		{
+			u32 argc2 = 0;
+			tempArgv[argc2++] = "/bin/sh";
+			tempArgv[argc2++] = path; // the .sh script (still in result.data)
+			for (int i = 1; argv[i]; ++i)
+				tempArgv[argc2++] = argv[i];
+			tempArgv[argc2] = nullptr;
+			argv = (char* const*)tempArgv;
+			path = "/bin/sh";
+		}
 	}
 	//DEBUG_LOG("LIBRARY_SEARCH_PATHS: %s", getenv("LIBRARY_SEARCH_PATHS"));
 	//DEBUG_LOG("PATH: %s", getenv("PATH"));
@@ -2221,9 +2238,25 @@ int shared_posix_spawn(pid_t* pid, const char* path, const posix_spawn_file_acti
 	TString cmdLineWithoutApplication;
 	for (u32 i = 1; argv[i]; ++i)
 	{
-		if (i != 0)
-			cmdLineWithoutApplication.append(" ");
-		cmdLineWithoutApplication.append(argv[i]);
+		cmdLineWithoutApplication.append(" ");
+		const char* arg = argv[i];
+		// Quote arguments that contain spaces or double-quotes so ParseArguments
+		// on the receiving end reconstructs the correct argv (e.g. sed -e scripts).
+		bool needsQuoting = false;
+		for (const char* p = arg; *p; ++p)
+			if (*p == ' ' || *p == '\t' || *p == '"') { needsQuoting = true; break; }
+		if (needsQuoting)
+		{
+			cmdLineWithoutApplication.append("\"");
+			for (const char* p = arg; *p; ++p)
+			{
+				if (*p == '"') cmdLineWithoutApplication.append("\\\"");
+				else           cmdLineWithoutApplication += *p;
+			}
+			cmdLineWithoutApplication.append("\"");
+		}
+		else
+			cmdLineWithoutApplication.append(arg);
 	}
 
 	TString realApplication;
@@ -3070,8 +3103,13 @@ namespace uba
 		SuppressDetourScope s;
 		static CriticalSection cs;
 		ScopedCriticalSection scs(cs);
+
+		#if UBA_DEBUG_LOG_ENABLED
+		FlushDebugLog();
+		#endif
+
 		static auto& sb =  *new StringBuffer<8*1024>;
-		WriteAssertInfo(sb, text, file, line, expr, context);
+		WriteAssertInfo(sb.Clear(), text, file, line, expr, context);
 
 		if (g_isInitialized)
 		{
@@ -3083,12 +3121,16 @@ namespace uba
 
 			g_stats.detoursMemory = g_memoryBlock.writtenSize;
 
-			RPC_MESSAGE(Exit, log);
-			writer.WriteU32(terminateCode); // Exit code
-			writer.WriteString(""); // Log name
-			g_stats.Write(writer);
-			g_kernelStats.Write(writer);
-			writer.Flush(false);
+			if (TakeLockForRpc())
+			{
+				RPC_MESSAGE_NO_LOCK(Exit, log);
+				writer.WriteU32(terminateCode); // Exit code
+				writer.WriteString(""); // Log name
+				g_stats.Write(writer);
+				g_kernelStats.Write(writer);
+				writer.Flush(false);
+				g_communicationLock.Leave();
+			}
 		}
 		else
 		{
diff --git a/Engine/Source/Programs/UnrealBuildAccelerator/Detours/Private/UbaDetoursShared.cpp b/Engine/Source/Programs/UnrealBuildAccelerator/Detours/Private/UbaDetoursShared.cpp
--- a/Engine/Source/Programs/UnrealBuildAccelerator/Detours/Private/UbaDetoursShared.cpp
+++ b/Engine/Source/Programs/UnrealBuildAccelerator/Detours/Private/UbaDetoursShared.cpp
@@ -27,6 +27,10 @@ namespace uba
 	VARIABLE_MEM(ReaderWriterLock, g_consoleStringCs);
 	VARIABLE_MEM(Futex, g_trackInputsLock);
 
+	#if UBA_DEBUG_LOG_ENABLED
+	VARIABLE_MEM(StringBuffer<MaxPath>, g_debugLogPath);
+	#endif
+
 	bool g_conEnabled[2];
 	u32 g_rulesIndex;
 	ApplicationRules* g_rules;
@@ -56,6 +60,10 @@ namespace uba
 		g_systemRootMem.Create();
 		g_systemTempMem.Create();
 
+		#if UBA_DEBUG_LOG_ENABLED
+		g_debugLogPathMem.Create();
+		#endif
+
 		g_vfsEntryCount = 0;
 		g_vfsMatchingLength = 0;
 
@@ -480,7 +488,6 @@ namespace uba
 		{
 			out.Append(TCV("\n   Failed to resolve callstack\n"));
 		}
-		// Note, we leave the lock even though we might not have it because we want to be able to report
 		g_communicationLock.Leave();
 	}
 
diff --git a/Engine/Source/Programs/UnrealBuildAccelerator/Detours/Private/UbaDetoursShared.h b/Engine/Source/Programs/UnrealBuildAccelerator/Detours/Private/UbaDetoursShared.h
--- a/Engine/Source/Programs/UnrealBuildAccelerator/Detours/Private/UbaDetoursShared.h
+++ b/Engine/Source/Programs/UnrealBuildAccelerator/Detours/Private/UbaDetoursShared.h
@@ -79,6 +79,7 @@ namespace uba
 	#if UBA_DEBUG_LOG_ENABLED
 	inline constexpr u32 LogBufSize = 16 * 1024;
 	extern FileHandle g_debugFile;
+	extern StringBuffer<MaxPath>& g_debugLogPath;
 	inline bool isLogging() { return g_debugFile != InvalidFileHandle; }
 	struct LogScope { LogScope(); ~LogScope(); void Flush(); };
 	void WriteDebugLogWithPrefix(const char* prefix, LogScope& scope, const tchar* command, const tchar* format, ...);
