[UBA]
* Fixed assert when processes that was ready to be tested (still not in queue) phase got moved to readyToSkip

diff --git a/Engine/Source/Programs/UnrealBuildAccelerator/Common/Private/UbaScheduler.cpp b/Engine/Source/Programs/UnrealBuildAccelerator/Common/Private/UbaScheduler.cpp
--- a/Engine/Source/Programs/UnrealBuildAccelerator/Common/Private/UbaScheduler.cpp
+++ b/Engine/Source/Programs/UnrealBuildAccelerator/Common/Private/UbaScheduler.cpp
@@ -492,6 +492,7 @@ namespace uba
 				continue;
 			if (dependency.runStatus == RunStatus_Failed || dependency.runStatus == RunStatus_Skipped || m_cancelled)
 			{
+				entry.cacheStatus = CacheStatus_Failed;
 				AddToQueue(m_readyToSkip, entryIndex);
 				return; // Not counted in stats — will be skipped immediately
 			}
@@ -610,11 +611,13 @@ namespace uba
 				// In test queue but not yet started — remove and skip directly.
 				TransitionStats(dependent, StatsState_None, weight); // left → none
 				m_readyToTestCache.erase(dependentIndex);
+				dependent.cacheStatus = CacheStatus_Failed;
 				AddToQueue(m_readyToSkip, dependentIndex);
 			}
 			else if (dependent.runStatus == RunStatus_QueuedForRun)
 			{
 				TransitionStats(dependent, StatsState_None, weight); // left → none
+				dependent.cacheStatus = CacheStatus_Failed;
 				AddToQueue(m_readyToSkip, dependentIndex);
 			}
 			else if (dependent.cacheStatus == CacheStatus_PendingCompletion)
@@ -1615,7 +1618,7 @@ namespace uba
 			auto info = processEntry.info;
 			bool canDetour = processEntry.canDetour && !m_forceNative;
 			bool wasSkipped = processEntry.runStatus == RunStatus_Skipped;
-			UBA_ASSERT(processEntry.cacheStatus == CacheStatus_Failed);
+			UBA_ASSERTF(processEntry.cacheStatus == CacheStatus_Failed, TC("cacheStatus: %u, runStatus: %u"), u32(processEntry.cacheStatus), u32(processEntry.runStatus));
 			const Category runningCategory = GetRunningCategory(isLocal, isCrossArchitecture, isCrossPlatform);
 
 			if (!wasSkipped && !m_cancelled)
