[UBA]
* Reimplemented fork/vfork detouring code paths in uba. Now there are no hacky thread_local variables and additional in-between processes. This is much tighter and correct than the previous implementation and should solve a bunch of edgecases
* Fixed a whole bunch of bugs around posix. Wrong file sizes could be reported, PWD had a slash at the end which violates posix rules. Multiple smaller things

diff --git a/Engine/Source/Programs/UnrealBuildAccelerator/Common/Private/UbaFile.cpp b/Engine/Source/Programs/UnrealBuildAccelerator/Common/Private/UbaFile.cpp
--- a/Engine/Source/Programs/UnrealBuildAccelerator/Common/Private/UbaFile.cpp
+++ b/Engine/Source/Programs/UnrealBuildAccelerator/Common/Private/UbaFile.cpp
@@ -871,7 +871,7 @@ namespace uba
 			SetLastError(ERROR_SUCCESS);
 			return true;
 		}
-		UBA_ASSERTF(false, TC("MoveFileExW error handling not implemented"));
+		UBA_ASSERTF(false, TC("MoveFileExW error handling not implemented (%s to %s): %s"), existingFileName, newFileName, strerror(errno));
 		return false;
 #endif
 	}
diff --git a/Engine/Source/Programs/UnrealBuildAccelerator/Common/Private/UbaProcess.cpp b/Engine/Source/Programs/UnrealBuildAccelerator/Common/Private/UbaProcess.cpp
--- a/Engine/Source/Programs/UnrealBuildAccelerator/Common/Private/UbaProcess.cpp
+++ b/Engine/Source/Programs/UnrealBuildAccelerator/Common/Private/UbaProcess.cpp
@@ -1337,45 +1337,13 @@
 		return true;
 	}
 
-	bool ProcessImpl::HandleCreateProcess(BinaryReader& reader, BinaryWriter& writer)
+	#if PLATFORM_LINUX
+	// This process will become the parent of a process if it becomes orphaned
+	void SetSubReaper() { static bool subreaper = []() { prctl(PR_SET_CHILD_SUBREAPER, 1, 0, 0, 0); return true; }(); }
+	#endif
+
+	bool ProcessImpl::WriteCreateProcessResponse(ProcessHandle h, BinaryWriter& writer, bool isChild)
 	{
-		#if PLATFORM_LINUX
-		// This process will become the parent of a process if it becomes orphaned
-		static bool subreaper = []() { prctl(PR_SET_CHILD_SUBREAPER, 1, 0, 0, 0); return true; }();
-		#endif
-
-		TString applicationStr = reader.ReadString();
-		TString commandLineWithoutApplication = reader.ReadLongString();
-
-		UBA_ASSERTF(!applicationStr.empty() && (applicationStr[0] != '\"' || applicationStr[1] != '\"'), TC("Invalid application name: %s  (Length %llu)"), applicationStr.c_str(), applicationStr.size());
-
-		StringBuffer<> currentDir;
-		reader.ReadString(currentDir);
-		if (currentDir.IsEmpty())
-			currentDir.Append(m_startInfo.workingDir);
-		else if (!IsAbsolutePath(currentDir.data))
-		{
-			TString temp(currentDir.ToString());
-			FixPath(temp.c_str(), m_startInfo.workingDir, m_startInfo.workingDirStr.size(), currentDir.Clear());
-		}
-		bool startSuspended = reader.ReadBool();
-		bool isChild = reader.ReadBool();
-
-		StringBuffer<> temp;
-		ProcessStartInfo info;
-		info.application = applicationStr.c_str();
-		info.arguments = commandLineWithoutApplication.c_str();
-		info.workingDir = currentDir.data;
-		info.logFile = InternalGetChildLogFile(temp);
-		info.priorityClass = m_startInfo.priorityClass;
-		info.logLineUserData = this;
-		info.logLineFunc = [](void* userData, const tchar* line, u32 length, LogEntryType type) { ((ProcessImpl*)userData)->LogLine(false, TString(line, length), type); };
-		info.startSuspended = startSuspended;
-		info.rootsHandle = m_startInfo.rootsHandle;
-		info.trackInputs = m_startInfo.trackInputs;
-
-		ProcessImpl* parent = isChild ? this : nullptr;
-		ProcessHandle h = m_session.InternalRunProcess(info, true, parent, true, nullptr);
 		if (!h.m_process)
 		{
 			// TODO: Do we need to log something here? This failure is most likely a cause of something else so error logging might show somewhere else
@@ -1397,10 +1365,11 @@
 		UBA_ASSERT(detoursLibLen < 1024);
 
 		writer.WriteU32(childProcessId);
-		writer.WriteU32(info.rules->index);
-		
+		writer.WriteU32(process.m_startInfo.rules->index);
+
 		writer.WriteU32(detoursLibLen);
 		writer.WriteBytes(detoursLib.data(), detoursLibLen);
+		
 		#if PLATFORM_WINDOWS
 		if (m_session.m_sharedMemory.GetFileMappingBackend().linuxBackend)
 		{
@@ -1411,6 +1380,7 @@
 			writer.WriteBytes(ubaWineModule, ubaWineModulePath.count);
 		}
 		#endif
+		
 		writer.WriteU32(0);
 
 		writer.WriteString(m_realWorkingDir);
@@ -1424,8 +1394,53 @@
 		#if !PLATFORM_WINDOWS
 		writer.WriteU64(process.m_comMemory.handle.ToU64());
 		writer.WriteU32(process.m_comMemory.offset);
-		writer.WriteString(info.logFile);
+		writer.WriteString(process.m_startInfo.logFile);
+		#endif
+		
+		return true;
+	}
+
+	bool ProcessImpl::HandleCreateProcess(BinaryReader& reader, BinaryWriter& writer)
+	{
+		#if PLATFORM_LINUX
+		SetSubReaper();
 		#endif
 
+		TString applicationStr = reader.ReadString();
+		TString commandLineWithoutApplication = reader.ReadLongString();
+
+		UBA_ASSERTF(!applicationStr.empty() && (applicationStr[0] != '\"' || applicationStr[1] != '\"'), TC("Invalid application name: %s  (Length %llu)"), applicationStr.c_str(), applicationStr.size());
+
+		StringBuffer<> currentDir;
+		reader.ReadString(currentDir);
+		if (currentDir.IsEmpty())
+			currentDir.Append(m_startInfo.workingDir);
+		else if (!IsAbsolutePath(currentDir.data))
+		{
+			TString temp(currentDir.ToString());
+			FixPath(temp.c_str(), m_startInfo.workingDir, m_startInfo.workingDirStr.size(), currentDir.Clear());
+		}
+		bool startSuspended = reader.ReadBool();
+		bool isChild = reader.ReadBool();
+
+		StringBuffer<> temp;
+		ProcessStartInfo info;
+		info.application = applicationStr.c_str();
+		info.arguments = commandLineWithoutApplication.c_str();
+		info.workingDir = currentDir.data;
+		info.logFile = InternalGetChildLogFile(temp);
+		info.priorityClass = m_startInfo.priorityClass;
+		info.logLineUserData = this;
+		info.logLineFunc = [](void* userData, const tchar* line, u32 length, LogEntryType type) { ((ProcessImpl*)userData)->LogLine(false, TString(line, length), type); };
+		info.startSuspended = startSuspended;
+		info.rootsHandle = m_startInfo.rootsHandle;
+		info.trackInputs = m_startInfo.trackInputs;
+
+		ProcessImpl* parent = isChild ? this : nullptr;
+		ProcessHandle h = m_session.InternalRunProcess(info, true, parent, true, nullptr);
+
+		if (!WriteCreateProcessResponse(h, writer, isChild))
+			return false;
+
 		#if UBA_DEBUG_TRACK_PROCESS
 		m_session.m_debugLogger->Info(TC("CreateChildProcess (%u creating child %u at index %u) %s %s (%s)"), m_id, process.m_id, childProcessId-1, process.m_realApplication.c_str(), commandLineWithoutApplication.c_str(), info.logFile);
@@ -1436,7 +1451,7 @@
 	bool ProcessImpl::HandleStartProcess(BinaryReader& reader, BinaryWriter& writer)
 	{
 		u32 processId = reader.ReadU32();
-		UBA_ASSERTF(processId > 0 && processId <= m_childProcesses.size(), TC("Got process id %u, expected number lower than %u"), processId, u32(m_childProcesses.size()));
+		UBA_ASSERTF(processId > 0 && processId <= m_childProcesses.size(), TC("Got process id %u for child, expected number lower than %u"), processId, u32(m_childProcesses.size()));
 		auto& process = *(ProcessImpl*)m_childProcesses[processId - 1].m_process;
 		bool result = reader.ReadBool();
 		u32 lastError = reader.ReadU32();
@@ -1500,6 +1515,86 @@
 		return true;
 	}
 
+	bool ProcessImpl::HandleForkProcess(BinaryReader& reader, BinaryWriter& writer)
+	{
+	#if PLATFORM_WINDOWS
+		return false;
+	#else
+		#if PLATFORM_LINUX
+		SetSubReaper();
+		#endif
+
+		StringBuffer<> temp;
+		ProcessStartInfo info;
+		info = m_startInfo;
+		info.logFile = InternalGetChildLogFile(temp);
+		info.exitedFunc = nullptr;
+
+		ProcessHandle h = m_session.InternalRunProcess(info, true, this, true, nullptr);
+		if (!h.m_process)
+		{
+			// TODO: Do we need to log something here? This failure is most likely a cause of something else so error logging might show somewhere else
+			writer.WriteU32(0); // childProcessId
+			return true;
+		}
+
+		m_childProcesses.push_back(h);
+		u32 childProcessId = u32(m_childProcesses.size());
+		writer.WriteU32(childProcessId);
+
+		auto& process = *(ProcessImpl*)h.m_process;
+		writer.WriteU64(process.m_comMemory.handle.ToU64());
+		writer.WriteU32(process.m_comMemory.offset);
+		writer.WriteString(info.logFile);
+		return true;
+	#endif
+	}
+
+	bool ProcessImpl::HandleExecProcess(BinaryReader& reader, BinaryWriter& writer)
+	{
+	#if PLATFORM_WINDOWS
+		return false;
+	#else
+		TString applicationStr = reader.ReadString();
+		TString commandLineWithoutApplication = reader.ReadLongString();
+
+		UBA_ASSERTF(!applicationStr.empty() && (applicationStr[0] != '\"' || applicationStr[1] != '\"'), TC("Invalid application name: %s  (Length %llu)"), applicationStr.c_str(), applicationStr.size());
+
+		StringBuffer<> currentDir;
+		reader.ReadString(currentDir);
+		if (currentDir.IsEmpty())
+			currentDir.Append(m_startInfo.workingDir);
+		else if (!IsAbsolutePath(currentDir.data))
+		{
+			TString temp(currentDir.ToString());
+			FixPath(temp.c_str(), m_startInfo.workingDir, m_startInfo.workingDirStr.size(), currentDir.Clear());
+		}
+
+		m_startInfo.applicationStr = applicationStr.c_str();
+		m_startInfo.argumentsStr = commandLineWithoutApplication.c_str();
+		m_startInfo.workingDirStr = currentDir.data;
+
+		m_startInfo.application = m_startInfo.applicationStr.c_str();
+		m_startInfo.arguments = m_startInfo.argumentsStr.c_str();
+		m_startInfo.workingDir = m_startInfo.workingDirStr.c_str();
+
+		m_startInfo.rules = m_session.GetRules(m_startInfo);
+
+		// TODO: Lots of fixes here... m_startInfo should be fixed up here to be the correct (Kind of like HandleCreateProcess)
+		StringBuffer<> realApplication(m_startInfo.application);
+		const tchar* realWorkingDir = m_startInfo.workingDir;
+		bool res = m_session.PrepareProcess(*this, m_parentProcess != nullptr, realApplication, realWorkingDir);
+		UBA_ASSERT(res);
+		m_realApplication = realApplication.data;
+		m_realWorkingDir = realWorkingDir;
+		if (realWorkingDir == m_startInfo.workingDir)
+			m_realWorkingDir = m_startInfo.workingDir;
+
+		ProcessHandle h(this);
+		return WriteCreateProcessResponse(h, writer, false); // isChild is false because we have already registered the child when for was created
+	#endif
+	}
+
 	bool ProcessImpl::HandleExitChildProcess(BinaryReader& reader, BinaryWriter& writer)
 	{
 		// TODO: Should this traverse down? This is used for when child processes are terminated through job being closed. used by node.exe (building chrome)
@@ -2318,10 +2413,11 @@
 			
 			if (*m_realWorkingDir)
 			{
-				#if PLATFORM_MAC
+				// posix_spawn_file_actions_addchdir_np: available on macOS and Linux glibc >= 2.29
+				// (Ubuntu 19.04+, Debian 10+). For older glibc the child inherits the parent CWD,
+				// which works as long as UbaNinja is invoked from the build directory.
+				#if PLATFORM_MAC || (defined(__GLIBC__) && (__GLIBC__ > 2 || (__GLIBC__ == 2 && __GLIBC_MINOR__ >= 29)))
 				posix_spawn_file_actions_addchdir_np(&fileActions, m_realWorkingDir);
-				#else
-				//UBA_ASSERT(false); // TODO: Revisit
 				#endif
 			}
 
diff --git a/Engine/Source/Programs/UnrealBuildAccelerator/Common/Private/UbaProcessStartInfoHolder.cpp b/Engine/Source/Programs/UnrealBuildAccelerator/Common/Private/UbaProcessStartInfoHolder.cpp
--- a/Engine/Source/Programs/UnrealBuildAccelerator/Common/Private/UbaProcessStartInfoHolder.cpp
+++ b/Engine/Source/Programs/UnrealBuildAccelerator/Common/Private/UbaProcessStartInfoHolder.cpp
@@ -45,8 +45,44 @@ namespace uba
 		if (!firstArgEnd)
 			return false;
 		return InternalExpand(firstArgBegin, firstArgEnd);
-	#else
-		return false;
+	  #else
+		bool requiresShell = false;
+
+		// Wrap any .sh script (may lack execute bit or have a virtualised shebang)
+		u32 appLen = u32(applicationStr.size());
+		if (appLen >= 3 && applicationStr[appLen-3] == TC('.') && applicationStr[appLen-2] == TC('s') && applicationStr[appLen-1] == TC('h'))
+			requiresShell = true;
+
+		// Check for shell metacharacters via token parsing
+		if (!requiresShell)
+		{
+			ParseArguments(argumentsStr.c_str(), argumentsStr.size(), [&](const tchar* arg, u32 argLen)
+			{
+				StringView view(arg, argLen);
+				requiresShell |= argLen == 1 && (*arg == '<' || *arg == '>' || *arg == '|' || *arg == ';');
+				requiresShell |= argLen == 2 && (view.Equals(TCV("&&")) || view.Equals(TCV("||")));
+			});
+		}
+
+		if (requiresShell)
+		{
+			// Build /bin/sh -c "escaped_command" — escape \ and " inside the outer double-quotes
+			TString fullCmd = applicationStr + TC(" ") + argumentsStr;
+			TString wrapped = TC("-c \"");
+			wrapped.reserve(fullCmd.size() + 8);
+			for (tchar c : fullCmd)
+			{
+				if      (c == TC('\\')) wrapped += TC("\\\\");
+				else if (c == TC('"'))  wrapped += TC("\\\"");
+				else                    wrapped += c;
+			}
+			wrapped += TC('"');
+			argumentsStr = std::move(wrapped);
+			arguments = argumentsStr.c_str();
+			applicationStr = TC("/bin/sh");
+			application = applicationStr.c_str();
+		}
+		return true;
 	#endif
 	}
 
@@ -152,8 +188,9 @@ namespace uba
 			arguments = argumentsStr.c_str();
 			return true;
 		}
-	#endif
-		return false;
+	  #else
+		  return false;
+	  #endif
 	}
 
 	void ProcessStartInfoHolder::Apply(const Config& config, const uba::tchar* configTable)
diff --git a/Engine/Source/Programs/UnrealBuildAccelerator/Common/Private/UbaSessionServer.cpp b/Engine/Source/Programs/UnrealBuildAccelerator/Common/Private/UbaSessionServer.cpp
--- a/Engine/Source/Programs/UnrealBuildAccelerator/Common/Private/UbaSessionServer.cpp
+++ b/Engine/Source/Programs/UnrealBuildAccelerator/Common/Private/UbaSessionServer.cpp
@@ -2476,7 +2476,6 @@ namespace uba
 			m_remoteProcessSlotAvailableEvent(IsArmBinary != session.isArm, IsLinux != session.isLinux);
 			hasCalledCallback = true;
 		}
-		return nullptr;
 	}
 
 	void SessionServer::OnCancelled(RemoteProcess* process)
diff --git a/Engine/Source/Programs/UnrealBuildAccelerator/Common/Public/UbaProcess.h b/Engine/Source/Programs/UnrealBuildAccelerator/Common/Public/UbaProcess.h
--- a/Engine/Source/Programs/UnrealBuildAccelerator/Common/Public/UbaProcess.h
+++ b/Engine/Source/Programs/UnrealBuildAccelerator/Common/Public/UbaProcess.h
@@ -91,6 +91,7 @@ namespace uba
 		#define UBA_PROCESS_MESSAGE(M) bool Handle##M(BinaryReader& reader, BinaryWriter& writer);
 		UBA_PROCESS_MESSAGES
 		#undef UBA_PROCESS_MESSAGE
+		bool WriteCreateProcessResponse(ProcessHandle h, BinaryWriter& writer, bool isChild);
 
 		void LogLine(bool printInSession, TString&& line, LogEntryType logType);
 		bool CreateTempFile(BinaryReader& reader, ProcHandle nativeProcessHandle, const tchar* application);
diff --git a/Engine/Source/Programs/UnrealBuildAccelerator/Core/Public/UbaProcessUtils.h b/Engine/Source/Programs/UnrealBuildAccelerator/Core/Public/UbaProcessUtils.h
--- a/Engine/Source/Programs/UnrealBuildAccelerator/Core/Public/UbaProcessUtils.h
+++ b/Engine/Source/Programs/UnrealBuildAccelerator/Core/Public/UbaProcessUtils.h
@@ -16,6 +16,7 @@ namespace uba
 		const CharType* argStart = arguments;
 		bool isInArg = false;
 		bool isInQuotes = false;
+		bool isInSingleQuotes = false;
 		bool isInBackslashQuote = false;
 		bool isEnd = *arguments == 0;
 		CharType currentChar = 0;
@@ -37,7 +38,7 @@ namespace uba
 			isEnd = pastEnd || currentChar == 0;
 			if (isEnd || currentChar == ' ' || currentChar == '\t' || currentChar == '\n')
 			{
-				if (isInQuotes || !isInArg)
+				if (isInQuotes || isInSingleQuotes || !isInArg)
 					continue;
 
 				const CharType* argIt = argStart;
@@ -60,21 +61,38 @@ namespace uba
 				CharType* argWriteIt = argWriteBuffer;
 				CharType lastChar2 = 0;
 				isBackslashOwned = false;
+				bool inSingleQ = false;
+				bool inDoubleQ = false;
 				while (argIt != argEnd)
 				{
-					if (*argIt == '\"')
+					// Single quotes: strip markers; content inside is literal (no " or \ processing)
+					if (*argIt == '\'' && !inDoubleQ)
+					{
+						inSingleQ = !inSingleQ;
+						lastChar2 = 0;
+						++argIt;
+						continue;
+					}
+
+					// Double quotes: strip markers (with \" escape support); only outside single-quoted regions
+					if (*argIt == '\"' && !inSingleQ)
 					{
 						if (lastChar2 == '\\' && !isBackslashOwned)
 							argWriteIt[-1] = '\"';
 						lastChar2 = 0;
+						inDoubleQ = !inDoubleQ;
 						++argIt;
 						continue;
 					}
 
-					if (*argIt == '\\' && lastChar2 == '\\')
-						isBackslashOwned = !isBackslashOwned;
-					else
-						isBackslashOwned = false;
+					// Backslash pair tracking only outside single-quoted regions
+					if (!inSingleQ)
+					{
+						if (*argIt == '\\' && lastChar2 == '\\')
+							isBackslashOwned = !isBackslashOwned;
+						else
+							isBackslashOwned = false;
+					}
 
 					*argWriteIt++ = *argIt;
 
@@ -98,10 +116,17 @@ namespace uba
 				argStart = it;
 				if (currentChar == '\"')
 					isInQuotes = true;
+				else if (currentChar == '\'')
+					isInSingleQuotes = true;
 				continue;
 			}
 
-			if (*it == '\"')
+			// Toggle single-quote mode (only outside double-quoted regions)
+			if (*it == '\'' && !isInQuotes)
+				isInSingleQuotes = !isInSingleQuotes;
+
+			// Toggle double-quote mode (only outside single-quoted regions)
+			if (*it == '\"' && !isInSingleQuotes)
 			{
 				if (isInQuotes && lastChar == '\\' && (!isBackslashOwned && !isInBackslashQuote))
 					continue;
diff --git a/Engine/Source/Programs/UnrealBuildAccelerator/Core/Public/UbaProtocol.h b/Engine/Source/Programs/UnrealBuildAccelerator/Core/Public/UbaProtocol.h
--- a/Engine/Source/Programs/UnrealBuildAccelerator/Core/Public/UbaProtocol.h
+++ b/Engine/Source/Programs/UnrealBuildAccelerator/Core/Public/UbaProtocol.h
@@ -24,6 +24,8 @@ namespace uba
 		UBA_PROCESS_MESSAGE(GetWrittenFiles) \
 		UBA_PROCESS_MESSAGE(CreateProcess) \
 		UBA_PROCESS_MESSAGE(StartProcess) \
+		UBA_PROCESS_MESSAGE(ForkProcess) \
+		UBA_PROCESS_MESSAGE(ExecProcess) \
 		UBA_PROCESS_MESSAGE(ExitChildProcess) \
 		UBA_PROCESS_MESSAGE(VirtualAllocFailed) \
 		UBA_PROCESS_MESSAGE(Log) \
diff --git a/Engine/Source/Programs/UnrealBuildAccelerator/Detours/Private/PosixOS/UbaDetoursFunctionsPosix.cpp b/Engine/Source/Programs/UnrealBuildAccelerator/Detours/Private/PosixOS/UbaDetoursFunctionsPosix.cpp
--- a/Engine/Source/Programs/UnrealBuildAccelerator/Detours/Private/PosixOS/UbaDetoursFunctionsPosix.cpp
+++ b/Engine/Source/Programs/UnrealBuildAccelerator/Detours/Private/PosixOS/UbaDetoursFunctionsPosix.cpp
@@ -6,6 +6,7 @@
 #include "UbaBinaryReaderWriter.h"
 #include "UbaDefinitions.h"
 #include "UbaDirectoryTable.h"
+#include "UbaEvent.h"
 #include "UbaFileMapping.h"
 #include "UbaPlatform.h"
 #include "UbaProcessStats.h"
@@ -26,8 +27,10 @@
 
 #if PLATFORM_LINUX
 #include <sys/prctl.h>
+//long clone3(struct clone_args *cl_args, size_t size);
 #else
 #include <crt_externs.h>
+extern char** environ;
 #endif
 
 #define UBA_DETOUR_DEBUG 0//UBA_DEBUG
@@ -48,6 +51,12 @@ namespace uba
 
 	SharedMemoryAllocatorHandle g_sharedMemoryAllocator;
 	u8* g_sharedMemory;
+
+	extern int g_comFd;
+	extern SharedEvent* g_cancelEvent;
+	extern SharedEvent* g_readEvent;
+	extern SharedEvent* g_writeEvent;
+	extern u8* g_messageMappingMem;
 }
 
 using namespace uba;
@@ -76,6 +85,7 @@ using namespace uba;
 	DETOURED_FUNCTION(close) \
 	DETOURED_FUNCTION(fopen) \
 	DETOURED_FUNCTION(fdopen) \
+	DETOURED_FUNCTION(freopen) \
 	DETOURED_FUNCTION(fchmod) \
 	DETOURED_FUNCTION(fchmodat) \
 	DETOURED_FUNCTION(fstat) \
@@ -95,6 +105,7 @@ using namespace uba;
 	DETOURED_FUNCTION(stat) \
 	DETOURED_FUNCTION(truncate) \
 	DETOURED_FUNCTION(lstat) \
+	DETOURED_FUNCTION(lseek) \
 	DETOURED_FUNCTION(glob) \
 	DETOURED_FUNCTION(chmod) \
 	DETOURED_FUNCTION(rename) \
@@ -134,6 +145,9 @@ using namespace uba;
 	DETOURED_FUNCTION_LINUX \
 	DETOURED_FUNCTION_MACOS \
 
+	//DETOURED_FUNCTION(clone)
+	//DETOURED_FUNCTION(clone3)
+	
 // Used by cp and might need detour
 // fadvise64
 // copy_file_range
@@ -141,6 +155,7 @@ using namespace uba;
 #if UBA_DETOUR_DEBUG && PLATFORM_LINUX
 #define DETOURED_FUNCTION_DEBUG \
 	DETOURED_FUNCTION(write) \
+	DETOURED_FUNCTION(fwrite) \
 
 #else
 #define DETOURED_FUNCTION_DEBUG
@@ -452,7 +467,8 @@ int Shared_open(const char* funcName, const char* file, int flags, int mode, con
 		if (isWrite)
 		{
 			//UBA_ASSERT(!info.isFileMap);
-			bool shouldReport = !(info.lastDesiredAccess & AccessFlag_Write) || info.deleted;
+			bool isTruncating = (flags & O_TRUNC) != 0;
+			bool shouldReport = !(info.lastDesiredAccess & AccessFlag_Write) || info.deleted || isTruncating;
 			shouldReport = shouldReport && !keepInMemory;
 			if (shouldReport)
 			{
@@ -568,7 +584,7 @@ FILE* Shared_fopen(const char* funcName, const char* path, const char* mode, con
 	if (r)
 	{
 		UBA_ASSERTF(!p, "%s with + not implemented (%s)", funcName, mode);
-		flags = O_NONBLOCK | O_RDONLY;
+		flags = O_RDONLY | O_NONBLOCK;
 	}
 	else if (w)
 	{
@@ -578,7 +594,7 @@ FILE* Shared_fopen(const char* funcName, const char* path, const char* mode, con
 		else
 			flags |= O_WRONLY;
 	}
-	int openMode = S_IRUSR | S_IWUSR;
+	int openMode = S_IRUSR | S_IWUSR | S_IRGRP | S_IROTH;
 	int fd = Shared_open(trueOpenName, path, flags, openMode, trueOpen);
 	if (fd == -1)
 	{
@@ -636,11 +652,26 @@ void Shared_close(int fd, const TrueClose& trueClose)
 
 	if (fo->closeId)
 	{
-		fi.size = lseek(fd, 0, SEEK_END); // is this safe? not sure about fflush and user space FILE* streams
 		fi.created = true; // TODO: Should be set earlier
+		// Dup the fd before closing so we can measure the true file size AFTER fclose
+		// flushes any stdio buffers. Measuring via lseek before fclose misses any bytes
+		// still sitting in the stdio buffer (e.g. when 'as' uses fdopen on top of open64).
+		int dupFd = TRUE_WRAPPER(dup)(fd);
 		closeGuard.Execute();
+		if (dupFd != -1)
+		{
+			fi.size = TRUE_WRAPPER(lseek)(dupFd, 0, SEEK_END);
+			TRUE_WRAPPER(close)(dupFd);
+		}
 		//for (auto& kv : g_fileHandles)
 		//	UBA_ASSERT(&fi != kv.second.fileObject->fileInfo);
+		if (fi.size == (off_t)-1)
+		{
+			// Non-seekable fd (pipe etc.) — should not happen for tracked output files
+			UBA_ASSERT(false);
+			delete fo;
+			return;
+		}
 		Rpc_UpdateCloseHandle(path, fo->closeId, fo->deleteOnClose, fo->newName.c_str(), memoryHandle, memoryWritten, true);
 	}
 	delete fo;
@@ -930,17 +961,21 @@ UBA_EXPORT int UBA_WRAPPER(chdir)(const char* path)
 	if (path == nullptr || *path == '\0')
 	{
 		errno = ENOENT;
+		DEBUG_LOG_TRUE("chdir", "(BAD_PATH) -> -1");
 		return -1;
 	}
 	size_t pathlen = strlen(path);
 	if (pathlen >= g_virtualWorkingDir.capacity)
 	{
 		errno = ENAMETOOLONG;
+		DEBUG_LOG_TRUE("chdir", "(TOO_LONG) %s -> -1", path);
 		return -1;
 	}
 	memcpy(g_virtualWorkingDir.data, path, pathlen + 1);
+	g_virtualWorkingDir.count = u32(pathlen);
 	setenv("PWD", g_virtualWorkingDir.data, 1);
 	g_virtualWorkingDir.EnsureEndsWithSlash();
+	DEBUG_LOG_TRUE("chdir", "%s -> 0", path);
 	return 0;
 }
 
@@ -1033,23 +1068,28 @@ UBA_EXPORT char* UBA_WRAPPER(getcwd)(char* buf, size_t size)
 {
 	UBA_INIT_DETOUR(getcwd, buf, size);
 
+	// g_virtualWorkingDir always ends with '/' (required by FixPath2), but
+	// POSIX getcwd() must NOT return a trailing slash except for root "/".
+	u32 len = g_virtualWorkingDir.count;
+	if (len > 1 && g_virtualWorkingDir.data[len - 1] == '/')
+		--len;
+
 	if (size == 0)
 	{
 		DEBUG_LOG_DETOURED("getcwd", "-> null (Size 0)");
 
 		if (buf == nullptr)
 		{
-			UBA_ASSERT(size == 0);
-			size = g_virtualWorkingDir.count+1;
-			buf = (char*)malloc(size);
-			memcpy(buf, g_virtualWorkingDir.data, size);
+			buf = (char*)malloc(len + 1);
+			memcpy(buf, g_virtualWorkingDir.data, len);
+			buf[len] = 0;
 			return buf;
 		}
 
 		errno = EINVAL;
 		return nullptr;
 	}
-	if (size < g_virtualWorkingDir.count+1)
+	if (size < len + 1)
 	{
 		DEBUG_LOG_DETOURED("getcwd", "-> null (Buffer too small: %u)", u32(size));
 		errno = ERANGE;
@@ -1057,9 +1097,9 @@ UBA_EXPORT char* UBA_WRAPPER(getcwd)(char* buf, size_t size)
 	}
 	if (!buf)
 		buf = (char*)malloc(size);
-		
-	UBA_ASSERTF(g_virtualWorkingDir.count < size, "getcwd with size smaller than path not implemented");
-	memcpy(buf, g_virtualWorkingDir.data, g_virtualWorkingDir.count+1);
+
+	memcpy(buf, g_virtualWorkingDir.data, len);
+	buf[len] = 0;
 	DEBUG_LOG_DETOURED("getcwd", "%s -> %p", buf, buf);
 	return buf;
 }
@@ -1280,6 +1320,14 @@ UBA_EXPORT ssize_t UBA_WRAPPER(write)(int fd, const void* buf, size_t count)
 	TimerScope ts(g_kernelStats.writeFile);
 	return TRUE_WRAPPER(write)(fd, buf, count);
 }
+
+UBA_EXPORT size_t UBA_WRAPPER(fwrite)(const void* ptr, size_t size, size_t nitems, FILE* stream)
+{
+	UBA_INIT_DETOUR(fwrite, ptr, size, nitems, stream);
+	DEBUG_LOG_TRUE("fwrite", "(size: %llu)", size*nitems);
+	TimerScope ts(g_kernelStats.writeFile);
+	return TRUE_WRAPPER(fwrite)(ptr, size, nitems, stream);
+}
 #endif
 
 //UBA_EXPORT int fcntl(int __fd, int __cmd, ...)
@@ -1354,6 +1402,14 @@ UBA_EXPORT FILE* UBA_WRAPPER(fdopen)(int fd, const char* mode)
 	return TRUE_WRAPPER(fdopen)(fd, mode);
 }
 
+UBA_EXPORT FILE* UBA_WRAPPER(freopen)(const char* filename, const char* mode, FILE* stream)
+{
+	UBA_INIT_DETOUR(freopen, filename, mode, stream);
+	DEBUG_LOG_TRUE("freopen", "%s", filename);
+	TimerScope ts(g_kernelStats.createFile);
+	return TRUE_WRAPPER(freopen)(filename, mode, stream);
+}
+
 UBA_EXPORT int UBA_WRAPPER(fchmod)(int fd, mode_t mode)
 {
 	UBA_INIT_DETOUR(fchmod, fd, mode);
@@ -1415,6 +1471,13 @@ UBA_EXPORT int UBA_WRAPPER(futimens)(int fd, const struct timespec* times)
 UBA_EXPORT int UBA_WRAPPER(fclose)(FILE* stream)
 {
 	UBA_INIT_DETOUR(fclose, stream);
+
+	#define FCLOSE_STD(x) if (stream == x) { DEBUG_LOG_TRUE("fclose", "(" #x ")"); return TRUE_WRAPPER(fclose)(stream); }
+
+	FCLOSE_STD(stdin)
+	FCLOSE_STD(stdout)
+	FCLOSE_STD(stderr)
+
 	int fd = fileno(stream);
 	int res = 0;
 	int error = 0;
@@ -1664,6 +1727,14 @@ UBA_EXPORT int UBA_WRAPPER(lstat)(const char *path, struct stat *buf)
 	return res;
 }
 
+UBA_EXPORT off_t  UBA_WRAPPER(lseek)(int fd, off_t offset, int whence)
+{
+	UBA_INIT_DETOUR(lseek, fd, offset, whence);
+	int res = TRUE_WRAPPER(lseek)(fd, offset, whence);
+	DEBUG_LOG_TRUE("lseek", "%i %i %i -> %i", fd, offset, whence, res);
+	return res;
+}
+
 UBA_EXPORT int UBA_WRAPPER(stat)(const char* file, struct stat* attr)
 {
 	UBA_INIT_DETOUR(stat, file, attr);
@@ -2050,20 +2121,11 @@ UBA_EXPORT int UBA_WRAPPER(unlink)(const char* pathname)
 UBA_EXPORT int UBA_WRAPPER(unlinkat)(int dirfd, const char *pathname, int flags)
 {
 	UBA_INIT_DETOUR(unlinkat, dirfd, pathname, flags);
-
 	UBA_ASSERT(!g_runningRemote);
 	UBA_ASSERT(dirfd == AT_FDCWD);
-
-	StringBuffer<> fixedPath;
-	FixPath(fixedPath, pathname);
-	DevirtualizePath(fixedPath);
-	int res = TRUE_WRAPPER(unlinkat)(dirfd, fixedPath.data, flags);
-	DEBUG_LOG_TRUE("unlinkat", "%s -> %i (%s)", pathname, res, StrError(res, errno));
-	return res;
+	return Shared_DeleteFile("unlinkat", pathname);
 }
 
-thread_local int t_inVfork;
-
 void FlattenArgs(StringBufferBase& out, const char* const argv[])
 {
 	if (!argv)
@@ -2134,12 +2196,185 @@ void UnsupportedHostRun(char* const argv[], const char* msg)
 	UbaAssert(msg, __FILE__, __LINE__, command.data, true, 1999, nullptr, 0);
 }
 
+struct ProcessCreator
+{
+	int SendMessage(MessageType messageType, const char* messageName, const char* path, char* const argv[], char* const envp[], bool fixupPath, bool isExec)
+	{
+		fixedArgv = argv;
+
+		if (fixupPath)
+		{
+			if (!path || !*path)
+				FixPath(result, argv[0]);
+			else
+				FixPath(result, path);
+
+			DevirtualizePath(result);
+			path = result.data;
+
+			// If the executable is a shell script, wrap it with /bin/sh so the kernel
+			// does not need to parse the shebang — and so UBA's virtual FS interception
+			// does not prevent the script from being opened by the kernel's exec path.
+			u32 plen = result.count;
+			if (plen >= 3 && result.data[plen-3] == '.' && result.data[plen-2] == 's' && result.data[plen-1] == 'h')
+			{
+				u32 argc2 = 0;
+				tempArgv[argc2++] = "/bin/sh";
+				tempArgv[argc2++] = path; // the .sh script (still in result.data)
+				for (int i = 1; argv[i]; ++i)
+					tempArgv[argc2++] = argv[i];
+				tempArgv[argc2] = nullptr;
+				fixedArgv = (char* const*)tempArgv;
+				path = "/bin/sh";
+			}
+		}
+		else if (path && strchr(path, '/'))
+		{
+			// execvp-style (no fixup for bare filenames), but if the path contains a slash
+			// it may still be a virtualized absolute/relative path — devirtualize it.
+			result.Append(path);
+			DevirtualizePath(result);
+			path = result.data;
+		}
+
+		for (u32 i = 1; fixedArgv[i]; ++i)
+		{
+			if (i != 1)
+				cmdLineWithoutApplication.append(" ");
+			const char* arg = fixedArgv[i];
+			// Quote arguments that contain spaces or double-quotes so ParseArguments
+			// on the receiving end reconstructs the correct argv (e.g. sed -e scripts).
+			bool needsQuoting = false;
+			for (const char* p = arg; *p; ++p)
+				if (*p == ' ' || *p == '\t' || *p == '"') { needsQuoting = true; break; }
+			if (needsQuoting)
+			{
+				cmdLineWithoutApplication.append("\"");
+				for (const char* p = arg; *p; ++p)
+				{
+					if (*p == '"') cmdLineWithoutApplication.append("\\\"");
+					else           cmdLineWithoutApplication += *p;
+				}
+				cmdLineWithoutApplication.append("\"");
+			}
+			else
+				cmdLineWithoutApplication.append(arg);
+		}
+
+		{
+			RPC_MESSAGE2(messageType, messageName, createProcess);
+
+			const char* pwd = "";
+#if 0
+			for (char* const* env = envp; *env;++env)
+			{
+				if (StartsWith(*env, "PWD="))
+				{
+					printf("%s\n", *env);
+					pwd = *env + 4;
+					break;
+				}
+			}
+#endif
+
+			writer.WriteString(path); // Application
+			writer.WriteLongString(cmdLineWithoutApplication, CommunicationMemSize - 4096);
+			writer.WriteString(pwd); // Current dir
+			writer.WriteBool(false); // Start suspended
+			writer.WriteBool(true); // Is child
+			BinaryReader reader = writer.Flush();
+
+			processId = reader.ReadU32();
+
+			if (!processId) // Can happen if session client got disconnected from server session
+			{
+				errno = EINVAL; // This is not really correct but there is no errno for this failure
+				return -1;
+			}
+
+			rulesStr.Append("UBA_RULES=").AppendValue(reader.ReadU32());
+
+			while (true)
+			{
+				u32 dllNameSize = reader.ReadU32();
+				if (!dllNameSize)
+					break;
+				reader.Skip(dllNameSize);
+			}
+
+			currentDir.Append("UBA_CWD=");
+			reader.ReadString(currentDir);
+
+			realApplication = reader.ReadString();
+
+			#if UBA_USE_NATIVE_MAC_SEMAPHORES
+			pidVar.Append("UBA_PID=").AppendValue(reader.ReadU32());
+			#endif
+
+			comIdVar.Append("UBA_COMID=").AppendValue(reader.ReadU64()).Append('+').AppendValue(reader.ReadU32());
+
+			logFile.Append("UBA_LOGFILE=");
+			reader.ReadString(logFile);
+		}
+
+		for (u32 i = 0; envp[i]; ++i)
+			envvars.push_back(envp[i]);
+		envvars.push_back(comIdVar.data);
+		envvars.push_back(currentDir.data);
+		envvars.push_back(rulesStr.data);
+		envvars.push_back(logFile.data);
+		if (isExec)
+			envvars.push_back("UBA_EXEC=1");
+		//envvars.push_back("LD_DEBUG=bindings");
+
+		#if UBA_USE_NATIVE_MAC_SEMAPHORES
+		envvars.push_back(pidVar.data);
+		#endif
+
+		envvars.push_back(nullptr);
+
+		//	int i=0;
+		//	printf("spawnng %s\n", envvars.data()[i]);
+		//	while (envvars.data()[i])
+		//	{
+		//		printf("env: %s\n", envvars.data()[i]);
+		//		i++;
+		//	}
+		return 0;
+	}
+
+	TString realApplication;
+	TString cmdLineWithoutApplication;
+	u32 processId = 0;
+	StringBuffer<512> currentDir;
+	StringBuffer<256> comIdVar;
+	StringBuffer<32> rulesStr;
+	StringBuffer<512> logFile;
+	#if UBA_USE_NATIVE_MAC_SEMAPHORES
+	StringBuffer<64> pidVar;
+	#endif
+	StringBuffer<> result;
+	const char* tempArgv[1024];
+	char* const* fixedArgv;
+	std::vector<const char*> envvars;
+};
+
+void SendStartProcess(u32 processId, u32 pid, bool success, u32 res)
+{
+	RPC_MESSAGE(StartProcess, createProcess);
+	writer.WriteU32(processId);
+	writer.WriteBool(success);
+	writer.WriteU32(res);
+	writer.WriteU64(1); // Process handle
+	writer.WriteU32(pid);
+	writer.WriteU64(0); // Thread handle
+	writer.Flush();
+}
+
 int shared_posix_spawn(pid_t* pid, const char* path, const posix_spawn_file_actions_t* file_actions, const posix_spawnattr_t* attrp, char* const argv[], char* const envp[])
 {
 	UBA_INIT_DETOUR(posix_spawn, pid, path, file_actions, attrp, argv, envp);
 
-	t_inVfork = 0;
-
 	{
 		//StringBuffer<4096> command;
 		//FlattenArgs(command, argv);
@@ -2150,6 +2385,8 @@ int shared_posix_spawn(pid_t* pid, const char* path, const posix_spawn_file_acti
 	StringBuffer<> result;
 	StringBuffer<> additionalArg;
 
+	bool fixupPath = true;
+
 	if (strstr(path, "xcode-select"))
 	{
 		if (strcmp(argv[1], "--print-path") != 0)
@@ -2158,7 +2395,8 @@ int shared_posix_spawn(pid_t* pid, const char* path, const posix_spawn_file_acti
 			return -1;
 		return SpawnEcho(result.data, pid,  file_actions, attrp, envp);
 	}
-	else if(strstr(path, "xcrun"))
+	
+	if(strstr(path, "xcrun"))
 	{
 		if (strcmp(argv[1], "--sdk") != 0)
 			UnsupportedHostRun(argv, "xcrun unsupported first param");
@@ -2205,175 +2443,27 @@ int shared_posix_spawn(pid_t* pid, const char* path, const posix_spawn_file_acti
 		}
 		tempArgv[argc3] = 0;
 		argv = (char*const*)tempArgv;
-	}
-	else
-	{
-		if (!path || !*path)
-			FixPath(result, argv[0]);
-		else
-			FixPath(result, path);
-
-		DevirtualizePath(result);
-		path = result.data;
-
-		// If the executable is a shell script, wrap it with /bin/sh so the kernel
-		// does not need to parse the shebang — and so UBA's virtual FS interception
-		// does not prevent the script from being opened by the kernel's exec path.
-		u32 plen = result.count;
-		if (plen >= 3 && result.data[plen-3] == '.' && result.data[plen-2] == 's' && result.data[plen-1] == 'h')
-		{
-			u32 argc2 = 0;
-			tempArgv[argc2++] = "/bin/sh";
-			tempArgv[argc2++] = path; // the .sh script (still in result.data)
-			for (int i = 1; argv[i]; ++i)
-				tempArgv[argc2++] = argv[i];
-			tempArgv[argc2] = nullptr;
-			argv = (char* const*)tempArgv;
-			path = "/bin/sh";
-		}
+		fixupPath = false;
 	}
 	//DEBUG_LOG("LIBRARY_SEARCH_PATHS: %s", getenv("LIBRARY_SEARCH_PATHS"));
 	//DEBUG_LOG("PATH: %s", getenv("PATH"));
 
-	TString cmdLineWithoutApplication;
-	for (u32 i = 1; argv[i]; ++i)
-	{
-		cmdLineWithoutApplication.append(" ");
-		const char* arg = argv[i];
-		// Quote arguments that contain spaces or double-quotes so ParseArguments
-		// on the receiving end reconstructs the correct argv (e.g. sed -e scripts).
-		bool needsQuoting = false;
-		for (const char* p = arg; *p; ++p)
-			if (*p == ' ' || *p == '\t' || *p == '"') { needsQuoting = true; break; }
-		if (needsQuoting)
-		{
-			cmdLineWithoutApplication.append("\"");
-			for (const char* p = arg; *p; ++p)
-			{
-				if (*p == '"') cmdLineWithoutApplication.append("\\\"");
-				else           cmdLineWithoutApplication += *p;
-			}
-			cmdLineWithoutApplication.append("\"");
-		}
-		else
-			cmdLineWithoutApplication.append(arg);
-	}
-
-	TString realApplication;
-	u32 processId = 0;
-	StringBuffer<512> currentDir;
-	StringBuffer<256> comIdVar;
-	StringBuffer<32> rulesStr;
-	StringBuffer<512> logFile;
-
-	#if UBA_USE_NATIVE_MAC_SEMAPHORES
-	StringBuffer<64> pidVar;
-	#endif
-	
-	{
-		RPC_MESSAGE(CreateProcess, createProcess);
-
-		const char* pwd = "";
-		//while (char* env = envp[i++])
-		//{
-		//	if (StartsWith(env, "PWD="))
-		//	{
-		//		printf("%s\n", env);
-		//		pwd = env + 4;
-		//		break;
-		//	}
-		//}
-
-		writer.WriteString(path); // Application
-		writer.WriteLongString(cmdLineWithoutApplication, CommunicationMemSize - 4096);
-		writer.WriteString(pwd); // Current dir
-		writer.WriteBool(false); // Start suspended
-		writer.WriteBool(true); // Is child
-		BinaryReader reader = writer.Flush();
-
-		processId = reader.ReadU32();
-
-		if (!processId) // Can happen if session client got disconnected from server session
-		{
-			errno = EINVAL; // This is not really correct but there is no errno for this failure
-			return -1;
-		}
-		
-		rulesStr.Append("UBA_RULES=").AppendValue(reader.ReadU32());
-
-		while (true)
-		{
-			u32 dllNameSize = reader.ReadU32();
-			if (!dllNameSize)
-				break;
-			reader.Skip(dllNameSize);
-		}
-
-		currentDir.Append("UBA_CWD=");
-		reader.ReadString(currentDir);
-
-		realApplication = reader.ReadString();
-
-		#if UBA_USE_NATIVE_MAC_SEMAPHORES
-		pidVar.Append("UBA_PID=").AppendValue(reader.ReadU32());
-		#endif
-
-		comIdVar.Append("UBA_COMID=").AppendValue(reader.ReadU64()).Append('+').AppendValue(reader.ReadU32());
-
-		logFile.Append("UBA_LOGFILE=");
-		reader.ReadString(logFile);
-	}
-	
-	std::vector<const char*> envvars;
-	{
-		for (u32 i = 0; envp[i]; ++i)
-		{
-//			if (envp[i][0] == 'D' || envp[i][0] == 'P')
-				envvars.push_back(envp[i]);
-		}
-		envvars.push_back(comIdVar.data);
-		envvars.push_back(currentDir.data);
-		envvars.push_back(rulesStr.data);
-		envvars.push_back(logFile.data);
-		//envvars.push_back("LD_DEBUG=bindings");
-
-		#if UBA_USE_NATIVE_MAC_SEMAPHORES
-		envvars.push_back(pidVar.data);
-		#endif
-		
-		envvars.push_back(nullptr);
-	}
-
-//	int i=0;
-//	printf("spawnng %s\n", envvars.data()[i]);
-//	while (envvars.data()[i])
-//	{
-//		printf("env: %s\n", envvars.data()[i]);
-//		i++;
-//	}
+	ProcessCreator pc;
+	int res = pc.SendMessage(MessageType_CreateProcess, "CreateProcess", path, argv, envp, fixupPath, false);
+	if (res != 0)
+		return res;
 	
 	#if UBA_DEBUG_LOG_ENABLED
-	DEBUG_LOG_TRUE("posix_spawn", "%s (%s)", realApplication.data(), logFile.data);
-	for (u32 i = 0; argv[i]; ++i)
-		DEBUG_LOG("            %s", argv[i]);
+	DEBUG_LOG_TRUE("posix_spawn", "%s (%s)", pc.realApplication.data(), pc.logFile.data);
+	for (u32 i = 0; pc.fixedArgv[i]; ++i)
+		DEBUG_LOG("            %s", pc.fixedArgv[i]);
 	#endif
 
-	int res = TRUE_WRAPPER(posix_spawn)(pid, realApplication.data(), file_actions, attrp, argv, (char**)envvars.data());
+	res = TRUE_WRAPPER(posix_spawn)(pid, pc.realApplication.data(), file_actions, attrp, pc.fixedArgv, (char**)pc.envvars.data());
 	bool success = res == 0;
 
-	{
-		RPC_MESSAGE(StartProcess, createProcess);
-		writer.WriteU32(processId);
-		writer.WriteBool(success);
-		writer.WriteU32(res);
-		writer.WriteU64(1); // Process handle
-		writer.WriteU32(*pid);
-		writer.WriteU64(0); // Thread handle
-		writer.Flush();
-	}
-
+	SendStartProcess(pc.processId, *pid, success, res);
 	DEBUG_LOG("         Child process started %s -> %i (pid: %i)", path, res, *pid);
-
 	return res;
 }
 
@@ -2433,8 +2523,6 @@ const char* GetResult(siginfo_t* info)
 UBA_EXPORT int UBA_WRAPPER(waitid)(idtype_t idtype, id_t id, siginfo_t* infop, int options)
 {
 	UBA_INIT_DETOUR(waitid, idtype, id, infop, options);
-	UBA_ASSERTF(!t_inVfork, "waitid: is in fork");
-	// TODO: Should probably report id to session
 	auto res = TRUE_WRAPPER(waitid)(idtype, id, infop, options);
 	DEBUG_LOG_TRUE("waitid", "%i -> %i (%s)", id, res, GetResult(infop));
 	if (infop->si_code == CLD_EXITED)
@@ -2445,7 +2533,6 @@ UBA_EXPORT int UBA_WRAPPER(waitid)(idtype_t idtype, id_t id, siginfo_t* infop, i
 UBA_EXPORT pid_t UBA_WRAPPER(wait3)(int* status, int options, struct rusage* rusage)
 {
 	UBA_INIT_DETOUR(wait3, status, options, rusage);
-	UBA_ASSERTF(!t_inVfork, "wait3: is in fork");
 	pid_t res =TRUE_WRAPPER(wait3)(status, options, rusage);
 	DEBUG_LOG_TRUE("wait3", "-> %i (%i)", res, *status);
 	if (WIFEXITED(*status))
@@ -2456,7 +2543,6 @@ UBA_EXPORT pid_t UBA_WRAPPER(wait3)(int* status, int options, struct rusage* rus
 UBA_EXPORT pid_t UBA_WRAPPER(wait4)(pid_t pid, int* status, int options, struct rusage* rusage)
 {
 	UBA_INIT_DETOUR(wait4, pid, status, options, rusage);
-	//UBA_ASSERTF(!t_inVfork, "wait4: is in fork");
 	pid_t res = TRUE_WRAPPER(wait4)(pid, status, options, rusage);
 	if (WIFEXITED(*status))
 		Rpc_UpdateTables();
@@ -2535,93 +2621,67 @@ UBA_EXPORT int UBA_WRAPPER(dladdr)(const void *addr, Dl_info *info)
 	return res;
 }
 
-int Internal_execve(const char* pathname, char* const _Nullable argv[], char* const _Nullable envp[])
+int Internal_exec(const char* func, const char* pathname, char* const _Nullable argv[], char* const _Nullable envp[], bool fixupPath)
 {
 	UBA_INIT_DETOUR(execve, pathname, argv, envp);
-	// We are most likely in a vfork/fork here, which means that we won't see the exit call since _exit will drop back out to the entrance of fork.
-	DEBUG_LOG_TRUE("execve", "%s", pathname);
-
-	if (StartsWith(pathname, "/usr/bin/stat"))
-	{
-		UBA_ASSERTF(!g_runningRemote, "Tried using posix_spawn with ExecuteHostRun+SpawnEcho but didn't work.");
-		return TRUE_WRAPPER(execve)(pathname, argv, envp);
-	}
-
-	pid_t pid = getpid();
-	bool inVfork = g_pid != pid;
-	g_pid = pid;
-
-	int res = shared_posix_spawn(&pid, pathname, nullptr, nullptr, argv, envp);
-	if (inVfork)
-		t_inVfork = pid;
 
+	ProcessCreator pc;
+	int res = pc.SendMessage(MessageType_ExecProcess, "ExecProcess", pathname, argv, envp, fixupPath, true);
 	if (res != 0)
-	{
-		UBA_ASSERTF(false, "Failed to spawn %s", pathname);
-		return -1;
-	}
-
-	int status;
-	{
-		res = TRUE_WRAPPER(waitpid)(pid, &status, WUNTRACED | WCONTINUED);
-		DEBUG_LOG_TRUE("waitpid", "(execve) (%i) -> %i (%i)", pid, res, status);
-		UBA_ASSERTF(res == pid, "execve: wait result was not same as pid");
-		UBA_ASSERTF(WIFEXITED(status), "execve: Unsupported status from waitpid");
-	}
+		return res;
 
-	{
-		RPC_MESSAGE(ExitChildProcess, createProcess);
-		writer.WriteU32(pid);
-	}
+	DEBUG_LOG_TRUE(func, "%s %s", pc.realApplication.data(), pc.cmdLineWithoutApplication.data());
 
-	res = WEXITSTATUS(status);
+	// TODO: We should clean up stuff here
+	#if UBA_DEBUG_LOG_ENABLED
+	FlushDebugLog();
+	int debugFile = g_debugFile;
+	g_debugFile = InvalidFileHandle;
+	close(debugFile);
+	#endif
 
-	if (inVfork)
-	{
-		TRUE_WRAPPER(_exit)(res);
-	}
-	else
-	{
-		TRUE_WRAPPER(exit)(res);
-	}
-	//return -1;
+	return TRUE_WRAPPER(execve)(pc.realApplication.data(), pc.fixedArgv, (char**)pc.envvars.data());
 }
 
 UBA_EXPORT int UBA_WRAPPER(execv)(const char* path, char* const argv[])
 {
-	return Internal_execve(path, argv, environ);
+	return Internal_exec("execv", path, argv, environ, true);
 }
 
 UBA_EXPORT int UBA_WRAPPER(execve)(const char* pathname, char* const _Nullable argv[], char* const _Nullable envp[])
 {
-	return Internal_execve(pathname, argv, envp);
+	return Internal_exec("execve", pathname, argv, envp, true);
 }
 
 UBA_EXPORT int UBA_WRAPPER(execvp)(const char* file, char* const argv[])
 {
-	UBA_INIT_DETOUR(execvp, file, argv);
-	DEBUG_LOG_TRUE("execvp", "");
-	return TRUE_WRAPPER(execvp)(file , argv);
+	// This devirtualize should not be needed but fixes Mac ld for some weird reason. This should be fixed up deeper in to the callstack
+	StringBuffer<> temp(file);
+	DevirtualizePath(temp);
+	file = temp.data;
+	// Remove above after debugged
+
+	return Internal_exec("execvp", file, argv, environ, false);
 }
 
 UBA_EXPORT int UBA_WRAPPER(execl)(const char *path, const char *arg0, ...)
 {
 	DEBUG_LOG_TRUE("execl", "");
-	UBA_ASSERT(false);
+	UBA_ASSERTF(false, "execl NotImplemented");
 	return -1;
 }
 
 UBA_EXPORT int UBA_WRAPPER(execle)(const char *path, const char *arg0, ...)
 {
 	DEBUG_LOG_TRUE("execle", "");
-	UBA_ASSERT(false);
+	UBA_ASSERTF(false, "execle NotImplemented");
 	return -1;
 }
 
 UBA_EXPORT int UBA_WRAPPER(execlp)(const char *file, const char *arg0, ...)
 {
 	DEBUG_LOG_TRUE("execlp", "");
-	UBA_ASSERT(false);
+	UBA_ASSERTF(false, "execlp NotImplemented");
 	return -1;
 }
 
@@ -2629,7 +2689,7 @@ UBA_EXPORT int UBA_WRAPPER(execlp)(const char *file, const char *arg0, ...)
 UBA_EXPORT int UBA_WRAPPER(execvP)(const char *file, const char *search_path, char *const argv[])
 {
 	DEBUG_LOG_TRUE("execvP", "");
-	UBA_ASSERT(false);
+	UBA_ASSERTF(false, "execvP NotImplemented");
 	return -1;
 }
 
@@ -2641,6 +2701,80 @@ UBA_EXPORT pid_t UBA_WRAPPER(fork)(void)
 }
 #endif
 
+pid_t Internal_fork()
+{
+	UBA_INIT_DETOUR(fork);
+
+	u32 childProcessId = ~0u;
+	StringBuffer<256> comId;
+	StringBuffer<512> logFile;
+	{
+		RPC_MESSAGE(ForkProcess, createProcess);
+		BinaryReader reader = writer.Flush();
+		childProcessId = reader.ReadU32();
+		comId.AppendValue(reader.ReadU64()).Append('+').AppendValue(reader.ReadU32());
+		reader.ReadString(logFile);
+	}
+
+	pid_t pid = TRUE_WRAPPER(fork)();
+	if (pid == 0)
+	{
+		#if PLATFORM_LINUX
+		prctl(PR_SET_PDEATHSIG, SIGHUP, 0, 0, 0); // We want the process to die if the parent die
+		#endif
+
+		#if UBA_DEBUG_LOG_ENABLED
+		g_debugFile = (FileHandle)TRUE_WRAPPER(open)(logFile.data, O_CREAT | O_TRUNC | O_WRONLY, S_IRUSR | S_IWUSR);
+		#endif
+
+		DEBUG_LOG(TC("FORKED pid: %i"), getpid());
+		StringBuffer<256> comIdName;
+		const char* plusIndex = comId.First('+');
+		comIdName.Append(comId.data, u64(plusIndex - comId.data));
+		u64 comIdUid;
+		if (!comIdName.Parse(comIdUid))
+			UBA_ASSERT(false);
+
+		u32 comIdOffset = strtoul(plusIndex + 1, nullptr, 10);
+		GetMappingHandleName(comIdName.Clear(), comIdUid);
+
+		g_comFd = shm_open(comIdName.data, O_RDWR, S_IRUSR | S_IWUSR);
+
+		if (g_comFd == -1)
+		{
+			printf("UbaDetours: Failed to open shared mem: %s\n", comIdName.data);
+			UBA_ASSERTF(false, "UbaDetours: Failed to open shared mem: %s", comIdName.data);
+		}
+		u8* rptr = (u8*)mmap(NULL, CommunicationMemSize, PROT_READ | PROT_WRITE, MAP_SHARED, g_comFd, s64(comIdOffset));
+		if (rptr == MAP_FAILED)
+		{
+			printf("UbaDetours: Failed to mmap fd: %u\n", g_comFd);
+			UBA_ASSERT(false);
+		}
+		fcntl(g_comFd, F_SETFD, FD_CLOEXEC); // Don't let grandchild processes inherit this fd
+
+		g_messageMappingMem = rptr;
+		g_cancelEvent = ((SharedEvent*)rptr);
+		g_readEvent = ((SharedEvent*)rptr) + 1;
+		g_writeEvent = ((SharedEvent*)rptr) + 2;
+		g_messageMappingMem += sizeof(SharedEvent) * 3;
+	}
+	else
+	{
+		SendStartProcess(childProcessId, pid, true, 0);
+		DEBUG_LOG("         Process forked (child pid: %i)", pid);
+	}
+
+	return pid;
+}
+
+#if PLATFORM_LINUX
+UBA_EXPORT pid_t UBA_WRAPPER(fork)(void)
+{
+	DEBUG_LOG_TRUE("fork", "");
+	return Internal_fork();
+}
+#endif
 
 #if PLATFORM_MAC
 #pragma clang diagnostic ignored "-Wdeprecated-declarations"
@@ -2648,19 +2782,58 @@ UBA_EXPORT pid_t UBA_WRAPPER(fork)(void)
 
 UBA_EXPORT pid_t UBA_WRAPPER(vfork)(void)
 {
-	UBA_INIT_DETOUR(vfork);
-	DEBUG_LOG_TRUE("vfork", "");
-	pid_t pid = fork();
-	if (pid == 0)
+	DEBUG_LOG_DETOURED("vfork", "");
+	return Internal_fork();
+}
+/*
+UBA_EXPORT int UBA_WRAPPER(clone)(int (*fn)(void *), void *stack, int flags, void *args, ...)
+{
+	va_list ap;
+	va_start(ap, args);
+
+	UBA_INIT_DETOUR(clone, fn, stack, flags, va_arg(ap, void *), va_arg(ap, void *), va_arg(ap, void *));
+
+	pid_t *parent_tid = NULL;
+	void *tls = NULL;
+	pid_t *child_tid = NULL;
+
+	if (flags & CLONE_PARENT_SETTID)
+		parent_tid = va_arg(ap, pid_t *);
+
+	if (flags & CLONE_SETTLS)
+		tls = va_arg(ap, void *);
+
+	if (flags & CLONE_CHILD_CLEARTID)
+		child_tid = va_arg(ap, pid_t *);
+
+	va_end(ap);
+
+	DEBUG_LOG_TRUE("CLONE!!!", "");
+	//fprintf(stderr, "clone intercepted: flags=0x%x\n", flags);
+
+	if ((flags & CLONE_PARENT_SETTID) && (flags & CLONE_SETTLS) && (flags & CLONE_CHILD_CLEARTID))
+		return TRUE_WRAPPER(clone)(fn, stack, flags, args, parent_tid, tls, child_tid);
+	if ((flags & CLONE_PARENT_SETTID) && (flags & CLONE_SETTLS))
+		return TRUE_WRAPPER(clone)(fn, stack, flags, args, parent_tid, tls);
+	if ((flags & CLONE_PARENT_SETTID))
+		return TRUE_WRAPPER(clone)(fn, stack, flags, args, parent_tid);
+	return TRUE_WRAPPER(clone)(fn, stack, flags, args);
+}
+
+UBA_EXPORT long UBA_WRAPPER(clone3)(struct clone_args *cl_args, size_t size)
+{
+	UBA_INIT_DETOUR(clone3, cl_args, size);
+	DEBUG_LOG_TRUE("CLONE3!!!", "");
+	long ret = TRUE_WRAPPER(clone3)(cl_args, size);
+	if (ret == 0)
 	{
-	#if PLATFORM_LINUX
+		#if PLATFORM_LINUX
 		prctl(PR_SET_PDEATHSIG, SIGHUP, 0, 0, 0); // We want the process to die if the parent die
-	#endif
-		t_inVfork = 1;
+		#endif
 	}
-	return pid;
+	return ret;
 }
-
+*/
 #if PLATFORM_LINUX
 UBA_EXPORT long UBA_WRAPPER(syscall)(long number, ...)
 {
@@ -2670,20 +2843,16 @@ UBA_EXPORT long UBA_WRAPPER(syscall)(long number, ...)
 	UBA_INIT_DETOUR(syscall, number, va_arg(args, void *), va_arg(args, void *), va_arg(args, void *), va_arg(args, void *), va_arg(args, void *));
 	DEBUG_LOG_TRUE("syscall", "%ld", number);
 
-	long ret = TRUE_WRAPPER(syscall)(number, va_arg(args, void *), va_arg(args, void *), va_arg(args, void *), va_arg(args, void *), va_arg(args, void *));
-	va_end(args);
+	bool isClone = number == SYS_clone;
+	#ifdef SYS_clone3
+	isClone |= number == SYS_clone3;
+	#endif
 
-	if (number == SYS_clone)
-	{
-		if (ret == 0)
-		{
-		#if PLATFORM_LINUX
-			prctl(PR_SET_PDEATHSIG, SIGHUP, 0, 0, 0); // We want the process to die if the parent die
-		#endif
-			t_inVfork = 1;
-		}
-	}
+	if (isClone)
+		return Internal_fork();
 
+	long ret = TRUE_WRAPPER(syscall)(number, va_arg(args, void *), va_arg(args, void *), va_arg(args, void *), va_arg(args, void *), va_arg(args, void *));
+	va_end(args);
 	return ret;
 }
 #endif
@@ -2767,22 +2936,16 @@ UBA_EXPORT void UBA_WRAPPER(exit)(int status)
 UBA_EXPORT void UBA_WRAPPER(_exit)(int status)
 {
 	DEBUG_LOG_TRUE("_exit", "(%i)", status);
-	if (!t_inVfork)
-	{
-		Deinit();
-		CloseCom();
-	}
+	Deinit();
+	CloseCom();
 	TRUE_WRAPPER(_exit)(status);
 }
 
 UBA_EXPORT void UBA_WRAPPER(_Exit)(int status)
 {
 	DEBUG_LOG_TRUE("_Exit", "(%i)", status);
-	if (!t_inVfork)
-	{
-		Deinit();
-		CloseCom();
-	}
+	Deinit();
+	CloseCom();
 	TRUE_WRAPPER(_Exit)(status);
 }
 
@@ -2854,7 +3017,7 @@ namespace uba
 #endif
 	}
 
-	void PreInit(const char* logFile)
+	void PreInit(const char* logFile, bool fromExec)
 	{
 		g_fileHandlesMem.Create();
 		g_fileHandlesLockMem.Create();
@@ -2878,7 +3041,7 @@ namespace uba
 		if (g_logToScreen)
 			g_debugFile = (FileHandle)open("/dev/tty", O_WRONLY);
 		else if (logFile)
-			g_debugFile = (FileHandle)open(logFile, O_CREAT | O_TRUNC | O_WRONLY, S_IRUSR | S_IWUSR);
+			g_debugFile = (FileHandle)TRUE_WRAPPER(open)(logFile, (fromExec?O_APPEND:O_CREAT|O_TRUNC) | O_WRONLY, S_IRUSR | S_IWUSR);
 		#endif
 
 		StringBuffer<> exePath;
@@ -2981,6 +3144,8 @@ namespace uba
 		UBA_ASSERTF(lastSlash, "Need fullpath for application (%s)", g_virtualApplication.data);
 		g_virtualApplicationDir.Append(g_virtualApplication.data, lastSlash - g_virtualApplication.data + 1);
 
+		if (g_virtualWorkingDir[g_virtualWorkingDir.count-1] == PathSeparator)
+			g_virtualWorkingDir.Resize(g_virtualWorkingDir.count-1);
 		setenv("PWD", g_virtualWorkingDir.data, 1);
 		g_virtualWorkingDir.EnsureEndsWithSlash();
 
diff --git a/Engine/Source/Programs/UnrealBuildAccelerator/Detours/Private/PosixOS/UbaDetoursMainPosix.cpp b/Engine/Source/Programs/UnrealBuildAccelerator/Detours/Private/PosixOS/UbaDetoursMainPosix.cpp
--- a/Engine/Source/Programs/UnrealBuildAccelerator/Detours/Private/PosixOS/UbaDetoursMainPosix.cpp
+++ b/Engine/Source/Programs/UnrealBuildAccelerator/Detours/Private/PosixOS/UbaDetoursMainPosix.cpp
@@ -24,7 +24,7 @@ namespace uba
 	u8* g_messageMappingMem;
 	extern bool g_isCancelled;
 
-	void PreInit(const char* logFile);
+	void PreInit(const char* logFile, bool fromExec);
 	void Init();
 	void Deinit();
 	void AddExceptionHandler();
@@ -43,8 +43,11 @@ static void __attribute__((constructor(102))) PreInitCtor()
 
 	g_runningRemote = getenv("UBA_REMOTE") != nullptr;
 	
+	bool fromExec = getenv("UBA_EXEC") != nullptr;
+	unsetenv("UBA_EXEC");
+
 	const char* logFile = getenv("UBA_LOGFILE");
-	PreInit(logFile);
+	PreInit(logFile, fromExec);
 	unsetenv("UBA_LOGFILE");
 
 	const char* comId = getenv("UBA_COMID");
diff --git a/Engine/Source/Programs/UnrealBuildAccelerator/Detours/Private/UbaDetoursShared.cpp b/Engine/Source/Programs/UnrealBuildAccelerator/Detours/Private/UbaDetoursShared.cpp
--- a/Engine/Source/Programs/UnrealBuildAccelerator/Detours/Private/UbaDetoursShared.cpp
+++ b/Engine/Source/Programs/UnrealBuildAccelerator/Detours/Private/UbaDetoursShared.cpp
@@ -97,6 +97,8 @@ namespace uba
 
 		out.Appendf(TC("[%5llu.%03llu]"), s, ms);
 		#endif
+		
+		//out.Appendf(TC("[%i]"), getpid());
 		//out.Appendf(TC("[%7u]"), GetCurrentThreadId());
 	}
 	void FlushDebug()
@@ -205,7 +207,7 @@ namespace uba
 		while (left)
 		{
 			RPC_MESSAGE(InputDependencies, log)
-				writer.Write7BitEncoded(reserveSize);
+			writer.Write7BitEncoded(reserveSize);
 			reserveSize = 0;
 			u32 toWrite = Min(left, u32(writer.GetCapacityLeft() - sizeof(u32)));
 			writer.WriteU32(toWrite);
diff --git a/Engine/Source/Programs/UnrealBuildAccelerator/Detours/Private/UbaDetoursShared.h b/Engine/Source/Programs/UnrealBuildAccelerator/Detours/Private/UbaDetoursShared.h
--- a/Engine/Source/Programs/UnrealBuildAccelerator/Detours/Private/UbaDetoursShared.h
+++ b/Engine/Source/Programs/UnrealBuildAccelerator/Detours/Private/UbaDetoursShared.h
@@ -144,11 +144,14 @@ namespace uba
 	#define VARIABLE_MEM(type, name) VariableMem<type> name##Mem; type& name = (type&)name##Mem.data;
 
 	#define RPC_MESSAGE(messageName, timerName) \
-		DEBUG_LOG(TC("RPC_MESSAGE %s"), TC(#messageName)); \
+		RPC_MESSAGE2(MessageType_##messageName, #messageName, timerName)
+
+	#define RPC_MESSAGE2(messageType, messageName, timerName) \
+		DEBUG_LOG(TC("RPC_MESSAGE %s"), TC(messageName)); \
 		TimerScope ts(g_stats.timerName); \
 		SCOPED_FUTEX(g_communicationLock, pcs); \
 		BinaryWriter writer(ProcessCommunication_Write); \
-		writer.WriteByte(MessageType_##messageName);
+		writer.WriteByte(messageType);
 
 	#define RPC_MESSAGE_NO_LOCK(messageName, timerName) \
 		DEBUG_LOG(TC("RPC_MESSAGE %s"), TC(#messageName)); \
