[UBA]
* Fixed compile warning when compiling UbaVisualizer for Mac

diff --git a/Engine/Source/Programs/UnrealBuildAccelerator/Visualizer/Private/Mac/UbaVisualizerApp.mm b/Engine/Source/Programs/UnrealBuildAccelerator/Visualizer/Private/Mac/UbaVisualizerApp.mm
--- a/Engine/Source/Programs/UnrealBuildAccelerator/Visualizer/Private/Mac/UbaVisualizerApp.mm
+++ b/Engine/Source/Programs/UnrealBuildAccelerator/Visualizer/Private/Mac/UbaVisualizerApp.mm
@@ -171,7 +171,7 @@ virtual void Text(int x, int y, StringView str, const RECT* clipRect) override
 					reinterpret_cast<const UInt8*>(str.data),
 					static_cast<CFIndex>(str.count),
 					kCFStringEncodingUTF8,
-					kCFStringEncodingUTF8 == kCFStringEncodingUTF8 ? false : false // just "false"
+					false //kCFStringEncodingUTF8 == kCFStringEncodingUTF8 ? false : false // just "false"
 				);
 				if (!cfStr)
 					return;
@@ -571,7 +571,8 @@ void Init()
 				GetFullPathNameW(filePath, 512, fixedPath, nullptr);
 
 				m_fileName.Append(fixedPath);
-				if (false)
+				bool playFile = false;
+				if (playFile)
 				{
 					m_trace.ReadFile(m_traceView, fixedPath, true);
 
