[UBA]
* Fixed bug in DirectoryTable EntryExistsNoLock where it didn't check overlay offset if failing to find main offset
* Bumped command line buffer size in UbaCli
* Set better default values in ClientInitInfo (for UbaCli)
* Fixed bug in scheduler when forceremote is set
* Fixed so scheduler properly honor maxLocalProcessors
* Added GetMaxCount to Bottleneck class
* Moved assert of number of workers to when DependencyCrawler is initialized
* Added asserts with better descriptions in BinaryWriter WriteBytesCompressed
* Added parse functions to StringView
* Fixed so CheckPath not trigger on backslash on posix since it is a valid character in paths

diff --git a/Engine/Source/Programs/UnrealBuildAccelerator/Cli/Private/UbaCli.cpp b/Engine/Source/Programs/UnrealBuildAccelerator/Cli/Private/UbaCli.cpp
--- a/Engine/Source/Programs/UnrealBuildAccelerator/Cli/Private/UbaCli.cpp
+++ b/Engine/Source/Programs/UnrealBuildAccelerator/Cli/Private/UbaCli.cpp
@@ -361,7 +361,7 @@ namespace uba
 		TString application;
 		TString arguments;
 
-		auto parseOption = [&](const StringView& name, StringBufferBase& value)
+		auto parseOption = [&](StringView name, StringBufferBase& value)
 			{
 				if (IsWindows && name.Equals(TCV("-visualizer")))
 				{
@@ -656,7 +656,7 @@ namespace uba
 
 		auto parseArg = [&](const tchar* arg)
 		{
-			StringBuffer<> name;
+			StringBuffer<1024> name;
 			StringBuffer<32*1024> value;
 
 			if (const tchar* equals = TStrchr(arg,'='))
diff --git a/Engine/Source/Programs/UnrealBuildAccelerator/Cli/Private/UbaClient.h b/Engine/Source/Programs/UnrealBuildAccelerator/Cli/Private/UbaClient.h
--- a/Engine/Source/Programs/UnrealBuildAccelerator/Cli/Private/UbaClient.h
+++ b/Engine/Source/Programs/UnrealBuildAccelerator/Cli/Private/UbaClient.h
@@ -13,10 +13,10 @@ namespace uba
 		LogWriter& logWriter;
 		NetworkBackend& networkBackend;
 		const tchar* rootDir = nullptr;
-		const tchar* host = nullptr;
-		u16 port = 0;
-		const tchar* zone = nullptr;
-		u32 maxProcessorCount = 0;
+		const tchar* host = TC("localhost");
+		u16 port = DefaultPort;
+		const tchar* zone = TC("");
+		u32 maxProcessorCount = ~0u;
 		u32 index = 0;
 		bool createSession = true;
 		bool addDirSuffix = true;
diff --git a/Engine/Source/Programs/UnrealBuildAccelerator/Common/Private/UbaScheduler.cpp b/Engine/Source/Programs/UnrealBuildAccelerator/Common/Private/UbaScheduler.cpp
--- a/Engine/Source/Programs/UnrealBuildAccelerator/Common/Private/UbaScheduler.cpp
+++ b/Engine/Source/Programs/UnrealBuildAccelerator/Common/Private/UbaScheduler.cpp
@@ -341,7 +341,7 @@ namespace uba
 			SCOPED_FUTEX(m_processEntriesLock, lock);
 			for (auto& entry : m_processEntries)
 			{
-				UBA_ASSERTF(entry.runStatus > RunStatus_Running, TC("Found processes in queue/running state when stopping scheduler."));
+				UBA_ASSERTF(entry.runStatus > RunStatus_Running, TC("Found process in queue/running state when stopping scheduler. (%u)"), u32(entry.runStatus));
 				ProcessLink::DeleteAll(entry.firstDependent);
 				delete entry.info;
 			}
@@ -1298,16 +1298,19 @@ namespace uba
 			return returnEntry(PopFromQueue(m_readyToRunCrossArchAndPlatform));
 		}
 
+		if (!m_maxLocalProcessors)
+			return ~0u;
+
 		auto getEntry = [&](float weightLeft)
 			{
 				if (ProcessEntry* entry = PeekFromQueue(m_readyToRunLocal)) // Check local only first
 					if (entry->info->weight <= weightLeft)
 						return PopFromQueue(m_readyToRunLocal);
-				if (!m_forceRemote || (!m_hadRemoteRequest && m_hadCrossSystemRequest)) // If force remote but we've had cross system requests we let local machine pop non-cross system queue
+				if (!m_forceRemote || (!m_hadMatchingSystemRequest && m_hadCrossSystemRequest)) // If force remote but we've had cross system requests we let local machine pop non-cross system queue
 					if (ProcessEntry* entry = PeekFromQueue(m_readyToRunRemote))
 						if (entry->info->weight <= weightLeft)
 							return PopFromQueue(m_readyToRunRemote);
-				if (m_forceRemote && (!m_hadRemoteRequest || m_hadCrossSystemRequest))
+				if (m_forceRemote && (m_hadMatchingSystemRequest || m_hadCrossSystemRequest))
 					return ~0u;
 				if (ProcessEntry* entry = PeekFromQueue(m_readyToRunCrossArchitecture))
 					if (entry->info->weight <= weightLeft)
@@ -1586,10 +1589,12 @@ namespace uba
 			SCOPED_FUTEX(m_processEntriesLock, lock);
 
 			if (!isLocal)
-				m_hadRemoteRequest = true;
-			if (isCrossArchitecture || isCrossPlatform)
-				m_hadCrossSystemRequest = true;
-
+			{
+				if (isCrossArchitecture || isCrossPlatform)
+					m_hadCrossSystemRequest = true;
+				else
+					m_hadMatchingSystemRequest = true;
+			}
 			bool shouldSleep;
 			u32 indexToRun = PopProcessNoLock(isLocal, isCrossArchitecture, isCrossPlatform, prevStatus, false, shouldSleep);
 			if (indexToRun == ~0u)
diff --git a/Engine/Source/Programs/UnrealBuildAccelerator/Common/Public/UbaBottleneck.h b/Engine/Source/Programs/UnrealBuildAccelerator/Common/Public/UbaBottleneck.h
--- a/Engine/Source/Programs/UnrealBuildAccelerator/Common/Public/UbaBottleneck.h
+++ b/Engine/Source/Programs/UnrealBuildAccelerator/Common/Public/UbaBottleneck.h
@@ -68,6 +68,8 @@ namespace uba
 			first->ev.Set();
 		}
 
+		u32 GetMaxCount() const { return m_maxCount; }
+
 	private:
 		Futex m_lock;
 		BottleneckTicket* m_first = nullptr;
diff --git a/Engine/Source/Programs/UnrealBuildAccelerator/Common/Public/UbaDependencyCrawler.h b/Engine/Source/Programs/UnrealBuildAccelerator/Common/Public/UbaDependencyCrawler.h
--- a/Engine/Source/Programs/UnrealBuildAccelerator/Common/Public/UbaDependencyCrawler.h
+++ b/Engine/Source/Programs/UnrealBuildAccelerator/Common/Public/UbaDependencyCrawler.h
@@ -116,7 +116,6 @@ namespace uba
 		{
 			// We never allow the crawler to use all workers since there could be jobs that do garbage collecting of cas entries 
 			u32 max = workManager.GetWorkerCount();
-			UBA_ASSERT(max > 2);
 			if (max <= 2)
 				return 0;
 			return Min(max - 2u, bottleneckMax);
@@ -132,6 +131,7 @@ namespace uba
 
 	void DependencyCrawler::Init(FileExistsFunc&& fileExistsFunc, TraverseFilesFunc&& traverseFilesFunc, bool useBloomFilter)
 	{
+		UBA_ASSERT(m_bottleneck.GetMaxCount() > 0);
 		m_fileExistsFunc = std::move(fileExistsFunc);
 		m_traverseFilesFunc = std::move(traverseFilesFunc);
 		m_useBloomFilter = useBloomFilter;
diff --git a/Engine/Source/Programs/UnrealBuildAccelerator/Common/Public/UbaScheduler.h b/Engine/Source/Programs/UnrealBuildAccelerator/Common/Public/UbaScheduler.h
--- a/Engine/Source/Programs/UnrealBuildAccelerator/Common/Public/UbaScheduler.h
+++ b/Engine/Source/Programs/UnrealBuildAccelerator/Common/Public/UbaScheduler.h
@@ -288,7 +288,7 @@ namespace uba
 		bool m_enableSpeculativeCacheTest = false;
 		bool m_traceDetails = false;
 		bool m_cancelled = false;
-		bool m_hadRemoteRequest = false;
+		bool m_hadMatchingSystemRequest = false;
 		bool m_hadCrossSystemRequest = false;
 		ConfigTable* m_processConfigs = nullptr;
 
diff --git a/Engine/Source/Programs/UnrealBuildAccelerator/Core/Private/UbaPlatform.cpp b/Engine/Source/Programs/UnrealBuildAccelerator/Core/Private/UbaPlatform.cpp
--- a/Engine/Source/Programs/UnrealBuildAccelerator/Core/Private/UbaPlatform.cpp
+++ b/Engine/Source/Programs/UnrealBuildAccelerator/Core/Private/UbaPlatform.cpp
@@ -791,11 +791,13 @@ namespace uba
 
 	bool BinaryWriter::WriteBytesCompressed(const void* data, u64 size)
 	{
+		UBA_ASSERT(data);
+		UBA_ASSERT(size);
 		u64 needed = (u64)OodleLZ_GetCompressedBufferSizeNeeded(OodleLZ_Compressor_Kraken, size);
 		Vector<u8> mem;
 		mem.resize(needed);
 		u64 compressedSize = OodleLZ_Compress(OodleLZ_Compressor_Kraken, data, (OO_SINTa)size, mem.data(), OodleLZ_CompressionLevel_Normal);
-		UBA_ASSERT(compressedSize != OODLELZ_FAILED);
+		UBA_ASSERTF(compressedSize != OODLELZ_FAILED, TC("Failed to compress data of size %llu"), size);
 		Write7BitEncoded(compressedSize);
 		WriteBytes(mem.data(), compressedSize);
 		return true;
diff --git a/Engine/Source/Programs/UnrealBuildAccelerator/Core/Private/UbaStringBuffer.cpp b/Engine/Source/Programs/UnrealBuildAccelerator/Core/Private/UbaStringBuffer.cpp
--- a/Engine/Source/Programs/UnrealBuildAccelerator/Core/Private/UbaStringBuffer.cpp
+++ b/Engine/Source/Programs/UnrealBuildAccelerator/Core/Private/UbaStringBuffer.cpp
@@ -391,74 +391,42 @@
 
 	bool StringBufferBase::Parse(u64& out) const
 	{
-		return uba::Parse(out, data, count);
+		return ToView(*this).Parse(out);
 	}
 
 	bool StringBufferBase::Parse(u32& out, u64 offset) const
 	{
-		if (count <= offset)
-			return false;
-		#if PLATFORM_WINDOWS
-		out = wcstoul(data + offset, nullptr, 10);
-		#else
-		out = strtoul(data + offset, nullptr, 10);
-		#endif
-		return out != 0 || Equals(TCV("0"));
+		return ToView(*this).Parse(out, offset);
 	}
 
 	bool StringBufferBase::Parse(u16& out, u64 offset) const
 	{
-		u32 temp;
-		if (!Parse(temp, offset))
-			return false;
-		if (temp > 65535)
-			return false;
-		out = u16(temp);
-		return true;
+		return ToView(*this).Parse(out, offset);
 	}
 
 	bool StringBufferBase::Parse(bool& out) const
 	{
-		if (Equals(TCV("true")) || Equals(TCV("1")))
-		{
-			out = true;
-			return true;
-		}
-		if (Equals(TCV("false")) || Equals(TCV("0")))
-		{
-			out = false;
-			return true;
-		}
-		return false;
+		return ToView(*this).Parse(out);
 	}
 
 	bool StringBufferBase::Parse(float& out) const
 	{
-		if (!count)
-			return false;
-		#if PLATFORM_WINDOWS
-		out = wcstof(data, nullptr);
-		#else
-		out = strtof(data, nullptr);
-		#endif
-		return out != 0 || Equals(TCV("0"));
+		return ToView(*this).Parse(out);
 	}
 
 	bool StringBufferBase::Parse(TString& out, u64 offset) const
 	{
-		out.append(data + offset, count - offset);
-		return true;
+		return ToView(*this).Parse(out, offset);
 	}
 
 	bool StringBufferBase::Parse(StringBufferBase& out, u64 offset) const
 	{
-		out.Append(data + offset, count - offset);
-		return true;
+		return ToView(*this).Parse(out, offset);
 	}
 
 	u32 StringBufferBase::Parse(char* out, u64 outCapacity) const
 	{
-		return StringView(*this).Parse(out, outCapacity);
+		return ToView(*this).Parse(out, outCapacity);
 	}
 
 	////////////////////////////////////////////////////////////////////////////////////////////////////
@@ -495,15 +463,82 @@
 		return *this;
 	}
 
-	StringView StringView::GetPath() const
+	StringView StringView::GetPath(bool keepLastSlash) const
 	{
 		if (count)
 			for (const tchar* i=data+count-1, *e=data; i!=e; --i)
 				if (*i == PathSeparator)
-					return StringView(data, u32(i-data));
+					return StringView(data, u32(i-data) + u32(keepLastSlash));
 		return {};
 	}
 
+	bool StringView::Parse(u64& out) const
+	{
+		return uba::Parse(out, data, count);
+	}
+
+	bool StringView::Parse(u32& out, u64 offset) const
+	{
+		if (count <= offset)
+			return false;
+		#if PLATFORM_WINDOWS
+		out = wcstoul(data + offset, nullptr, 10);
+		#else
+		out = strtoul(data + offset, nullptr, 10);
+		#endif
+		return out != 0 || Equals(TCV("0"));
+	}
+
+	bool StringView::Parse(u16& out, u64 offset) const
+	{
+		u32 temp;
+		if (!Parse(temp, offset))
+			return false;
+		if (temp > 65535)
+			return false;
+		out = u16(temp);
+		return true;
+	}
+
+	bool StringView::Parse(bool& out) const
+	{
+		if (Equals(TCV("true")) || Equals(TCV("1")))
+		{
+			out = true;
+			return true;
+		}
+		if (Equals(TCV("false")) || Equals(TCV("0")))
+		{
+			out = false;
+			return true;
+		}
+		return false;
+	}
+
+	bool StringView::Parse(float& out) const
+	{
+		if (!count)
+			return false;
+		#if PLATFORM_WINDOWS
+		out = wcstof(data, nullptr);
+		#else
+		out = strtof(data, nullptr);
+		#endif
+		return out != 0 || Equals(TCV("0"));
+	}
+
+	bool StringView::Parse(TString& out, u64 offset) const
+	{
+		out.append(data + offset, count - offset);
+		return true;
+	}
+
+	bool StringView::Parse(StringBufferBase& out, u64 offset) const
+	{
+		out.Append(data + offset, count - offset);
+		return true;
+	}
+
 	u32 StringView::Parse(char* out, u64 outCapacity) const
 	{
 		#if PLATFORM_WINDOWS
@@ -529,6 +564,11 @@
 		return sv;
 	}
 
+	StringView ToView(StringBufferBase& sb)
+	{
+		return StringView(sb);
+	}
+
 	////////////////////////////////////////////////////////////////////////////////////////////////////
 
 	LastErrorToText::LastErrorToText() : LastErrorToText(GetLastError())
diff --git a/Engine/Source/Programs/UnrealBuildAccelerator/Core/Public/UbaDirectoryTable.h b/Engine/Source/Programs/UnrealBuildAccelerator/Core/Public/UbaDirectoryTable.h
--- a/Engine/Source/Programs/UnrealBuildAccelerator/Core/Public/UbaDirectoryTable.h
+++ b/Engine/Source/Programs/UnrealBuildAccelerator/Core/Public/UbaDirectoryTable.h
@@ -366,7 +366,12 @@ namespace uba
 			if (findIt != m_lookup.end())
 			{
 				if (outTableOffset)
-					outTableOffset->internal = u32(findIt->second.latestOffset) | DirectoryEntryFlag; // Use significant bit to say that this is a dir
+				{
+					if (findIt->second.latestOverlayOffset)
+						outTableOffset->internal = u32(findIt->second.latestOverlayOffset) | DirectoryEntryFlag; // Use significant bit to say that this is a dir
+					else
+						outTableOffset->internal = u32(findIt->second.latestOffset) | DirectoryEntryFlag; // Use significant bit to say that this is a dir
+				}
 				return Exists_Yes;
 			}
 			startSkip = 1;
diff --git a/Engine/Source/Programs/UnrealBuildAccelerator/Core/Public/UbaHash.h b/Engine/Source/Programs/UnrealBuildAccelerator/Core/Public/UbaHash.h
--- a/Engine/Source/Programs/UnrealBuildAccelerator/Core/Public/UbaHash.h
+++ b/Engine/Source/Programs/UnrealBuildAccelerator/Core/Public/UbaHash.h
@@ -227,7 +227,7 @@ namespace uba
 			UBA_ASSERTF(!CaseInsensitiveFs || c < 'A' || c > 'Z', TC("Path is not valid (%.*s)"), fileName.count, fileName.data); // No upper case characters if case insensitive
 			UBA_ASSERTF(!(c == PathSeparator && history[0] == '.' && history[1] == PathSeparator) || (IsWindows && fileName.Contains(TC("\\\\.\\pipe"))), TC("Path is not valid (%.*s)"), fileName.count, fileName.data);  //  /./ not allowed
 			UBA_ASSERTF(!(c == PathSeparator && history[0] == '.' && history[1] == '.' && history[2] == PathSeparator), TC("Path is not valid (%.*s)"), fileName.count, fileName.data); // /../ not allowed
-			UBA_ASSERTF(c != NonPathSeparator, TC("Path is not valid (%.*s)"), fileName.count, fileName.data);
+			UBA_ASSERTF(!IsWindows || c != NonPathSeparator, TC("Path is not valid (%.*s)"), fileName.count, fileName.data);
 			//UBA_ASSERTF(c != PathSeparator || history[0] != PathSeparator || (IsWindows && fileName.Contains(TC("\\\\.\\pipe"))), TC("Path is not valid (%.*s)"), fileName.count, fileName.data); // Double backslash
 			history[2] = history[1];
 			history[1] = history[0];
diff --git a/Engine/Source/Programs/UnrealBuildAccelerator/Core/Public/UbaStringBuffer.h b/Engine/Source/Programs/UnrealBuildAccelerator/Core/Public/UbaStringBuffer.h
--- a/Engine/Source/Programs/UnrealBuildAccelerator/Core/Public/UbaStringBuffer.h
+++ b/Engine/Source/Programs/UnrealBuildAccelerator/Core/Public/UbaStringBuffer.h
@@ -110,6 +110,7 @@ namespace uba
 		StringView(const TString& str) : data(str.data()), count(u32(str.size())) {}
 
 		tchar operator[](u64 i) const { return data[i]; }
+		tchar Last() { return count ? data[count - 1] : 0; }
 
 		bool IsEmpty() const { return count == 0; }
 		bool StartsWith(StringView str, bool ignoreCase = true) const; // Requires terminated StringViews
@@ -122,12 +123,20 @@ namespace uba
 		bool Equals(const tchar* str, bool ignoreCase = true) const { return uba::Equals(data, str, ignoreCase); } // Requires terminated StringViews
 		bool Equals(StringView str, bool ignoreCase = true) const { return count == str.count && uba::Equals(data, str.data, count, ignoreCase); }
 		StringView GetFileName() const;
-		StringView GetPath() const;
+		StringView GetPath(bool keepLastSlash = false) const;
 		StringView Skip(u32 skipCount) const { return StringView(data + skipCount, count - skipCount); }
+		StringView RemoveSuffix(u32 removeCount) { return StringView(data, count - removeCount); }
 		StringView SubStr(u32 startOffset, u32 length) const { return StringView(data + startOffset, length); }
 
 		TString ToString() const { return TString(data, data + count); }
 
+		bool Parse(u64& out) const;
+		bool Parse(u32& out, u64 offset = 0) const;
+		bool Parse(u16& out, u64 offset = 0) const;
+		bool Parse(bool& out) const;
+		bool Parse(float& out) const;
+		bool Parse(TString& out, u64 offset = 0) const;
+		bool Parse(StringBufferBase& out, u64 offset = 0) const;
 		u32 Parse(char* out, u64 outCapacity) const; // Note, return value contains length + null termination.. zero if failed
 
 		const tchar* data;
@@ -136,6 +145,7 @@ namespace uba
 
 	StringView ToView(const tchar* s);
 	StringView ToView(StringView sv);
+	StringView ToView(StringBufferBase& sb);
 
 	template<typename T, size_t size>
 	constexpr StringView AsView(T (&buffer)[size]) { return StringView(buffer, size-1); }
