[UBA]
* Added code for handling ENOTEMPTY when calling RemoveDirectory

diff --git a/Engine/Source/Programs/UnrealBuildAccelerator/Common/Private/UbaFile.cpp b/Engine/Source/Programs/UnrealBuildAccelerator/Common/Private/UbaFile.cpp
--- a/Engine/Source/Programs/UnrealBuildAccelerator/Common/Private/UbaFile.cpp
+++ b/Engine/Source/Programs/UnrealBuildAccelerator/Common/Private/UbaFile.cpp
@@ -713,6 +713,11 @@ namespace uba
 			SetLastError(ERROR_FILE_NOT_FOUND);
 			return false;
 		}
+		if (errno == ENOTEMPTY)
+		{
+			SetLastError(ERROR_DIR_NOT_EMPTY);
+			return false;
+		}
 		UBA_ASSERTF(false, TC("RemoveDirectoryW error handling not implemented (%s): %s"), strerror(errno), pathName);
 		return false;
 #endif
diff --git a/Engine/Source/Programs/UnrealBuildAccelerator/Common/Public/UbaFile.h b/Engine/Source/Programs/UnrealBuildAccelerator/Common/Public/UbaFile.h
--- a/Engine/Source/Programs/UnrealBuildAccelerator/Common/Public/UbaFile.h
+++ b/Engine/Source/Programs/UnrealBuildAccelerator/Common/Public/UbaFile.h
@@ -16,6 +16,7 @@ namespace uba
 	inline constexpr u32 ERROR_FILE_EXISTS = EEXIST;
 	inline constexpr u32 ERROR_ACCESS_DENIED = EACCES;
 	inline constexpr u32 ERROR_DIRECTORY = ENOTDIR;
+	inline constexpr u32 ERROR_DIR_NOT_EMPTY = ENOTEMPTY;
 	inline constexpr u32 MOVEFILE_REPLACE_EXISTING = 0x00000001;
 	inline constexpr u32 INVALID_FILE_ATTRIBUTES = (u32)-1;
 
