[UBA]
* Changed so changes to file system (Create/Delete files/directories etc) is stored in overlay instead of main directory table and then committed after process hierarchy has finished successfully (THIS IS A HUGE CHANGE)
* Include temp dir in overlay and removed lots of special cases
* Fixed bug where detoured handles were transferred to child processes. Now detoured handles always transfer the true handles instead (cbReserved2)
* Fixed bug in GetVolumeInformationW when vfs is enabled. Now it will just use the local drive of the first vfs entry just to get some real information (this fixes the annoying extra std outs when cross compiling linux)
* Fixed bugs in how command lines are expanded with cmd.exe and posix shell. Now it should handle more complex command lines
* Fixed bug related to dynamically loading .so files on remote machines. Solved by always adding LD_LIBRARY_PATH/DYLD_LIBRARY_PATH with executable dir
* Removed flushing of written files on remotes when sh exits, this is now properly handled with overlay
* Added proper support for symlinks on posix
* Fixed proper support for remotes doing read+write on files (which means they will need to download file, create new temporary file and copy content from downloaded file to temporary file
* Added support for a lot of posix/linux functions in detouring (mkstemp, __lxstat64, openat64, stat64, scandir64, statx, dup3 to name a few)

diff --git a/Engine/Source/Programs/UnrealBuildAccelerator/Common/Private/UbaDirectoryTableHolder.cpp b/Engine/Source/Programs/UnrealBuildAccelerator/Common/Private/UbaDirectoryTableHolder.cpp
--- a/Engine/Source/Programs/UnrealBuildAccelerator/Common/Private/UbaDirectoryTableHolder.cpp
+++ b/Engine/Source/Programs/UnrealBuildAccelerator/Common/Private/UbaDirectoryTableHolder.cpp
@@ -33,7 +33,7 @@ namespace uba
 		return true;
 	}
 
-	StringKey GetKeyAndFixedName2(StringBuffer<>& outFixedFilePath, StringKeyHasher& outHasher, const tchar* filePath, StringKey* outDirKey, StringView* outFileName)
+	StringKey GetKeyAndFixedName2(StringBuffer<>& outFixedFilePath, StringKeyHasher& outHasher, StringView filePath, StringKey* outDirKey, StringView* outFileName)
 	{
 		StringBuffer<> workingDir;
 		if (!IsAbsolutePath(filePath))
@@ -41,7 +41,7 @@ namespace uba
 			GetCurrentDirectoryW(workingDir);
 			workingDir.EnsureEndsWithSlash();
 		}
-		FixPath(filePath, workingDir.data, workingDir.count, outFixedFilePath);
+		FixPath(outFixedFilePath, filePath, workingDir);
 
 		StringKey dirKey;
 		StringBuffer<> dirNameForHash;
@@ -83,7 +83,7 @@ namespace uba
 	StringKey GetKeyAndFixedName(StringBuffer<>& outFixedFilePath, const tchar* filePath)
 	{
 		StringKeyHasher hasher;
-		return GetKeyAndFixedName2(outFixedFilePath, hasher, filePath, nullptr, nullptr);
+		return GetKeyAndFixedName2(outFixedFilePath, hasher, ToView(filePath), nullptr, nullptr);
 	}
 
 	////////////////////////////////////////////////////////////////////////////////////////////////////
@@ -127,7 +127,7 @@ namespace uba
 
 		StringBuffer<> fixedDirPath;
 		StringKeyHasher hasher;
-		GetKeyAndFixedName2(fixedDirPath, hasher, dirPath, nullptr, nullptr);
+		GetKeyAndFixedName2(fixedDirPath, hasher, ToView(dirPath), nullptr, nullptr);
 
 		StringKey dirKey = ToStringKey(hasher);
 
@@ -162,7 +162,7 @@ namespace uba
 		UBA_ASSERT(!m_runningRemote);
 		StringBuffer<> fixedFilePath;
 		auto key = GetKeyAndFixedName(fixedFilePath, filePath);
-		return RegisterCreateFileForWrite(key, fixedFilePath, nullptr, 0, 0, true, true);
+		return RegisterCreateFileForWrite(key, fixedFilePath, 0, 0, 0, true, true);
 	}
 
 	void DirectoryTableHolder::RegisterDeleteFile(const tchar* filePath)
@@ -170,7 +170,7 @@ namespace uba
 		UBA_ASSERT(!m_runningRemote);
 		StringBuffer<> fixedFilePath;
 		auto key = GetKeyAndFixedName(fixedFilePath, filePath);
-		RegisterDeleteFile(key, fixedFilePath, nullptr);
+		RegisterDeleteFile(key, fixedFilePath);
 	}
 
 	bool DirectoryTableHolder::RegisterNewDirectory(const tchar* directoryPath)
@@ -178,7 +178,7 @@ namespace uba
 		UBA_ASSERT(!m_runningRemote);
 		StringBuffer<> fixedDirPath;
 		auto dirKey = GetKeyAndFixedName(fixedDirPath, directoryPath);
-		if (!RegisterCreateFileForWrite(dirKey, fixedDirPath, nullptr, 0, 0, true, true))
+		if (!RegisterCreateFileForWrite(dirKey, fixedDirPath, 0, 0, 0, true, true))
 			return false;
 		WriteDirectoryEntries(dirKey, fixedDirPath);
 		return true;
@@ -291,13 +291,36 @@ namespace uba
 		}
 	}
 
+	bool DirectoryTableHolder::RegisterOverlayCreateOrWrite(DirectoryTableOverlay& overlay, StringKey filePathKey, StringView filePath, u64 fileSize, u64 lastWriteTime, u32 attributes)
+	{
+		UBA_ASSERTF(attributes != 0, TC("Registering file %.*s with 0 attributes"), filePath.count, filePath.data);
+		FileInformation info;
+		info.attributes = attributes;
+		info.index = ++m_fileIndexCounter;
+		info.volumeSerialNumber = 1;
+		info.size = fileSize;
+		info.lastWriteTime = lastWriteTime;
+		return WriteOverlayEntry(overlay, filePath, info);
+	}
+
+	bool DirectoryTableHolder::RegisterOverlayDelete(DirectoryTableOverlay& overlay, StringKey filePathKey, StringView filePath)
+	{
+		StringKey dirKey;
+		const tchar* lastSlash;
+		StringBuffer<> dirName;
+		if (!GetDirKey(dirKey, dirName, lastSlash, filePath))
+			return false;
+		FileInformation fileInfo;
+		return WriteOverlayEntry(overlay, dirKey, filePath, ToView(lastSlash+1), fileInfo, !m_runningRemote);
+	}
+
 	bool DirectoryTableHolder::PopulateOverlayFile(DirectoryTableOverlay& overlay, StringView filePath, const Function<bool(FileInformation&, StringKey)>& getFileInformation)
 	{
 		StringKey dirKey;
 		StringView fileName;
 		StringBuffer<> fixedFilePath;
 		StringKeyHasher hasher;
-		StringKey filePathKey = GetKeyAndFixedName2(fixedFilePath, hasher, filePath.data, &dirKey, &fileName);
+		StringKey filePathKey = GetKeyAndFixedName2(fixedFilePath, hasher, filePath, &dirKey, &fileName);
 		FileInformation fileInfo;
 		if (!getFileInformation(fileInfo, filePathKey))
 			return false;
@@ -313,21 +336,18 @@ namespace uba
 		return WriteOverlayEntry(overlay, filePath, fi);
 	}
 
-	bool DirectoryTableHolder::MoveOverlayFile(DirectoryTableOverlay& overlay, StringView from, StringView to)
-	{
-		if (!WriteOverlayEntry(overlay, from, {}))
-			return false;
-		return AddOverlayFile(overlay, to);
-	}
-
 	bool DirectoryTableHolder::WriteOverlayEntry(DirectoryTableOverlay& overlay, StringView filePath, const FileInformation& fileInfo)
 	{
 		StringKey dirKey;
 		StringView fileName;
 		StringBuffer<> fixedFilePath;
 		StringKeyHasher hasher;
-		StringKey filePathKey = GetKeyAndFixedName2(fixedFilePath, hasher, filePath.data, &dirKey, &fileName);
-		return WriteOverlayEntry(overlay, dirKey, fixedFilePath, fileName, fileInfo, !m_runningRemote);
+		StringKey filePathKey = GetKeyAndFixedName2(fixedFilePath, hasher, filePath, &dirKey, &fileName);
+		if (!WriteOverlayEntry(overlay, dirKey, fixedFilePath, fileName, fileInfo, !m_runningRemote))
+			return false;
+		if (IsDirectory(fileInfo.attributes))
+			return WriteOverlayDirInLookup(overlay, filePath);
+		return true;
 	}
 
 	bool DirectoryTableHolder::WriteOverlayEntry(DirectoryTableOverlay& overlay, StringKey dirKey, StringView filePath, StringView fileName, const FileInformation& fileInfo, bool populateMainTableParentDir)
@@ -392,6 +412,33 @@ namespace uba
 		return true;
 	}
 
+	bool DirectoryTableHolder::WriteOverlayDirInLookup(DirectoryTableOverlay& overlay, StringView dirPath)
+	{
+		UBA_ASSERT(dirPath.Last() != PathSeparator);
+		StringKey dirKey = CaseInsensitiveFs ? ToStringKeyLower(dirPath) : ToStringKey(dirPath);
+		auto insres = overlay.lookup.try_emplace(dirKey);
+		if (!insres.second)
+			return true;
+		auto& dir = insres.first->second;
+		u8 temp[256];
+		BinaryWriter writer(temp, 0, sizeof(temp));
+		writer.WriteStringKey(dirKey);
+		writer.Write7BitEncoded(0); // Previous entry for same directory
+		writer.WriteFileAttributes(DefaultDirAttributes());
+		writer.WriteVolumeSerial(0);
+		writer.WriteFileIndex(++m_fileIndexCounter);
+		writer.Write7BitEncoded(0); // Count
+		u64 written = writer.GetPosition();
+		if (!EnsureOverlayMemory(overlay, written + 8))
+			return false;
+		BinaryWriter memWriter(overlay.memory, overlay.size, overlay.committed);
+		memWriter.Write7BitEncoded(writer.GetPosition()); // Storage size
+		dir.offset = u32(memWriter.GetPosition()) + sizeof(StringKey);
+		memWriter.WriteBytes(writer.GetData(), written);
+		overlay.size = u32(memWriter.GetPosition());
+		return true;
+	}
+
 	bool DirectoryTableHolder::WriteOverlayDirectoryEntriesRecursive(DirectoryTableOverlay& overlay, StringKey filePathKey, StringView filePath)
 	{
 		StringView dirPath = filePath.GetPath();
@@ -421,7 +468,7 @@ namespace uba
 		return WriteOverlayEntry(overlay, dirKey, dir, filePath.GetFileName(), fi, res.first->second);
 	}
 
-	bool DirectoryTableHolder::RegisterCreateFileForWrite(StringKey filePathKey, StringView filePath, DirectoryTableOverlay* overlay, u64 fileSize, u64 lastWriteTime, bool registerRealFile, bool invalidateStorage)
+	bool DirectoryTableHolder::RegisterCreateFileForWrite(StringKey filePathKey, StringView filePath, u64 fileSize, u64 lastWriteTime, u32 attributes, bool registerRealFile, bool invalidateStorage)
 	{
 		auto& dirTable = m_directoryTable;
 
@@ -432,22 +479,10 @@ namespace uba
 			return true;
 
 		// Remote is not updating its own directory table
-		if (m_runningRemote)
-		{
-			if (overlay)
-			{
-				FileInformation info;
-				info.attributes = DefaultAttributes();
-				info.index = ++m_fileIndexCounter;
-				info.volumeSerialNumber = 1;
-				info.size = fileSize;
-				info.lastWriteTime = lastWriteTime;
-				WriteOverlayEntry(*overlay, filePath, info);
-			}
-			return true;
-		}
+		UBA_ASSERT(!m_runningRemote);
+		UBA_ASSERTF(!IsInTemp(filePath), TC("Registering temp files to main table not allowed (%.*s)"), filePath.count, filePath.data);
 
-#if 0//_DEBUG  // Bring this back, turned off right now because a few lines above the call to this method we add a mapping
+		#if 0//_DEBUG  // Bring this back, turned off right now because a few lines above the call to this method we add a mapping
 		{
 			SCOPED_FUTEX(m_fileMappingTableLookupLock, lookupLock);
 			auto findIt = m_fileMappingTableLookup.find(filePathKey);
@@ -459,7 +494,7 @@ namespace uba
 				UBA_ASSERT(!entry.mapping);
 			}
 		}
-#endif
+		#endif
 
 		bool shouldWriteToDisk = registerRealFile && ShouldWriteToDisk(filePath);
 
@@ -510,7 +545,6 @@ namespace uba
 		auto insres = dir.files.try_emplace(filePathKey, ~u32(0));
 
 		u64 fileIndex = InvalidValue;
-		u32 attributes = 0;
 		u32 volumeSerial = 0;
 		bool isDirectory = false;
 
@@ -533,7 +567,7 @@ namespace uba
 		else
 		{
 			// TODO: Do we need more code here?
-			attributes = DefaultAttributes();
+			UBA_ASSERT(attributes != 0);
 			volumeSerial = 1;
 			fileIndex = ++m_fileIndexCounter;
 		}
@@ -559,14 +593,17 @@ namespace uba
 			if (lastWriteTime == oldLastWriteTime)
 			{
 #if !PLATFORM_WINDOWS
-				UBA_ASSERT(oldFileIndex == fileIndex); // Checking so it is really the same file
+				if (oldFileIndex == fileIndex)
+				//UBA_ASSERTF(oldFileIndex == fileIndex, TC("Same write time but Not same file? %llu vs %llu (%.*s)"), oldFileIndex, fileIndex, filePath.count, filePath.data); // Checking so it is really the same file
 #endif
-				u64 oldSize = reader.ReadFileSize();
-				if (oldSize == fileSize && oldAttr == attributes) // Only attributes could change from a chmod
-					return true;
-				// TODO: Somehow this can happen and I have no idea how. last written time should be set on close file so it shouldnt be possible.
-				//else
-				//	m_logger.Warning(TC("Somehow file %s has same last written time at two points in time but different size (old %llu new %llu)"), filePath.data, oldSize, fileSize);
+				{
+					u64 oldSize = reader.ReadFileSize();
+					if (oldSize == fileSize && oldAttr == attributes) // Only attributes could change from a chmod
+						return true;
+					// TODO: Somehow this can happen and I have no idea how. last written time should be set on close file so it shouldnt be possible.
+					//else
+					//	m_logger.Warning(TC("Somehow file %s has same last written time at two points in time but different size (old %llu new %llu)"), filePath.data, oldSize, fileSize);
+				}
 			}
 		}
 
@@ -603,19 +640,6 @@ namespace uba
 		m_debugLogger->Info(TC("TRACKADD    %s (Size: %llu, Attr: %u, Key: %s, Id: %llu)"), filePath.data, fileSize, attributes, KeyToString(filePathKey).data, fileIndex);
 #endif
 
-		if (overlay && !isDirectory)
-		{
-			FileInformation info;
-			info.attributes = attributes;
-			info.index = fileIndex;
-			info.lastWriteTime = lastWriteTime;
-			info.size = fileSize;
-			info.volumeSerialNumber = volumeIndex;
-			if (!WriteOverlayEntry(*overlay, dirKey, filePath, ToView(lastSlash + 1), info, false))
-				return false;
-		}
-
-
 		SCOPED_WRITE_LOCK(dirTable.m_memoryLock, memoryLock);
 		u32 writePos = dirTable.m_memorySize;
 		EnsureDirectoryTableMemory(writePos + 8 + written);
@@ -629,7 +653,7 @@ namespace uba
 		return true;
 	}
 
-	bool DirectoryTableHolder::RegisterDeleteFile(StringKey filePathKey, StringView filePath, DirectoryTableOverlay* overlay)
+	bool DirectoryTableHolder::RegisterDeleteFile(StringKey filePathKey, StringView filePath)
 	{
 		StringKey dirKey;
 		const tchar* lastSlash;
@@ -637,17 +661,8 @@ namespace uba
 		if (!GetDirKey(dirKey, dirName, lastSlash, filePath))
 			return false;
 
-		if (overlay)
-		{
-			// TODO: If it does not exist in overlay nor main it should not create any entry
-			FileInformation fileInfo;
-			if (!WriteOverlayEntry(*overlay, dirKey, filePath, ToView(lastSlash+1), fileInfo, !m_runningRemote))
-				return false;
-		}
-
-		// Remote is not updating its own directory table
-		if (m_runningRemote)
-			return true;
+		UBA_ASSERT(!m_runningRemote);
+		UBA_ASSERT(!IsInTemp(filePath));
 
 		auto& dirTable = m_directoryTable;
 		SCOPED_READ_LOCK(dirTable.m_lookupLock, lookupCs);
@@ -796,7 +811,7 @@ namespace uba
 		dirLock.Leave();
 
 		if (parentKey == m_directoryForcedEmpty) // We need to manually add this dir to the forced empty dir
-			RegisterCreateFileForWrite(parentKey, dirPath, nullptr, 0, 0, true, false);
+			RegisterCreateFileForWrite(parentKey, dirPath, 0, 0, DefaultDirAttributes(), true, false);
 	}
 
 	bool DirectoryTableHolder::WriteDirectoryEntriesInternal(DirectoryTable::Directory& dir, const StringKey& dirKey, StringView dirPath, bool isRefresh, u32& outTableOffset)
@@ -916,9 +931,7 @@ namespace uba
 #if UBA_DEBUG_TRACK_DIR
 					m_debugLogger->Info(TC("    FAILED (not existing?)"));
 #endif
-
-					if (!IsWindows || dirPath.count > 3)
-						return false;
+					return false;
 				}
 		}
 		else
@@ -997,6 +1010,11 @@ namespace uba
 		LockMemory(m_logger, address, toCommit, TC("DirectoryTable"));
 	}
 
+	bool DirectoryTableHolder::IsInTemp(StringView path)
+	{
+		return false;
+	}
+
 	void DirectoryTableHolder::FileEntryAdded(StringKey filePathKey, u64 lastWritten, u64 size)
 	{
 	}
diff --git a/Engine/Source/Programs/UnrealBuildAccelerator/Common/Private/UbaProcess.cpp b/Engine/Source/Programs/UnrealBuildAccelerator/Common/Private/UbaProcess.cpp
--- a/Engine/Source/Programs/UnrealBuildAccelerator/Common/Private/UbaProcess.cpp
+++ b/Engine/Source/Programs/UnrealBuildAccelerator/Common/Private/UbaProcess.cpp
@@ -278,8 +278,14 @@ namespace uba
 		if (m_parentProcess)
 			m_waitForParent.Create(EventResetType_Manual);
 		else if (!m_shared.overlay.handle.IsValid())
+		{
 			m_session.CreateOverlay(m_shared.overlay);
 
+			// This is a little bit hacky but we only need this dir and not its parents. So we write only a dir entry with no children or parents
+			if (m_runningRemote)
+				m_session.WriteOverlayDirInLookup(m_shared.overlay, ToView(m_session.m_tempPath).RemoveSuffix(1));
+		}
+
 		if (!m_startInfo.overlayFilesStr.empty())
 		{
 			UBA_ASSERT(!m_parentProcess);
@@ -951,13 +957,13 @@ namespace uba
 
 		m_session.m_storage.InvalidateCachedFileInfo(toKey);
 
-		m_session.RegisterCreateFileForWrite(toKey, toName, nullptr, 0, 0, true, true);
+		m_session.RegisterCreateFileForWrite(toKey, toName, 0, 0, DefaultAttributes(), true, true);
 
 		//m_writtenFiles.try_emplace(name);
 		//WrittenFile& writtenFile = m_writtenFiles[toName.data];
 		//writtenFile.key = toKey;
 		//writtenFile.name = toName.data;
-		//writtenFile.owner = this;
+		//writtenFile.ownerId = GetId();
 		//writtenFile.attributes = DefaultAttributes();
 
 		StackBinaryWriter<1024> trackedInputs;
@@ -1103,11 +1109,12 @@ namespace uba
 		CloseFileMessage msg { *this };
 		reader.ReadString(msg.fileName);
 		msg.closeId = reader.ReadU32();
-		msg.attributes = DefaultAttributes(); // reader.ReadFileAttributes(); TODO
 		msg.deleteOnClose = reader.ReadBool();
 		msg.success = reader.ReadBool();
 		msg.memoryHandle = SharedMemoryHandle::FromU64(reader.ReadU64());
-		msg.memoryWritten = reader.ReadU64();
+		msg.fileSize = reader.ReadU64();
+		msg.lastWriteTime = reader.ReadU64();
+		msg.attributes = reader.ReadU32();
 		msg.newNameKey = reader.ReadStringKey();
 		if (msg.newNameKey != StringKeyZero)
 			reader.ReadString(msg.newName);
@@ -1156,6 +1163,10 @@ namespace uba
 		MoveFileMessage msg { *this };
 		msg.fromKey = reader.ReadStringKey();
 		reader.ReadString(msg.fromName);
+		msg.fileSize = reader.ReadU64();
+		msg.lastWriteTime = reader.ReadU64();
+		msg.fileAttributes = reader.ReadU32();
+
 		msg.toKey = reader.ReadStringKey();
 		reader.ReadString(msg.toName);
 		msg.flags = reader.ReadU32();
@@ -1179,16 +1190,34 @@ namespace uba
 		writer.WriteU32(response.errorCode);
 		return true;
 	}
+
+	bool ProcessImpl::HandleSymlink(BinaryReader& reader, BinaryWriter& writer)
+	{
+		SymlinkMessage msg { *this };
+		msg.filePathKey = reader.ReadStringKey();
+		reader.ReadString(msg.symlinkFilePath);
+		msg.contentPathKey = reader.ReadStringKey();
+		reader.ReadString(msg.symlinkContentPath);
+		reader.ReadString(msg.symlinkContent);
+		msg.fileSize = reader.ReadFileSize();
+		msg.lastWriteTime = reader.ReadFileTime();
+		msg.fileAttributes = reader.ReadFileAttributes();
+		SymlinkResponse response;
+		CallSession(MessageType_Symlink, m_session.Symlink(response, msg));
+		writer.WriteBool(response.success);
+		writer.WriteU32(response.errorCode);
+		writer.WriteU64(ToU64(response.directoryTableSize));
+		return true;
+	}
+
 	bool ProcessImpl::HandleCreateDirectory(BinaryReader& reader, BinaryWriter& writer)
 	{
 		// TODO: Should be removed and handled by overlay
-		CreateDirectoryMessage msg;
+		CreateDirectoryMessage msg { *this };
 		msg.nameKey = reader.ReadStringKey();
 		reader.ReadString(msg.name);
 		CreateDirectoryResponse response;
-		CallSession(MessageType_CreateDirectory, m_session.CreateDirectory(response, msg, nullptr));
-		writer.WriteBool(response.result);
-		writer.WriteU32(response.errorCode);
+		CallSession(MessageType_CreateDirectory, m_session.CreateDirectory(response, msg));
 		writer.WriteU64(ToU64(response.directoryTableSize));
 		return true;
 	}
@@ -1196,13 +1225,11 @@ namespace uba
 	bool ProcessImpl::HandleRemoveDirectory(BinaryReader& reader, BinaryWriter& writer)
 	{
 		// TODO: Should be removed and handled by overlay
-		RemoveDirectoryMessage msg;
+		RemoveDirectoryMessage msg { *this };
 		msg.nameKey = reader.ReadStringKey();
 		reader.ReadString(msg.name);
 		RemoveDirectoryResponse response;
-		CallSession(MessageType_RemoveDirectory, m_session.RemoveDirectory(response, msg, &m_shared.overlay));
-		writer.WriteBool(response.result);
-		writer.WriteU32(response.errorCode);
+		CallSession(MessageType_RemoveDirectory, m_session.RemoveDirectory(response, msg));
 		writer.WriteU64(ToU64(response.directoryTableSize));
 		return true;
 	}
@@ -1245,6 +1272,7 @@ namespace uba
 	{
 		if (auto msg = m_shared.firstOverflowMessage)
 		{
+			m_session.m_logger.Info(TC("%u OVERFLOW!!!"), m_id);
 			writer.WriteBytes(msg->buffer, msg->written);
 			m_shared.firstOverflowMessage = msg->next;
 			delete msg;
@@ -1273,7 +1301,7 @@ namespace uba
 				// and track write file stamp or something on the file and track what child processes have
 				bool ownedByThisOrAncestor = false;
 				for (ProcessImpl* it = this; it && !ownedByThisOrAncestor; it=it->m_parentProcess)
-					ownedByThisOrAncestor = wf.owner == it;
+					ownedByThisOrAncestor = wf.ownerId == it->m_id;
 				if (ownedByThisOrAncestor)
 					continue;
 			}
@@ -1885,9 +1913,20 @@ namespace uba
 			// But for now we only let the root process write out all output files
 			if (m_parentProcess)
 				continue;
-			kv.second.owner = nullptr;
+			kv.second.ownerId = 0;
 			if (kv.second.memoryHandle.IsValid())
+			{
 				files.push_back(&kv.second);
+				continue;
+			}
+			if (m_runningRemote)
+				continue;
+			StringView file(kv.second.name);
+			if (file.StartsWith(m_session.m_tempPath) && !m_startInfo.rules->KeepTempOutputFile(file))
+				continue;
+			if (m_startInfo.rules->IsInvisible(file))
+				continue;
+			m_session.RegisterCreateFileForWrite(StringKeyZero, kv.second.name, 0, 0, 0, true, true); // Files that were written directly to disk like Default.rc2.res.obj
 		}
 
 		// We want file count to match number of files.. and it is actually fine that count is 0
@@ -2348,12 +2387,12 @@ namespace uba
 
 		if (!m_parentProcess)
 		{
-			const char* realApplication = m_realApplication.c_str();
+			StringView realApplication = m_realApplication;
 			StringBuffer<> tempApplication;
 			if (m_realApplication.find(' ') != TString::npos)
 			{
 				tempApplication.Append('\"').Append(m_realApplication).Append('\"');
-				realApplication = tempApplication.data;
+				realApplication = tempApplication;
 			}
 
 			Vector<TString> arguments;
@@ -2421,17 +2460,17 @@ namespace uba
 				#endif
 			}
 
-			StringBuffer<128> comIdVar;
-			StringBuffer<512> workingDir;
-			StringBuffer<32> rulesStr;
-			StringBuffer<512> logFile;
-			StringBuffer<512> ldLibraryPath;
-			StringBuffer<512> detoursVar;
-			StringBuffer<32> processVar;
+			TString comIdVar;
+			TString workingDir;
+			TString rulesStr;
+			TString logFile;
+			TString ldLibraryPath;
+			TString detoursVar;
+			TString processVar;
 
 			#if UBA_USE_NATIVE_MAC_SEMAPHORES
-			StringBuffer<32> sessionId;
-			StringBuffer<32> pidVar;
+			TString sessionId;
+			TString pidVar;
 			#endif
 
 			Vector<const char*> envvars;
@@ -2451,62 +2490,72 @@ namespace uba
 
 			if (m_detourEnabled)
 			{
-				const char* detoursLib = m_session.m_detoursLibrary[IsLinux][IsArmBinary].c_str();
-				if (*detoursLib)
+				auto& t = m_session.m_detoursLibrary[IsLinux][IsArmBinary];
+				StringView detoursLib(t.data(), u32(t.size()));
+
+				StringBuffer<1024> temp;
+
+				if (detoursLib.count)
 				{
 #if PLATFORM_LINUX
-					//if (strchr(detoursLib, ' '))
+					StringView detoursLibPath = detoursLib.GetPath();
+					if (detoursLibPath.count)
 					{
-						const char* lastSlash = strrchr(detoursLib, '/');
-						UBA_ASSERT(lastSlash);
-						StringBuffer<> ldLibPath;
-						ldLibPath.Append(detoursLib, lastSlash - detoursLib);
-						ldLibraryPath.Append("LD_LIBRARY_PATH=").Append(ldLibPath);
-						detoursLib = lastSlash + 1;
+						temp.Append("LD_LIBRARY_PATH=").Append(detoursLib.GetPath());
+						detoursLib = detoursLib.Skip(detoursLibPath.count+1);
 					}
 #endif
 				}
 				else
-					detoursLib = "./" UBA_DETOURS_LIBRARY;
+					detoursLib = TCV("./" UBA_DETOURS_LIBRARY);
 
+				if (m_runningRemote)
+				{
+					constexpr auto ldLibraryPathName = IsLinux ? TCV("LD_LIBRARY_PATH") : TCV("DYLD_LIBRARY_PATH");
+					if (temp.count)
+						temp.Append(':');
+					else
+						temp.Append(ldLibraryPathName).Append('=');
+					temp.Append(realApplication.GetPath());
+				}
+
+				ldLibraryPath = temp.ToStringAndClear();
 #if PLATFORM_LINUX
-				detoursVar.Append("LD_PRELOAD=").Append(detoursLib);
+				detoursVar = temp.Append("LD_PRELOAD=").Append(detoursLib).ToStringAndClear();
 #else
-				detoursVar.Append("DYLD_INSERT_LIBRARIES=").Append(detoursLib);
-				if (m_runningRemote)
-					ldLibraryPath.Append(TCV("DYLD_LIBRARY_PATH=")).AppendDir(realApplication); // This is needed in case @rpath is "wrong" so it wants to resolve dylibs outside of own directory
+				detoursVar = temp.Append("DYLD_INSERT_LIBRARIES=").Append(detoursLib).ToStringAndClear();
 #endif
 
-				processVar.Append("UBA_SESSION_PROCESS=").AppendValue(getpid());
+				processVar = temp.Append("UBA_SESSION_PROCESS=").AppendValue(getpid()).ToStringAndClear();
 
 
-				comIdVar.Append("UBA_COMID=").AppendValue(communicationHandle.internal).Append('+').AppendValue(communicationOffset);
+				comIdVar = temp.Append("UBA_COMID=").AppendValue(communicationHandle.internal).Append('+').AppendValue(communicationOffset).ToStringAndClear();
 
 				#if UBA_USE_NATIVE_MAC_SEMAPHORES
-				sessionId.Append("UBA_SESSION_ID=").AppendValue(m_session.GetSessionId());
-				pidVar.Append("UBA_PID=").AppendValue(m_id);
+				sessionId = temp.Append("UBA_SESSION_ID=").AppendValue(m_session.GetSessionId()).ToStringAndClear;
+				pidVar = temp.Append("UBA_PID=").AppendValue(m_id).ToStringAndClear();
 				#endif
 
-				workingDir.Append("UBA_CWD=").Append(m_realWorkingDir);
-				rulesStr.Append("UBA_RULES=").AppendValue(m_startInfo.rules->index);
+				workingDir = temp.Append("UBA_CWD=").Append(m_realWorkingDir).ToStringAndClear();
+				rulesStr = temp.Append("UBA_RULES=").AppendValue(m_startInfo.rules->index).ToStringAndClear();
 
 				if (*m_startInfo.logFile)
 				{
 #if !UBA_DEBUG_LOG_ENABLED
 					static bool runOnce = [&]() { logger.Warning(TC("Build has log files disabled so no logs will be produced")); return false; }();
 #endif
-					logFile.Append("UBA_LOGFILE=").Append(m_startInfo.logFile);
+					logFile = temp.Append("UBA_LOGFILE=").Append(m_startInfo.logFile).ToStringAndClear();
 				}
 
 				//envvars.push_back("LD_DEBUG=bindings");
 
-				if (ldLibraryPath.count)
-					envvars.push_back(ldLibraryPath.data);
-				envvars.push_back(detoursVar.data);
-				envvars.push_back(processVar.data);
-				envvars.push_back(comIdVar.data);
-				envvars.push_back(workingDir.data);
-				envvars.push_back(rulesStr.data);
+				if (!ldLibraryPath.empty())
+					envvars.push_back(ldLibraryPath.data());
+				envvars.push_back(detoursVar.data());
+				envvars.push_back(processVar.data());
+				envvars.push_back(comIdVar.data());
+				envvars.push_back(workingDir.data());
+				envvars.push_back(rulesStr.data());
 
 				#if UBA_USE_NATIVE_MAC_SEMAPHORES
 				envvars.push_back(sessionId.data);
@@ -2515,8 +2564,8 @@ namespace uba
 
 				if (m_runningRemote)
 					envvars.push_back("UBA_REMOTE=1");
-				if (!logFile.IsEmpty())
-					envvars.push_back(logFile.data);
+				if (!logFile.empty())
+					envvars.push_back(logFile.data());
 			}
 
 			envvars.push_back(nullptr);
diff --git a/Engine/Source/Programs/UnrealBuildAccelerator/Common/Private/UbaProcessStartInfoHolder.cpp b/Engine/Source/Programs/UnrealBuildAccelerator/Common/Private/UbaProcessStartInfoHolder.cpp
--- a/Engine/Source/Programs/UnrealBuildAccelerator/Common/Private/UbaProcessStartInfoHolder.cpp
+++ b/Engine/Source/Programs/UnrealBuildAccelerator/Common/Private/UbaProcessStartInfoHolder.cpp
@@ -61,20 +61,31 @@ namespace uba
 				StringView view(arg, argLen);
 				requiresShell |= argLen == 1 && (*arg == '<' || *arg == '>' || *arg == '|' || *arg == ';');
 				requiresShell |= argLen == 2 && (view.Equals(TCV("&&")) || view.Equals(TCV("||")));
+				// Backslash-escaped glob metacharacters (e.g. \* \? \[) mean the command
+				// relies on the shell to strip the escape. Only these three are needed —
+				// avoiding a broad check that would trigger on \\n in Python -c strings etc.
+				if (!requiresShell)
+					for (u32 i = 0; i + 1 < argLen; ++i)
+						if (arg[i] == TC('\\') && (arg[i+1] == TC('*') || arg[i+1] == TC('?') || arg[i+1] == TC('[')))
+							{ requiresShell = true; break; }
 			});
 		}
 
 		if (requiresShell)
 		{
-			// Build /bin/sh -c "escaped_command" — escape \ and " inside the outer double-quotes
+			// Build /bin/sh -c "escaped_command" — only escape " inside the outer double-quotes.
+			// Backslashes must NOT be doubled: ninja commands already use \\n etc. to represent
+			// shell-level escapes, and the shell halves them when processing the double-quoted
+			// string. Doubling again would corrupt e.g. '\\n' (→ '\n' newline in Python) into
+			// '\\\\n' (→ '\\n' literal backslash+n in Python). Real ninja passes the raw
+			// command as argv[2] directly, achieving the same single-level shell processing.
 			TString fullCmd = applicationStr + TC(" ") + argumentsStr;
 			TString wrapped = TC("-c \"");
 			wrapped.reserve(fullCmd.size() + 8);
 			for (tchar c : fullCmd)
 			{
-				if      (c == TC('\\')) wrapped += TC("\\\\");
-				else if (c == TC('"'))  wrapped += TC("\\\"");
-				else                    wrapped += c;
+				if (c == TC('"')) wrapped += TC("\\\"");
+				else              wrapped += c;
 			}
 			wrapped += TC('"');
 			argumentsStr = std::move(wrapped);
diff --git a/Engine/Source/Programs/UnrealBuildAccelerator/Common/Private/UbaSession.cpp b/Engine/Source/Programs/UnrealBuildAccelerator/Common/Private/UbaSession.cpp
--- a/Engine/Source/Programs/UnrealBuildAccelerator/Common/Private/UbaSession.cpp
+++ b/Engine/Source/Programs/UnrealBuildAccelerator/Common/Private/UbaSession.cpp
@@ -38,6 +38,7 @@ extern char **environ;
 namespace uba
 {
 	bool g_dummy;
+	u32 DefaultDirAttributes();
 
 	////////////////////////////////////////////////////////////////////////////////////////////////////
 
@@ -494,6 +495,8 @@ namespace uba
 		if (shouldWriteToDisk)
 			TraceWrittenFile(process.m_id, file.name, writtenSize);
 
+		RegisterCreateFileForWrite(StringKeyZero, file.name, writtenSize, file.lastWriteTime, file.attributes, false, true);
+
 		#if UBA_DEBUG_TRACK_WRITES
 		allFilesSize += fileSize;
 		m_debugLogger->Info(TC("[%5s %8s] %6s %9s %-13s %s"), TimeToText(start-firstFile), BytesToText(allFilesSize).str, TimeToText(GetTime() - start).str, BytesToText(fileSize).str, writeMode.data, StringView(file.name).GetFileName().data);
@@ -1075,7 +1078,7 @@ namespace uba
 		return WithVfs(rootsHandle, !entry.locals.empty());
 	}
 
-	bool Session::CopyImports(Vector<BinaryModule>& out, const tchar* library, tchar* applicationDir, tchar* applicationDirEnd, UnorderedSet<TString>& handledImports, const tchar* const* loaderPaths, bool isLinuxBinary)
+	bool Session::CopyImports(Vector<BinaryModule>& out, const tchar* library, tchar* applicationDir, tchar* applicationDirEnd, UnorderedSet<TString>& handledImports, const tchar* const* loaderPaths, bool isLinuxBinary, const tchar* hint)
 	{
 		if (!handledImports.insert(library).second)
 			return true;
@@ -1111,13 +1114,15 @@ namespace uba
 			if (attr == INVALID_FILE_ATTRIBUTES)
 				return m_logger.Error("CopyImports - Failed to find file %s (%u loader paths)", applicationName, loaderPathCount);
 			//UBA_ASSERTF(false, TC("DIR NOT FOUND: %s"), applicationName);
-			#else
+#else
 			// SOOO, for linux it might be fine not finding these files.
 			// DT_NEEDED does not mean it is needed(!).. and no way of knowing which ones that are needed
 			#if UBA_DEBUG
-			m_logger.Error("Code path not implemented for linux! (CopyImports %s)", library);
+			m_logger.Error("Code path not implemented for linux! (CopyImports %s, requested by %s)", library, hint);
 			#endif
-			return true;
+			bool makeCompilerHappy = true;
+			if (makeCompilerHappy)
+				return true;
 #endif
 
 			applicationName = temp;
@@ -1155,7 +1160,7 @@ namespace uba
 		auto handleImport = [&](const tchar* importName, bool isKnown, const tchar* const* importLoaderPaths)
 			{
 				if (result && !isKnown)
-					result = CopyImports(out, importName, applicationDir, applicationDirEnd, handledImports, importLoaderPaths, isLinuxBinary);
+					result = CopyImports(out, importName, applicationDir, applicationDirEnd, handledImports, importLoaderPaths, isLinuxBinary, hint);
 			};
 
 		bool parseRes;
@@ -1616,7 +1621,7 @@ namespace uba
 		auto fileNameKey = GetKeyAndFixedName(fixedFilePath, filePath);
 		if (!RegisterVirtualFileInternal(fileNameKey, fixedFilePath, sourceFile, sourceOffset, sourceSize))
 			return false;
-		return RegisterCreateFileForWrite(fileNameKey, fixedFilePath, nullptr, sourceSize, 0, false, false);
+		return RegisterCreateFileForWrite(fileNameKey, fixedFilePath, sourceSize, 0, DefaultAttributes(), false, false);
 	}
 
 	bool Session::CreateVirtualFile(const tchar* filePath, const void* memory, u64 memorySize, bool transient, bool registerToDir)
@@ -1628,7 +1633,7 @@ namespace uba
 			return false;
 		if (!registerToDir)
 			return true;
-		return RegisterCreateFileForWrite(fileNameKey, fixedFilePath, nullptr, memorySize, 0, false, false);
+		return RegisterCreateFileForWrite(fileNameKey, fixedFilePath, memorySize, 0, DefaultAttributes(), false, false);
 	}
 
 	bool Session::DeleteVirtualFile(const tchar* filePath, bool registerToDir)
@@ -1637,7 +1642,7 @@ namespace uba
 		StringBuffer<> fixedFilePath;
 		auto key = GetKeyAndFixedName(fixedFilePath, filePath);
 		if (registerToDir)
-			RegisterDeleteFile(key, fixedFilePath, nullptr);
+			RegisterDeleteFile(key, fixedFilePath);
 
 		m_storage.DeleteExternalFileMapping(key, filePath);
 
@@ -1782,7 +1787,6 @@ namespace uba
 			}
 			exeName.Clear().Append(si.arguments + firstArgumentStart, firstArgumentEnd - firstArgumentStart);
 		}
-		return rules[isDotnet].rules;
 	}
 
 	const tchar* Session::GetTempPath()
@@ -2226,7 +2230,7 @@ namespace uba
 		tchar* applicationDirEnd = applicationDir + applicationDirLen;
 
 		UnorderedSet<TString> handledImports;
-		return CopyImports(out, applicationName, applicationDir, applicationDirEnd, handledImports, nullptr, isLinuxBinary);
+		return CopyImports(out, applicationName, applicationDir, applicationDirEnd, handledImports, nullptr, isLinuxBinary, application);
 	}
 
 	void Session::Free(Vector<BinaryModule>& v)
@@ -2254,6 +2258,11 @@ namespace uba
 #endif
 	}
 
+	bool Session::IsInTemp(StringView path)
+	{
+		return path.StartsWith(m_tempPath);
+	}
+
 	bool Session::ShouldWriteToDisk(StringView fileName)
 	{
 		if (m_shouldWriteToDisk)
@@ -2430,6 +2439,9 @@ namespace uba
 
 		AddEnvironmentVariableNoLock(TCV("UBA_DETOURED"), TCV("1"));
 
+		// io_uring calls not detoured right now
+		AddEnvironmentVariableNoLock(TCV("UV_USE_IO_URING"), TCV("0"));
+
 		m_environmentVariables.push_back(0);
 		return m_environmentVariables.data();
 	}
@@ -2455,18 +2467,26 @@ namespace uba
 		// if ((message.Access & FileAccess.Write) != 0)
 		m_storage.ReportFileWrite(fileNameKey, fileName.data);
 
-		if (m_runningRemote && !fileName.StartsWith(m_tempPath.data))
+		if (m_runningRemote)
 		{
-			SCOPED_FUTEX(m_outputFilesLock, lock);
-			auto insres = m_outputFiles.try_emplace(fileName.data);
-			if (insres.second)
+			if (!fileName.StartsWith(m_tempPath.data))
 			{
-				out.fileName.Append(m_sessionOutputDir).Append(KeyToString(fileNameKey));
-				insres.first->second = out.fileName.data;
+				SCOPED_FUTEX(m_outputFilesLock, lock);
+				auto insres = m_outputFiles.try_emplace(fileName.data);
+				if (insres.second)
+				{
+					out.fileName.Append(m_sessionOutputDir).Append(KeyToString(fileNameKey));
+					insres.first->second = out.fileName.data;
+				}
+				else
+				{
+					out.fileName.Append(insres.first->second.c_str());
+				}
 			}
 			else
 			{
-				out.fileName.Append(insres.first->second.c_str());
+				// We need to create fake names because the file could be under a sub directory under temp and we don't create dirs
+				out.fileName.Append(m_tempPath).Append(KeyToString(fileNameKey));
 			}
 		}
 		else
@@ -2599,8 +2619,9 @@ namespace uba
 		}
 
 		bool registerRealFile = true;
-		u64 fileSize = 0;
-		u64 lastWriteTime = 0;
+		u64 fileSize = msg.fileSize;
+		u64 lastWriteTime = msg.lastWriteTime;
+		u32 attributes = msg.attributes;
 
 		ProcessImpl::Shared& shared = msg.process.m_shared;
 
@@ -2609,10 +2630,6 @@ namespace uba
 
 		if (msg.deleteOnClose)
 		{
-			auto& writtenFiles = shared.writtenFiles;
-			auto findIt2 = writtenFiles.find(activeFile.nameKey);
-			if (findIt2 != writtenFiles.end())
-				writtenFiles.erase(findIt2);
 			RemoveWrittenFile(msg.process, activeFile.nameKey);
 		}
 		else
@@ -2645,9 +2662,17 @@ namespace uba
 					writtenFile.memoryWritten = oldFindIt->second.memoryWritten;
 					writtenFiles.erase(oldFindIt);
 				}
+				else if (writtenFile.memoryHandle.IsValid())
+				{
+					UBA_ASSERT(shared.sharedMemoryHandles.find(writtenFile.memoryHandle) == shared.sharedMemoryHandles.end());
+					UBA_ASSERT(writtenFile.memoryHandle != msg.memoryHandle);
+					m_sharedMemory.CloseHandle(m_logger, writtenFile.memoryHandle, msg.newName.data);
+					writtenFile.memoryHandle = {};
+					writtenFile.originalMemoryHandle = {};
+				}
 			}
 
-			if (m_allowOutputFiles && writtenFile.owner != nullptr && writtenFile.owner != &msg.process)
+			if (m_allowOutputFiles && writtenFile.ownerId && writtenFile.ownerId != msg.process.m_id)
 			{
 				// This can happen when library has /GL (whole program optimization) but target has not.. then link.exe will restart
 				//UBA_ASSERTF(false, TC("File %s changed owner.. should not happen (OldOwner: %s New owner: %s)"), name, writtenFile.owner->m_realApplication.c_str(), msg.process.m_realApplication.c_str());
@@ -2661,7 +2686,7 @@ namespace uba
 				writtenFile.name = name.ToString();
 				writtenFile.key = key;
 				writtenFile.backedName = msgName.ToString();
-				writtenFile.owner = &msg.process;
+				writtenFile.ownerId = msg.process.m_id;
 			}
 			else
 			{
@@ -2671,11 +2696,11 @@ namespace uba
 					writtenFile.backedName = msgName.ToString();
 				}
 
-				if (!msg.memoryHandle.IsValid() || (msg.memoryHandle == writtenFile.originalMemoryHandle && writtenFile.owner == &msg.process))
+				if (!msg.memoryHandle.IsValid() || (msg.memoryHandle == writtenFile.originalMemoryHandle && writtenFile.ownerId == msg.process.m_id))
 				{
-					if (msg.memoryWritten)
+					if (msg.fileSize)
 					{
-						writtenFile.memoryWritten = msg.memoryWritten;
+						writtenFile.memoryWritten = msg.fileSize;
 						writtenFile.lastWriteTime = GetSystemTimeAsFileTime();
 					}
 					addMapping = false;
@@ -2692,7 +2717,7 @@ namespace uba
 					writtenFile.originalMemoryHandle = {};
 				}
 
-				writtenFile.owner = &msg.process;
+				writtenFile.ownerId = msg.process.m_id;
 			}
 
 			if (!m_runningRemote && HasVfs(startInfo.rootsHandle)) // For posix we write the dependency file directly to disk so we need to update it if vfs is enabled
@@ -2765,7 +2790,7 @@ namespace uba
 				shared.sharedMemoryHandles.erase(msg.memoryHandle); // memory handle transferred to written files
 				writtenFile.memoryHandle = msg.memoryHandle;
 				writtenFile.originalMemoryHandle = msg.memoryHandle;
-				writtenFile.memoryWritten = msg.memoryWritten;
+				writtenFile.memoryWritten = msg.fileSize;
 				writtenFile.lastWriteTime = GetSystemTimeAsFileTime();
 
 				#if UBA_DEBUG_TRACK_MAPPING
@@ -2785,36 +2810,21 @@ namespace uba
 					return false;
 		}
 
-		//
-		// This entire code below here needs to be revisited
-		// We want to only use overlay for everything here.
-		// And then transfer info back to the main directory table once things are done
-		//
-
-		if (m_runningRemote && registerRealFile && msg.fileName.count && !msg.deleteOnClose)
-		{
-			FileBasicInformation info;
-			if (GetFileBasicInformation(info, m_logger, msg.fileName.data))
-			{
-				lastWriteTime = info.lastWriteTime;
-				fileSize = info.size;
-			}
-		}
-
 		if (!rules.IsInvisible(activeFile.name))
 		{
+			UBA_ASSERT(fileSize != InvalidValue);
 			if (!msg.newName.IsEmpty())
 			{
-				RegisterDeleteFile(activeFile.nameKey, activeFile.name, &shared.overlay);
-				if (RegisterCreateFileForWrite(msg.newNameKey, msg.newName, &shared.overlay, fileSize, lastWriteTime, registerRealFile, true) && registerRealFile)
-					TraceWrittenFile(msg.process.m_id, msg.newName, fileSize);
+				RegisterOverlayDelete(shared.overlay, activeFile.nameKey, activeFile.name);
+				RegisterOverlayCreateOrWrite(shared.overlay, msg.newNameKey, msg.newName, fileSize, lastWriteTime, attributes);
+				//TraceWrittenFile(msg.process.m_id, msg.newName, fileSize);
 			}
 			else if (msg.deleteOnClose)
-				RegisterDeleteFile(activeFile.nameKey, activeFile.name, &shared.overlay);
+				RegisterOverlayDelete(shared.overlay, activeFile.nameKey, activeFile.name);
 			else
 			{
-				if (RegisterCreateFileForWrite(activeFile.nameKey, activeFile.name, &shared.overlay, fileSize, lastWriteTime, registerRealFile, true) && registerRealFile)
-					TraceWrittenFile(msg.process.m_id, activeFile.name, fileSize);
+				RegisterOverlayCreateOrWrite(shared.overlay, activeFile.nameKey, activeFile.name, fileSize, lastWriteTime, attributes);
+				//TraceWrittenFile(msg.process.m_id, activeFile.name, fileSize);
 			}
 		}
 
@@ -2861,7 +2871,7 @@ namespace uba
 			out.result = uba::DeleteFileW(msg.fileName.data);
 		out.errorCode = GetLastError();
 
-		RegisterDeleteFile(msg.fileNameKey, msg.fileName, &shared.overlay);
+		RegisterOverlayDelete(shared.overlay, msg.fileNameKey, msg.fileName);
 
 		out.directoryTableSize = GetDirectoryTableSize(&msg.process.m_shared.overlay);
 		return true;
@@ -2888,6 +2898,16 @@ namespace uba
 		auto& process = msg.process;
 		bool isMoved = false;
 		auto& shared = process.m_shared;
+
+		auto updateDirTable = [&]()
+			{
+				RegisterOverlayCreateOrWrite(shared.overlay, msg.toKey, msg.toName, msg.fileSize, msg.lastWriteTime, msg.fileAttributes);
+				//if (RegisterCreateFileForWrite(msg.toKey, msg.toName, &shared.overlay, msg.fileSize, msg.lastWriteTime, msg.fileAttributes, true, true))
+				//	TraceWrittenFile(process.m_id, msg.toName, 0);
+				RegisterOverlayDelete(shared.overlay, msg.fromKey, msg.fromName);
+			};
+
+
 		{
 			SCOPED_FUTEX(shared.writtenFilesLock, writtenLock);
 			auto findIt = shared.writtenFiles.find(msg.fromKey);
@@ -2925,10 +2945,8 @@ namespace uba
 				{
 					out.errorCode = ERROR_SUCCESS;
 					out.result = true;
-					
-					if (!MoveOverlayFile(msg.process.m_shared.overlay, msg.fromName, msg.toName))
-						return false;
-					out.directoryTableSize = GetDirectoryTableSize(&msg.process.m_shared.overlay);
+					updateDirTable();
+					out.directoryTableSize = GetDirectoryTableSize(&shared.overlay);
 					return true;
 				}
 			}
@@ -2944,17 +2962,11 @@ namespace uba
 			}
 		}
 
-		if (!process.GetStartInfo().rules->IsInvisible(msg.toName))
-			if (RegisterCreateFileForWrite(msg.toKey, msg.toName, nullptr, 0, 0, true, true))
-				TraceWrittenFile(process.m_id, msg.toName, 0);
-
-		if (!AddOverlayFile(msg.process.m_shared.overlay, msg.toName))
-			return false;
+		updateDirTable();
 
 		out.errorCode = ERROR_SUCCESS;
-		RegisterDeleteFile(msg.fromKey, msg.fromName, &shared.overlay);
 
-		out.directoryTableSize = GetDirectoryTableSize(&msg.process.m_shared.overlay);
+		out.directoryTableSize = GetDirectoryTableSize(&shared.overlay);
 		return true;
 	}
 
@@ -2966,7 +2978,10 @@ namespace uba
 		out.errorCode = 0;
 		if (chmod(msg.fileName.data, (mode_t)msg.fileMode) == 0)
 		{
-			RegisterCreateFileForWrite(msg.fileNameKey, msg.fileName, nullptr, 0, 0, true, true);
+			auto& shared = msg.process.m_shared;
+			// Should use overlay
+			// RegisterOverlayCreateOrWrite(shared.overlay, msg.fileNameKey, msg.fileName, f.memoryWritten, f.lastWriteTime, f.attributes);
+			RegisterCreateFileForWrite(msg.fileNameKey, msg.fileName, 0, 0, 0, true, true);
 			return true;
 		}
 		out.errorCode = errno;
@@ -2974,42 +2989,79 @@ namespace uba
 		return true;
 	}
 
-	bool Session::CreateDirectory(CreateDirectoryResponse& out, const CreateDirectoryMessage& msg, DirectoryTableOverlay* overlay)
+	bool Session::Symlink(SymlinkResponse& out, const SymlinkMessage& msg)
 	{
-		out.result = uba::CreateDirectoryW(msg.name.data);
+		#if PLATFORM_WINDOWS
+		UBA_ASSERT(false); // This is not used
+		#else
+		out.errorCode = 0;
+		if (!m_runningRemote)
+		{
+			out.success = symlink(msg.symlinkContent.data, msg.symlinkFilePath.data) == 0;
+			if (!out.success)
+			{
+				out.errorCode = errno;
+				return true;
+			}
+		}
+		else
+		{
+			// Should probably create output file, point it to the other file with a real symlink
+			// Let's see how far we get by not doing anything except updating the dir table
+			out.success = true;
+		}
 
-		if (!out.result)
-			out.errorCode = GetLastError();
+		auto& overlay = msg.process.m_shared.overlay;
 
-		StringKey dirKey;
-		const tchar* lastSlash;
-		StringBuffer<> dirName;
-		if (!GetDirKey(dirKey, dirName, lastSlash, msg.name))
-			return true;
+		// This is very wrong, but we don't want to add support for symlinks to DirectoryTable system unless we really have to
+		if (!RegisterOverlayCreateOrWrite(overlay, msg.contentPathKey, msg.symlinkContentPath, msg.fileSize, msg.lastWriteTime, msg.fileAttributes))
+			return false;
 
-		// There is a chance that another thread just created the directory and we can't return directoryTableSize until we know it is written
-		// So let's both success and already exists add entries
-		if (out.result || out.errorCode == ERROR_ALREADY_EXISTS)
-		{
-			// Both these functions need to be called. otherwise we can get created directories that does not end up in directory table
-			RegisterCreateFileForWrite(msg.nameKey, msg.name, nullptr, 0, 0, true, true);
-			WriteDirectoryEntries(dirKey, dirName);
-		}
+		overlay.symlinks.try_emplace(msg.symlinkFilePath.ToString(), msg.symlinkContent.ToString());
 
-		out.directoryTableSize = GetDirectoryTableSize(overlay);
+		out.directoryTableSize = GetDirectoryTableSize(&overlay);
+
+		#endif
 		return true;
 	}
 
-	bool Session::RemoveDirectory(RemoveDirectoryResponse& out, const RemoveDirectoryMessage& msg, DirectoryTableOverlay* overlay)
+	bool Session::CreateDirectory(CreateDirectoryResponse& out, const CreateDirectoryMessage& msg)
 	{
-		out.result = uba::RemoveDirectoryW(msg.name.data);
-		if (out.result)
-			RegisterDeleteFile(msg.nameKey, msg.name, overlay);
-		else
-			out.errorCode = GetLastError();
-		// This has a race condition. If same directory is removed at the same time the failing one
-		// might send back a directoryTableSize that does not include the delete
-		out.directoryTableSize = GetDirectoryTableSize(overlay);
+		auto& overlay = msg.process.m_shared.overlay;
+		
+		RegisterOverlayCreateOrWrite(overlay, msg.nameKey, msg.name, 0, 0, DefaultDirAttributes());
+
+		SCOPED_WRITE_LOCK(overlay.lock, lock);
+		TString dirPath = msg.name.ToString();
+		if (!overlay.removedDirectories.erase(dirPath))
+			overlay.createdDirectories.insert(dirPath);
+		lock.Leave();
+
+		out.directoryTableSize = GetDirectoryTableSize(&overlay);
+
+		// TODO: Remove?
+		if (!m_runningRemote)
+			uba::CreateDirectoryW(msg.name.data);
+		return true;
+	}
+
+	bool Session::RemoveDirectory(RemoveDirectoryResponse& out, const RemoveDirectoryMessage& msg)
+	{
+		// TODO: Remove?
+		if (!m_runningRemote)
+			uba::RemoveDirectoryW(msg.name.data);
+
+		auto& overlay = msg.process.m_shared.overlay;
+
+		RegisterOverlayDelete(overlay, msg.nameKey, msg.name);
+
+		SCOPED_WRITE_LOCK(overlay.lock, lock);
+		TString dirPath = msg.name.ToString();
+		if (!overlay.createdDirectories.erase(dirPath))
+			overlay.removedDirectories.insert(dirPath);
+		lock.Leave();
+
+		out.directoryTableSize = GetDirectoryTableSize(&overlay);
 		return true;
 	}
 
@@ -3034,6 +3086,12 @@ namespace uba
 		if (overlay)
 		{
 			SCOPED_READ_LOCK(overlay->lock, lock);
+			if (tableOffset == InvalidTableOffset)
+			{
+				auto findIt = overlay->lookup.find(dirKey);
+				if (findIt != overlay->lookup.end())
+					tableOffset = findIt->second.offset | OverlayTableFlag;
+			}
 			overlaySize = overlay->size;
 		}
 
@@ -3534,7 +3592,7 @@ namespace uba
 			m_sharedMemory.UnmapView(symHandle, mem2);
 
 			lastWriteTime = GetSystemTimeAsFileTime();
-			if (!RegisterCreateFileForWrite(symFileKey, exportsFile, nullptr, memoryBlock.writtenSize, lastWriteTime, false, true))
+			if (!RegisterCreateFileForWrite(symFileKey, exportsFile, memoryBlock.writtenSize, lastWriteTime, DefaultAttributes(), false, true))
 				return false;
 			mg.Cancel();
 		}
@@ -3547,16 +3605,16 @@ namespace uba
 				return false;
 			if (!destFile.Close(&lastWriteTime))
 				return false;
-			if (!RegisterCreateFileForWrite(symFileKey, exportsFile, nullptr, memoryBlock.writtenSize, lastWriteTime, true, true))
+			if (!RegisterCreateFileForWrite(symFileKey, exportsFile, memoryBlock.writtenSize, lastWriteTime, DefaultAttributes(), true, true))
 				return false;
 		}
 
 		auto insres = msg.process.m_shared.writtenFiles.try_emplace(symFileKey);
 		WrittenFile& writtenFile = insres.first->second;
 
-		UBA_ASSERT(writtenFile.owner == nullptr || writtenFile.owner == &msg.process);
+		UBA_ASSERT(!writtenFile.ownerId || writtenFile.ownerId == msg.process.m_id);
 		writtenFile.key = symFileKey;
-		writtenFile.owner = &msg.process;
+		writtenFile.ownerId = msg.process.m_id;
 		writtenFile.attributes = msg.attributes;
 		writtenFile.memoryHandle = symHandle;
 		writtenFile.memoryWritten = memoryBlock.writtenSize;
diff --git a/Engine/Source/Programs/UnrealBuildAccelerator/Common/Private/UbaSessionClient.cpp b/Engine/Source/Programs/UnrealBuildAccelerator/Common/Private/UbaSessionClient.cpp
--- a/Engine/Source/Programs/UnrealBuildAccelerator/Common/Private/UbaSessionClient.cpp
+++ b/Engine/Source/Programs/UnrealBuildAccelerator/Common/Private/UbaSessionClient.cpp
@@ -553,21 +553,6 @@ namespace uba
 		if (!Session::ProcessThreadStart(process))
 			return false;
 
-		#if PLATFORM_LINUX
-		// TODO. Linux linker has lots of stuff going on. it is a shell script deleting, renaming, overwriting same file causing confusion with child processes
-		// Since directory table is not updated from server between child processes we end up with not up-to-date table.. therefore we flush everything in between
-		if (auto parent = process.m_parentProcess)
-		{
-			auto& si = parent->m_startInfo;
-			if (ToView(si.application).EndsWith(TCV("/sh")))
-			{
-				FlushWrittenFiles(*parent);
-				StackBinaryReader<SendMaxSize> reader;
-				SendUpdateDirectoryTable(reader);
-			}
-		}
-		#endif
-
 		if (RootsHandle rootsHandle = process.GetStartInfo().rootsHandle)
 			if (!SendRootsHandle(rootsHandle))
 				return false;
@@ -736,19 +721,79 @@ namespace uba
 		}
 	}
 
+	bool SessionClient::SendFileSystemOperations(ProcessImpl& process)
+	{
+		auto& rules = *process.m_startInfo.rules;
+		auto& overlay = process.m_shared.overlay;
+
+		auto shouldSend = [&](StringView path) { return !path.StartsWith(m_tempPath) || rules.KeepTempOutputFile(StringView(path)); };
+
+		StackBinaryWriter<1024> writer;
+
+		SCOPED_READ_LOCK(overlay.lock, lock);
+
+		UBA_ASSERTF(overlay.removedDirectories.empty(), TC("Removing directories from remote not implemented."));
+
+		for (auto& createdDir : overlay.createdDirectories)
+		{
+			if (!shouldSend(createdDir))
+				continue;
+			NetworkMessage msg(m_client, ServiceId, SessionMessageType_CreateDirectory, writer.Reset());
+			writer.WriteU32(process.m_id);
+			writer.WriteString(createdDir);
+			StackBinaryReader<128> reader;
+			if (!msg.Send(reader, Stats().createDirMsg))
+				return false;
+			if (!reader.ReadBool())
+				return m_logger.Warning(TC("Failed to create %s on server side"), createdDir.c_str());
+		}
+
+		for (auto& deletedFile : overlay.deletedFiles)
+		{
+			if (!shouldSend(deletedFile))
+				continue;
+			NetworkMessage networkMsg(m_client, ServiceId, SessionMessageType_DeleteFile, writer.Reset());
+			writer.WriteU32(process.m_id);
+			writer.WriteString(deletedFile);
+			StackBinaryReader<SendMaxSize> reader;
+			if (!networkMsg.Send(reader, Stats().deleteFileMsg))
+				return false;
+			if (!reader.ReadBool())
+				return m_logger.Warning(TC("Failed to delete %s on server side"), deletedFile.c_str());
+		}
+
+		for (auto& symlink : overlay.symlinks)
+		{
+			NetworkMessage networkMsg(m_client, ServiceId, SessionMessageType_Symlink, writer.Reset());
+			writer.WriteU32(process.m_id);
+			writer.WriteString(symlink.first); // File
+			writer.WriteString(symlink.second); // Content
+			StackBinaryReader<SendMaxSize> reader;
+			if (!networkMsg.Send(reader, Stats().copyFileMsg))
+				return false;
+			if (!reader.ReadBool())
+				return m_logger.Warning(TC("Failed to create symlink %s->%s on server side"), symlink.first.c_str(), symlink.second.c_str());
+		}
+
+		return true;
+	}
+
 	bool SessionClient::SendFiles(ProcessImpl& process)
 	{
 		Atomic<bool> success = true;
 		auto& rules = *process.m_startInfo.rules;
-		auto& writtenFiles = process.m_shared.writtenFiles;
+		auto& shared = process.m_shared;
+		auto& writtenFiles = shared.writtenFiles;
 		TimerScope timer(process.m_processStats.sendFiles);
 
+		auto shouldSend = [&](StringView path) { return !path.StartsWith(m_tempPath) || rules.KeepTempOutputFile(StringView(path)); };
+
 		m_client.ParallelFor(u32(writtenFiles.size()), writtenFiles, [&](const WorkContext&, auto& it)
 			{
 				StorageStatsScope storageStatsScope(process.m_storageStats);
 				WrittenFile& file = it->second;
 
-				if (StringView(file.name).StartsWith(m_tempPath) && !rules.KeepTempOutputFile(StringView(file.name)))
+				if (!shouldSend(file.name))
 				{
 					m_sharedMemory.CloseHandle(m_logger, file.memoryHandle, file.backedName.c_str());
 					file.memoryHandle = {};
@@ -829,8 +874,14 @@ namespace uba
 		if (isTemp)
 			sendDelete = false;
 
+		auto& overlay = msg.process.m_shared.overlay;
+
+		RegisterOverlayDelete(overlay, msg.fileNameKey, msg.fileName);
+
 		if (!sendDelete)
 		{
+			out.directoryTableSize = GetDirectoryTableSize(&overlay);
+
 			if (isTemp && (!m_allowMemoryMaps || !msg.process.m_startInfo.rules->KeepInMemory(msg.fileName, m_tempPath, true, true)))
 			{
 				out.result = uba::DeleteFileW(msg.fileName.data);
@@ -842,21 +893,10 @@ namespace uba
 			return true;
 		}
 
-		// TODO: Cache this if it becomes noisy
-
-		StackBinaryWriter<1024> writer;
-		NetworkMessage networkMsg(m_client, ServiceId, SessionMessageType_DeleteFile, writer);
-		writer.WriteStringKey(msg.fileNameKey);
-		writer.WriteString(msg.fileName);
-		StackBinaryReader<SendMaxSize> reader;
-		if (!networkMsg.Send(reader, Stats().deleteFileMsg))
-			return false;
-		out.result = reader.ReadBool();
-		out.errorCode = reader.ReadU32();
-		if (out.result)
-			if (!SendUpdateDirectoryTable(reader.Reset()))
-				return false;
-		out.directoryTableSize = GetDirectoryTableSize(&msg.process.m_shared.overlay);
+		overlay.deletedFiles.insert(msg.fileName.ToString());
+		out.result = true;
+		out.errorCode = ERROR_SUCCESS;
+		out.directoryTableSize = GetDirectoryTableSize(&overlay);
 		return true;
 	}
 
@@ -929,7 +969,7 @@ namespace uba
 				toFile = findIt->second; // Copy all state first and then fixup
 				toFile.key = msg.toKey;
 				toFile.name = toName;
-				toFile.owner = &process;
+				toFile.ownerId = process.m_id;
 				writtenFiles.erase(findIt);
 			}
 			else if (!msg.fromName.StartsWith(m_tempPath) || !msg.toName.StartsWith(m_tempPath))
@@ -939,7 +979,9 @@ namespace uba
 			}
 		}
 
-		MoveOverlayFile(shared.overlay, msg.fromName, msg.toName);
+		RegisterOverlayCreateOrWrite(shared.overlay, msg.toKey, msg.toName, msg.fileSize, msg.lastWriteTime, msg.fileAttributes);
+		RegisterOverlayDelete(shared.overlay, msg.fromKey, msg.fromName);
+
 		out.directoryTableSize = GetDirectoryTableSize(&shared.overlay);
 
 		bool sendMove = true;
@@ -973,8 +1015,9 @@ namespace uba
 	bool SessionClient::Chmod(ChmodResponse& out, const ChmodMessage& msg)
 	{
 		{
-			SCOPED_FUTEX(msg.process.m_shared.writtenFilesLock, lock);
-			auto& writtenFiles = msg.process.m_shared.writtenFiles;
+			auto& shared = msg.process.m_shared;
+			SCOPED_FUTEX(shared.writtenFilesLock, lock);
+			auto& writtenFiles = shared.writtenFiles;
 			auto findIt = writtenFiles.find(msg.fileNameKey);
 			if (findIt != writtenFiles.end())
 			{
@@ -983,8 +1026,12 @@ namespace uba
 				if (msg.fileMode & S_IXUSR)
 					executable = true;
 				#endif
-				findIt->second.attributes = DefaultAttributes(executable);
+				auto& f = findIt->second;
+				f.attributes = DefaultAttributes(executable);
 				out.errorCode = 0;
+
+				RegisterOverlayCreateOrWrite(shared.overlay, msg.fileNameKey, msg.fileName, f.memoryWritten, f.lastWriteTime, f.attributes);
+				//out.directoryTableSize = GetDirectoryTableSize(&shared.overlay);
 				return true;
 			}
 		}
@@ -993,47 +1040,6 @@ namespace uba
 		return true;
 	}
 
-	bool SessionClient::CreateDirectory(CreateDirectoryResponse& out, const CreateDirectoryMessage& msg, DirectoryTableOverlay* overlay)
-	{
-		// TODO: Cache this if it becomes noisy
-
-		StackBinaryWriter<1024> writer;
-		NetworkMessage networkMsg(m_client, ServiceId, SessionMessageType_CreateDirectory, writer);
-		writer.WriteString(msg.name);
-		StackBinaryReader<SendMaxSize> reader;
-		if (!networkMsg.Send(reader, Stats().createDirMsg))
-			return false;
-		out.result = reader.ReadBool();
-		out.errorCode = reader.ReadU32();
-
-		if (out.result || out.errorCode == ERROR_ALREADY_EXISTS) // Even if it exists it might be that this client session is not aware
-			if (!SendUpdateDirectoryTable(reader.Reset()))
-				return false;
-
-		out.directoryTableSize = GetDirectoryTableSize(overlay);
-		return true;
-	}
-
-	bool SessionClient::RemoveDirectory(RemoveDirectoryResponse& out, const RemoveDirectoryMessage& msg, DirectoryTableOverlay* overlay)
-	{
-		// TODO: Should update overlay directory
-		StackBinaryWriter<1024> writer;
-		NetworkMessage networkMsg(m_client, ServiceId, SessionMessageType_RemoveDirectory, writer);
-		writer.WriteString(msg.name);
-		StackBinaryReader<SendMaxSize> reader;
-		if (!networkMsg.Send(reader, Stats().deleteFileMsg)) // Wrong message
-			return false;
-		out.result = reader.ReadBool();
-		out.errorCode = reader.ReadU32();
-
-		if (out.result)
-			if (!SendUpdateDirectoryTable(reader.Reset()))
-				return false;
-
-		out.directoryTableSize = GetDirectoryTableSize(overlay);
-		return true;
-	}
-
 	bool SessionClient::GetFullFileName(GetFullFileNameResponse& out, const GetFullFileNameMessage& msg)
 	{
 		StringView workingDir(msg.process.m_startInfo.workingDirStr);
@@ -1627,7 +1633,9 @@ namespace uba
 
 	bool SessionClient::SendProcessInputs(ProcessImpl& process)
 	{
-		auto inputs = process.GetTrackedInputs();
+		auto& inputs = process.GetTrackedInputs();
+		if (inputs.empty())
+			return true;
 		StackBinaryWriter<SendMaxSize> writer;
 		NetworkMessage msg(m_client, ServiceId, SessionMessageType_ProcessInputs, writer);
 		writer.Write7BitEncoded(process.m_id);
@@ -2135,7 +2143,7 @@ namespace uba
 					startInfo.rules = GetRules(startInfo);
 
 					StringBuffer<> logFile;
-					if (m_logToFile || (*startInfo.logFile && m_shouldSendLogToServer))
+					if (m_logToFile || (!*startInfo.logFile && m_shouldSendLogToServer))
 					{
 						GetLogFileName(logFile, startInfo.logFile, startInfo.arguments, startInfo.processId);
 						startInfo.logFile = logFile.data;
@@ -2227,8 +2235,9 @@ namespace uba
 						{
 							// Should we decrease weight before or after sending files?
 							//decreaseWeight.Execute();
-
-							if (!session.SendFiles(process))
+							bool success = session.SendFileSystemOperations(process);
+							success = success && session.SendFiles(process);
+							if (!success)
 							{
 								const tchar* desc = TC("Failed to send output files to host");
 								session.m_logger.Error(desc);
@@ -2882,8 +2891,9 @@ namespace uba
 
 	bool SessionClient::FlushWrittenFiles(ProcessImpl& process)
 	{
+		bool success = SendFileSystemOperations(process);
 		SCOPED_FUTEX(process.m_shared.writtenFilesLock, lock);
-		bool success = SendFiles(process);
+		success = success && SendFiles(process);
 		{
 			SCOPED_FUTEX(m_outputFilesLock, lock2);
 			for (auto& kv : process.m_shared.writtenFiles)
diff --git a/Engine/Source/Programs/UnrealBuildAccelerator/Common/Private/UbaSessionServer.cpp b/Engine/Source/Programs/UnrealBuildAccelerator/Common/Private/UbaSessionServer.cpp
--- a/Engine/Source/Programs/UnrealBuildAccelerator/Common/Private/UbaSessionServer.cpp
+++ b/Engine/Source/Programs/UnrealBuildAccelerator/Common/Private/UbaSessionServer.cpp
@@ -1492,7 +1492,7 @@ namespace uba
 		if (!rules.IsInvisible(destination))
 		{
 			bool invalidateStorage = false; // No need, already handled in m_storage.CopyOrLink
-			RegisterCreateFileForWrite(StringKeyZero, destination, nullptr, 0, 0, shouldWriteToDisk, invalidateStorage);
+			RegisterCreateFileForWrite(StringKeyZero, destination, 0, 0, file.attributes, shouldWriteToDisk, invalidateStorage);
 
 			if (startInfo.trackInputs)
 			{
@@ -1587,16 +1587,18 @@ namespace uba
 
 	bool SessionServer::HandleDeleteFile(const ConnectionInfo& connectionInfo, const WorkContext& workContext, BinaryReader& reader, BinaryWriter& writer)
 	{
-		// TODO: This should be removed and replaced with overlay
-		StringKey fileNameKey = reader.ReadStringKey();
+		u32 processId = reader.ReadU32();(void)processId;
 		StringBuffer<> fileName;
 		reader.ReadString(fileName);
-		bool result = uba::DeleteFileW(fileName.data);
-		u32 errorCode = GetLastError();
-		if (result)
-			RegisterDeleteFile(fileNameKey, fileName, nullptr);
-		writer.WriteBool(result);
-		writer.WriteU32(errorCode);
+		bool res = uba::DeleteFileW(fileName.data);
+		if (res)
+		{
+			StringKey fileNameKey = CaseInsensitiveFs ? ToStringKeyLower(fileName) : ToStringKey(fileName);
+			RegisterDeleteFile(fileNameKey, fileName);
+		}
+		else if (GetLastError() == ERROR_FILE_NOT_FOUND)
+			res = true;
+		writer.WriteBool(res);
 		return true;
 	}
 
@@ -1612,27 +1614,48 @@ namespace uba
 		bool result = uba::CopyFileW(fromName.data, toName.data, false);
 		u32 errorCode = GetLastError();
 		if (result)
-			RegisterCreateFileForWrite(toNameKey, toName, nullptr, 0, 0, true, true);
+			RegisterCreateFileForWrite(toNameKey, toName, 0, 0, 0, true, true);
 		writer.WriteU32(errorCode);
 		return true;
 	}
 
+	bool SessionServer::HandleSymlink(const ConnectionInfo& connectionInfo, const WorkContext& workContext, BinaryReader& reader, BinaryWriter& writer)
+	{
+		#if PLATFORM_WINDOWS
+		return false;
+		#else
+		u32 processId = reader.ReadU32();(void)processId;
+		StringBuffer<> symlinkFile;
+		StringBuffer<> symlinkContent;
+		reader.ReadString(symlinkFile);
+		reader.ReadString(symlinkContent);
+		bool res = symlink(symlinkContent.data, symlinkFile.data) == 0;
+		writer.WriteBool(res);
+		return true;
+		#endif
+	}
+
 	bool SessionServer::HandleCreateDirectory(const ConnectionInfo& connectionInfo, const WorkContext& workContext, BinaryReader& reader, BinaryWriter& writer)
 	{
-		// TODO: This should be removed and replaced with overlay
-		CreateDirectoryMessage msg;
-		reader.ReadString(msg.name);
-		CreateDirectoryResponse response;
-		if (!Session::CreateDirectory(response, msg, nullptr))
-			return false;
-		writer.WriteBool(response.result);
-		writer.WriteU32(response.errorCode);
+		u32 processId = reader.ReadU32();(void)processId;
+		StringBuffer<> createdDir;
+		reader.ReadString(createdDir);
+		bool res = uba::CreateDirectoryW(createdDir.data);
+		if (res)
+		{
+			StringKey createdDirKey = CaseInsensitiveFs ? ToStringKeyLower(createdDir) : ToStringKey(createdDir);
+			RegisterCreateFileForWrite(createdDirKey, createdDir, 0, 0, 0, true, true);
+		}
+		else if (GetLastError() == ERROR_ALREADY_EXISTS)
+			res = true;
+		writer.WriteBool(res);
 		return true;
 	}
 
 	bool SessionServer::HandleRemoveDirectory(const ConnectionInfo& connectionInfo, const WorkContext& workContext, BinaryReader& reader, BinaryWriter& writer)
 	{
-		// TODO: This should be removed and replaced with overlay
+		UBA_ASSERT(false);
+		#if 0
 		RemoveDirectoryMessage msg;
 		reader.ReadString(msg.name);
 		RemoveDirectoryResponse response;
@@ -1640,6 +1663,7 @@ namespace uba
 			return false;
 		writer.WriteBool(response.result);
 		writer.WriteU32(response.errorCode);
+		#endif
 		return true;
 	}
 
diff --git a/Engine/Source/Programs/UnrealBuildAccelerator/Common/Public/UbaDirectoryTableHolder.h b/Engine/Source/Programs/UnrealBuildAccelerator/Common/Public/UbaDirectoryTableHolder.h
--- a/Engine/Source/Programs/UnrealBuildAccelerator/Common/Public/UbaDirectoryTableHolder.h
+++ b/Engine/Source/Programs/UnrealBuildAccelerator/Common/Public/UbaDirectoryTableHolder.h
@@ -4,7 +4,7 @@
 
 #include "UbaDirectoryTable.h"
 #include "UbaFile.h"
-#include "UbaSharedMemoryView.h"
+#include "UbaFileMappingHandle.h"
 
 namespace uba
 {
@@ -22,6 +22,10 @@ namespace uba
 		struct Dir { u32 offset = 0; UnorderedMap<StringKey, u32> files; };
 		UnorderedMap<StringKey, Dir> lookup;
 		Vector<StringKey> roots;
+		Set<TString> createdDirectories;
+		Set<TString> removedDirectories;
+		Set<TString> deletedFiles;
+		Map<TString, TString> symlinks;
 	};
 
 	////////////////////////////////////////////////////////////////////////////////////////////////////
@@ -53,22 +57,26 @@ namespace uba
 
 	//protected:
 
+		bool RegisterOverlayCreateOrWrite(DirectoryTableOverlay& overlay, StringKey filePathKey, StringView filePath, u64 fileSize, u64 lastWriteTime, u32 attributes);
+		bool RegisterOverlayDelete(DirectoryTableOverlay& overlay, StringKey filePathKey, StringView filePath);
+
 		bool PopulateOverlayFile(DirectoryTableOverlay& overlay, StringView filePath, const Function<bool(FileInformation& out, StringKey)>& getFileInformation);
 		bool AddOverlayFile(DirectoryTableOverlay& overlay, StringView filePath);
-		bool MoveOverlayFile(DirectoryTableOverlay& overlay, StringView from, StringView to);
 		bool WriteOverlayEntry(DirectoryTableOverlay& overlay, StringView filePath, const FileInformation& fileInfo);
 		bool WriteOverlayEntry(DirectoryTableOverlay& overlay, StringKey dirKey, StringView filePath, StringView fileName, const FileInformation& fileInfo, bool populateMainTableParentDir);
 		bool WriteOverlayEntry(DirectoryTableOverlay& overlay, StringKey dirKey, DirectoryTableOverlay::Dir& dir, StringView fileName, const FileInformation& fileInfo, u32& outFileOffset);
+		bool WriteOverlayDirInLookup(DirectoryTableOverlay& overlay, StringView tempDir);
 		bool WriteOverlayDirectoryEntriesRecursive(DirectoryTableOverlay& overlay, StringKey filePathKey, StringView filePath);
 
-		bool RegisterCreateFileForWrite(StringKey filePathKey, StringView filePath, DirectoryTableOverlay* overlay, u64 fileSize, u64 lastWriteTime, bool registerRealFile, bool invalidateStorage);
-		bool RegisterDeleteFile(StringKey filePathKey, StringView filePath, DirectoryTableOverlay* overlay);
+		bool RegisterCreateFileForWrite(StringKey filePathKey, StringView filePath, u64 fileSize, u64 lastWriteTime, u32 attributes, bool registerRealFile, bool invalidateStorage);
+		bool RegisterDeleteFile(StringKey filePathKey, StringView filePath);
 
 		u32 WriteDirectoryEntries(const StringKey& dirKey, StringView dirPath, u32* outTableOffset = nullptr);
 		void WriteDirectoryEntriesRecursive(const StringKey& dirKey, StringView dirPath, u32& outTableOffset);
 		bool WriteDirectoryEntriesInternal(DirectoryTable::Directory& dir, const StringKey& dirKey, StringView dirPath, bool isRefresh, u32& outTableOffset);
 		void EnsureDirectoryTableMemory(u64 neededSize);
 
+		virtual bool IsInTemp(StringView path);
 		virtual void FileEntryAdded(StringKey filePathKey, u64 lastWritten, u64 size);
 		virtual bool ShouldWriteToDisk(StringView filePath);
 		virtual void InvalidateCachedFileInfo(const StringKey& filePathKey);
diff --git a/Engine/Source/Programs/UnrealBuildAccelerator/Common/Public/UbaNetwork.h b/Engine/Source/Programs/UnrealBuildAccelerator/Common/Public/UbaNetwork.h
--- a/Engine/Source/Programs/UnrealBuildAccelerator/Common/Public/UbaNetwork.h
+++ b/Engine/Source/Programs/UnrealBuildAccelerator/Common/Public/UbaNetwork.h
@@ -16,7 +16,7 @@ namespace uba
 	
 	static constexpr u32 SystemNetworkVersion = 1339;
 	static constexpr u32 StorageNetworkVersion = 4;
-	static constexpr u32 SessionNetworkVersion = 51;
+	static constexpr u32 SessionNetworkVersion = 52;
 	static constexpr u32 CacheNetworkVersion = 5;
 	static constexpr u32 CacheBucketVersion = 3;
 
@@ -78,6 +78,7 @@ namespace uba
 		UBA_SESSION_MESSAGE(GetSymbols) \
 		UBA_SESSION_MESSAGE(SpeedTest) \
 		UBA_SESSION_MESSAGE(LogEntry) \
+		UBA_SESSION_MESSAGE(Symlink) \
 
 	#define UBA_CACHE_MESSAGES \
 		UBA_CACHE_MESSAGE(Connect) \
diff --git a/Engine/Source/Programs/UnrealBuildAccelerator/Common/Public/UbaSession.h b/Engine/Source/Programs/UnrealBuildAccelerator/Common/Public/UbaSession.h
--- a/Engine/Source/Programs/UnrealBuildAccelerator/Common/Public/UbaSession.h
+++ b/Engine/Source/Programs/UnrealBuildAccelerator/Common/Public/UbaSession.h
@@ -44,6 +44,8 @@ namespace uba
 	struct MoveFileResponse;
 	struct ChmodResponse;
 	struct ChmodMessage;
+	struct SymlinkMessage;
+	struct SymlinkResponse;
 	struct GetFullFileNameMessage;
 	struct GetFullFileNameResponse;
 	struct GetLongPathNameMessage;
@@ -140,6 +142,7 @@ namespace uba
 		bool IsRarelyRead(ProcessImpl& process, const StringView& fileName) const;
 		bool IsMaybeReadAfterWritten(ProcessImpl& process, const StringView& fileName) const;
 		bool IsKnownSystemFile(const tchar* applicationName);
+		virtual bool IsInTemp(StringView path) override;
 		virtual bool ShouldWriteToDisk(StringView fileName) override;
 		virtual void InvalidateCachedFileInfo(const StringKey& fileNameKey);
 
@@ -174,8 +177,9 @@ namespace uba
 		virtual bool CopyFile(CopyFileResponse& out, const CopyFileMessage& msg);
 		virtual bool MoveFile(MoveFileResponse& out, const MoveFileMessage& msg);
 		virtual bool Chmod(ChmodResponse& out, const ChmodMessage& msg);
-		virtual bool CreateDirectory(CreateDirectoryResponse& out, const CreateDirectoryMessage& msg, DirectoryTableOverlay* overlay);
-		virtual bool RemoveDirectory(RemoveDirectoryResponse& out, const RemoveDirectoryMessage& msg, DirectoryTableOverlay* overlay);
+		virtual bool Symlink(SymlinkResponse& out, const SymlinkMessage& msg);
+		virtual bool CreateDirectory(CreateDirectoryResponse& out, const CreateDirectoryMessage& msg);
+		virtual bool RemoveDirectory(RemoveDirectoryResponse& out, const RemoveDirectoryMessage& msg);
 		virtual bool GetFullFileName(GetFullFileNameResponse& out, const GetFullFileNameMessage& msg);
 		virtual bool GetLongPathName(GetLongPathNameResponse& out, const GetLongPathNameMessage& msg);
 		virtual bool GetListDirectoryInfo(ListDirectoryResponse& out, const StringView& dirName, const StringKey& dirKey, DirectoryTableOverlay* overlay);
@@ -200,7 +204,7 @@ namespace uba
 		bool WriteMemoryToDisk(FileAccessor& destinationFile, const void* fileMem, u64 fileSize, StringBufferBase* outWriteMode);
 		bool GetFileMemory(const Function<bool(const void*, u64)>& func, const StringKey& fileNameKey, StringView filePath, bool deleteInternalMapping);
 		void AddEnvironmentVariableNoLock(StringView key, StringView value);
-		bool CopyImports(Vector<BinaryModule>& out, const tchar* library, tchar* applicationDir, tchar* applicationDirEnd, UnorderedSet<TString>& handledImports, const tchar* const* loaderPaths, bool isLinuxBinary);
+		bool CopyImports(Vector<BinaryModule>& out, const tchar* library, tchar* applicationDir, tchar* applicationDirEnd, UnorderedSet<TString>& handledImports, const tchar* const* loaderPaths, bool isLinuxBinary, const tchar* hint);
 		bool CreateProcessJobObject();
 		void GetSessionInfo(StringBufferBase& out);
 
@@ -438,7 +442,8 @@ namespace uba
 		bool deleteOnClose = false;
 		bool success = true;
 		SharedMemoryHandle memoryHandle;
-		u64 memoryWritten = 0;
+		u64 fileSize = 0;
+		u64 lastWriteTime = 0;
 	};
 
 	struct CloseFileResponse
@@ -484,6 +489,9 @@ namespace uba
 		ProcessImpl& process;
 		StringKey fromKey;
 		StringBuffer<> fromName;
+		u64 fileSize = 0;
+		u64 lastWriteTime = 0;
+		u32 fileAttributes = 0;
 
 		StringKey toKey;
 		StringBuffer<> toName;
@@ -511,6 +519,26 @@ namespace uba
 		u32 errorCode = ~0u;
 	};
 
+	struct SymlinkMessage
+	{
+		ProcessImpl& process;
+		StringKey filePathKey;
+		StringBuffer<> symlinkFilePath;
+		StringKey contentPathKey;
+		StringBuffer<> symlinkContentPath;
+		StringBuffer<> symlinkContent;
+		u64 fileSize = 0;
+		u64 lastWriteTime = 0;
+		u32 fileAttributes = 0;
+	};
+
+	struct SymlinkResponse
+	{
+		bool success = false;
+		u32 errorCode = ~0u;
+		DirTableSize directoryTableSize;
+	};
+
 	struct GetFullFileNameMessage
 	{
 		ProcessImpl& process;
@@ -541,27 +569,25 @@ namespace uba
 
 	struct CreateDirectoryMessage
 	{
+		ProcessImpl& process;
 		StringKey nameKey;
 		StringBuffer<> name;
 	};
 
 	struct CreateDirectoryResponse
 	{
-		bool result = false;
-		u32 errorCode = 0;
 		DirTableSize directoryTableSize;
 	};
 
 	struct RemoveDirectoryMessage
 	{
+		ProcessImpl& process;
 		StringKey nameKey;
 		StringBuffer<> name;
 	};
 
 	struct RemoveDirectoryResponse
 	{
-		bool result = false;
-		u32 errorCode = 0;
 		DirTableSize directoryTableSize;
 	};
 
@@ -579,7 +605,6 @@ namespace uba
 
 	struct WrittenFile
 	{
-		ProcessImpl* owner = nullptr;
 		StringKey key;
 		TString name;
 		TString backedName;
@@ -591,6 +616,8 @@ namespace uba
 		u64 lastWriteTime = 0;
 		u32 attributes = 0;
 
+		u32 ownerId = 0; // Last process who write it
+
 		// TODO: Ugly placement of this pointer.
 		// Here to make sure FileMappingEntry is not garbage collecting before process exit callback is called
 		Session::FileMappingEntry* entryToReduceRefCount = nullptr;
diff --git a/Engine/Source/Programs/UnrealBuildAccelerator/Common/Public/UbaSessionClient.h b/Engine/Source/Programs/UnrealBuildAccelerator/Common/Public/UbaSessionClient.h
--- a/Engine/Source/Programs/UnrealBuildAccelerator/Common/Public/UbaSessionClient.h
+++ b/Engine/Source/Programs/UnrealBuildAccelerator/Common/Public/UbaSessionClient.h
@@ -65,8 +65,6 @@ namespace uba
 		virtual bool CopyFile(CopyFileResponse& out, const CopyFileMessage& msg) override;
 		virtual bool MoveFile(MoveFileResponse& out, const MoveFileMessage& msg) override;
 		virtual bool Chmod(ChmodResponse& out, const ChmodMessage& msg) override;
-		virtual bool CreateDirectory(CreateDirectoryResponse& out, const CreateDirectoryMessage& msg, DirectoryTableOverlay* overlay) override;
-		virtual bool RemoveDirectory(RemoveDirectoryResponse& out, const RemoveDirectoryMessage& msg, DirectoryTableOverlay* overlay) override;
 		virtual bool GetFullFileName(GetFullFileNameResponse& out, const GetFullFileNameMessage& msg) override;
 		virtual bool GetLongPathName(GetLongPathNameResponse& out, const GetLongPathNameMessage& msg) override;
 		virtual bool GetListDirectoryInfo(ListDirectoryResponse& out, const StringView& dirName, const StringKey& dirKey, DirectoryTableOverlay* overlay) override;
@@ -91,6 +89,7 @@ namespace uba
 		bool ReadModules(List<ModuleInfo>& outModules, StringBufferBase& outVirtualApplication, u32 processId, StringView application, StringView workingDir);
 		bool EnsureBinaryFile(StringBufferBase& out, StringBufferBase& outVirtual, u32 processId, StringView fileName, const StringKey& fileNameKey, StringView applicationDir, StringView workingDir, const u8* loaderPaths, u32 loaderPathsSize);
 		bool WriteBinFile(StringBufferBase& out, const StringView& binaryName, const CasKey& casKey, const KeyToString& applicationDir, u32 fileAttributes);
+		bool SendFileSystemOperations(ProcessImpl& process);
 		bool SendFiles(ProcessImpl& process);
 		bool SendFile(WrittenFile& file, u32 processId, bool keepMappingInMemory, bool compressed);
 		bool SendUpdateDirectoryTable(StackBinaryReader<SendMaxSize>& reader); // Note, reader is sent in to save stack space.
diff --git a/Engine/Source/Programs/UnrealBuildAccelerator/Core/Private/UbaStringBuffer.cpp b/Engine/Source/Programs/UnrealBuildAccelerator/Core/Private/UbaStringBuffer.cpp
--- a/Engine/Source/Programs/UnrealBuildAccelerator/Core/Private/UbaStringBuffer.cpp
+++ b/Engine/Source/Programs/UnrealBuildAccelerator/Core/Private/UbaStringBuffer.cpp
@@ -1,6 +1,7 @@
 // Copyright Epic Games, Inc. All Rights Reserved.
 
 #include "UbaStringBuffer.h"
+#include "UbaBinaryReaderWriter.h"
 #include "UbaPlatform.h"
 
 namespace uba
@@ -330,6 +331,12 @@ namespace uba
 		return *this;
 	}
 
+	StringBufferBase& StringBufferBase::ReadString(BinaryReader& reader)
+	{
+		reader.ReadString(*this);
+		return *this;
+	}
+
 	bool StringBufferBase::Contains(tchar c) const
 	{
 		return TStrchr(data, c) != nullptr;
diff --git a/Engine/Source/Programs/UnrealBuildAccelerator/Core/Public/UbaProcessUtils.h b/Engine/Source/Programs/UnrealBuildAccelerator/Core/Public/UbaProcessUtils.h
--- a/Engine/Source/Programs/UnrealBuildAccelerator/Core/Public/UbaProcessUtils.h
+++ b/Engine/Source/Programs/UnrealBuildAccelerator/Core/Public/UbaProcessUtils.h
@@ -41,6 +41,14 @@ namespace uba
 				if (isInQuotes || isInSingleQuotes || !isInArg)
 					continue;
 
+#if !PLATFORM_WINDOWS
+				// Backslash-space in unquoted context (e.g. CARGO_PKG_DESCRIPTION=The\ Rust\ Core\ Library):
+				// a lone backslash escapes the following space, keeping the argument alive.
+				// A paired backslash (\\) does NOT escape the space -- isBackslashOwned tracks this.
+				if (!isEnd && currentChar == TC(' ') && lastChar == TC('\\') && !isBackslashOwned)
+					continue;
+#endif
+
 				const CharType* argIt = argStart;
 				const CharType* argEnd = it;
 
@@ -78,9 +86,16 @@ namespace uba
 					if (*argIt == '\"' && !inSingleQ)
 					{
 						if (lastChar2 == '\\' && !isBackslashOwned)
+						{
+							// Escaped quote: emit a literal " (already written the \, replace it)
+							// and do NOT toggle inDoubleQ — we remain in the same quoting mode.
 							argWriteIt[-1] = '\"';
+						}
+						else
+						{
+							inDoubleQ = !inDoubleQ;
+						}
 						lastChar2 = 0;
-						inDoubleQ = !inDoubleQ;
 						++argIt;
 						continue;
 					}
@@ -94,6 +109,20 @@ namespace uba
 							isBackslashOwned = false;
 					}
 
+#if !PLATFORM_WINDOWS
+					// POSIX unquoted backslash: strip it and emit the next character literally.
+					// Covers \%, \(, \), \\ → \, \ → space, etc.
+					if (!inSingleQ && !inDoubleQ && *argIt == '\\' && argIt + 1 != argEnd)
+					{
+						++argIt; // skip the backslash
+						*argWriteIt++ = *argIt;
+						lastChar2 = *argIt;
+						isBackslashOwned = false;
+						++argIt;
+						continue;
+					}
+#endif
+
 					*argWriteIt++ = *argIt;
 
 					lastChar2 = *argIt;
diff --git a/Engine/Source/Programs/UnrealBuildAccelerator/Core/Public/UbaProtocol.h b/Engine/Source/Programs/UnrealBuildAccelerator/Core/Public/UbaProtocol.h
--- a/Engine/Source/Programs/UnrealBuildAccelerator/Core/Public/UbaProtocol.h
+++ b/Engine/Source/Programs/UnrealBuildAccelerator/Core/Public/UbaProtocol.h
@@ -17,6 +17,7 @@ namespace uba
 		UBA_PROCESS_MESSAGE(CopyFile) \
 		UBA_PROCESS_MESSAGE(MoveFile) \
 		UBA_PROCESS_MESSAGE(Chmod) \
+		UBA_PROCESS_MESSAGE(Symlink) \
 		UBA_PROCESS_MESSAGE(CreateDirectory) \
 		UBA_PROCESS_MESSAGE(RemoveDirectory) \
 		UBA_PROCESS_MESSAGE(ListDirectory) \
@@ -53,7 +54,7 @@ namespace uba
 		#undef UBA_PROCESS_MESSAGE
 	};
 
-	inline constexpr u32 ProcessMessageVersion = 1348;
+	inline constexpr u32 ProcessMessageVersion = 1349;
 
 	inline constexpr u32 CommunicationMemSize = IsWindows ? 64*1024 : 64*1024*2; // Macos expands some commandlines to be crazy long
 
diff --git a/Engine/Source/Programs/UnrealBuildAccelerator/Core/Public/UbaStringBuffer.h b/Engine/Source/Programs/UnrealBuildAccelerator/Core/Public/UbaStringBuffer.h
--- a/Engine/Source/Programs/UnrealBuildAccelerator/Core/Public/UbaStringBuffer.h
+++ b/Engine/Source/Programs/UnrealBuildAccelerator/Core/Public/UbaStringBuffer.h
@@ -6,6 +6,7 @@
 
 namespace uba
 {
+	struct BinaryReader;
 	struct StringView;
 
 	////////////////////////////////////////////////////////////////////////////////////////////////////
@@ -46,6 +47,8 @@ namespace uba
 		StringBufferBase& Resize(u64 newSize);
 		StringBufferBase& Clear();
 
+		StringBufferBase& ReadString(BinaryReader& reader);
+
 		#if PLATFORM_WINDOWS
 		StringBufferBase& Append(const char* str);
 		StringBufferBase& Append(const char* str, u32 charCount);
@@ -85,6 +88,7 @@ namespace uba
 		u32 Parse(char* out, u64 outCapacity) const; // Note, return value contains length + null termination.. zero if failed
 
 		TString ToString() const { return TString(data, data + count); }
+		TString ToStringAndClear() { TString s(data, data + count); Clear(); return s; }
 		bool ToFalse() const { return false; }
 
 		u32 count;
diff --git a/Engine/Source/Programs/UnrealBuildAccelerator/Detours/Private/PosixOS/UbaDetoursFunctionsPosix.cpp b/Engine/Source/Programs/UnrealBuildAccelerator/Detours/Private/PosixOS/UbaDetoursFunctionsPosix.cpp
--- a/Engine/Source/Programs/UnrealBuildAccelerator/Detours/Private/PosixOS/UbaDetoursFunctionsPosix.cpp
+++ b/Engine/Source/Programs/UnrealBuildAccelerator/Detours/Private/PosixOS/UbaDetoursFunctionsPosix.cpp
@@ -27,7 +27,7 @@
 
 #if PLATFORM_LINUX
 #include <sys/prctl.h>
-//long clone3(struct clone_args *cl_args, size_t size);
+#include <sys/sysmacros.h>
 #else
 #include <crt_externs.h>
 extern char** environ;
@@ -93,6 +93,7 @@ using namespace uba;
 	DETOURED_FUNCTION(fstatat) \
 	DETOURED_FUNCTION(futimens) \
 	DETOURED_FUNCTION(fclose) \
+	DETOURED_FUNCTION(mkstemp) \
 	DETOURED_FUNCTION(opendir) \
 	DETOURED_FUNCTION(fdopendir) \
 	DETOURED_FUNCTION(dirfd) \
@@ -144,9 +145,6 @@ using namespace uba;
 	DETOURED_FUNCTION_DEBUG \
 	DETOURED_FUNCTION_LINUX \
 	DETOURED_FUNCTION_MACOS \
-
-	//DETOURED_FUNCTION(clone)
-	//DETOURED_FUNCTION(clone3)
 	
 // Used by cp and might need detour
 // fadvise64
@@ -164,6 +162,8 @@ using namespace uba;
 #if PLATFORM_LINUX
 #define DETOURED_FUNCTION_MACOS
 #define DETOURED_FUNCTION_LINUX \
+	DETOURED_FUNCTION(dup3) \
+	DETOURED_FUNCTION(clone) \
 	DETOURED_FUNCTION(get_current_dir_name) \
 	DETOURED_FUNCTION(fopen64) \
 	DETOURED_FUNCTION(secure_getenv) \
@@ -171,16 +171,24 @@ using namespace uba;
 	DETOURED_FUNCTION(__xstat) \
 	DETOURED_FUNCTION(__xstat64) \
 	DETOURED_FUNCTION(__lxstat) \
+	DETOURED_FUNCTION(__lxstat64) \
 	DETOURED_FUNCTION(__fxstat) \
 	DETOURED_FUNCTION(__fxstat64) \
 	DETOURED_FUNCTION(__fxstatat) \
 	DETOURED_FUNCTION(__fxstatat64) \
+	DETOURED_FUNCTION(pread64) \
 	DETOURED_FUNCTION(open64) \
+	DETOURED_FUNCTION(openat64) \
+	DETOURED_FUNCTION(stat64) \
+	DETOURED_FUNCTION(lstat64) \
+	DETOURED_FUNCTION(fstat64) \
+	DETOURED_FUNCTION(scandir64) \
 	DETOURED_FUNCTION(readdir64) \
 	DETOURED_FUNCTION(fstatat64) \
 	DETOURED_FUNCTION(fpathconf) \
 	DETOURED_FUNCTION(pathconf) \
 	DETOURED_FUNCTION(syscall) \
+	DETOURED_FUNCTION(statx) \
 
 #else
 #define DETOURED_FUNCTION_LINUX
@@ -294,7 +302,9 @@ namespace uba
 	u8 GetFileAccessFlags(int flags)
 	{
 		u8 access = 0;
-		if (flags & O_RDWR)
+		if (!flags)
+			access |= AccessFlag_Read;
+		else if (flags & O_RDWR)
 			access |= AccessFlag_Read | AccessFlag_Write;
 		else if (flags & O_RDONLY)
 			access |= AccessFlag_Read;
@@ -325,10 +335,8 @@ namespace uba
 	VARIABLE_MEM(ReaderWriterLock, g_fileHandlesLock);
 
 
-	StringKey ToFilenameKey(const StringBufferBase& b)
-	{
-		return CaseInsensitiveFs ? ToStringKeyLower(b) : ToStringKey(b);
-	}
+	StringKey ToFilenameKey(const StringBufferBase& b) { return CaseInsensitiveFs ? ToStringKeyLower(b) : ToStringKey(b); }
+	StringKey ToFilenameKey(StringView b) { return CaseInsensitiveFs ? ToStringKeyLower(b) : ToStringKey(b); }
 
 	bool CouldBeCompressedFile(const StringView& fileName) { return false; }
 	bool CanDetour(StringView file)
@@ -347,13 +355,39 @@ bool CanDetour2(const StringView& file)
 		&& !file.StartsWith("/etc/")
 	#if PLATFORM_LINUX
 		&& !file.StartsWith("/sys/") // Don't know if this is needed for macos.. but is needed for linux
+		&& !file.StartsWith("/run/") // Don't know if this is needed for macos.. but is needed for linux
 	#endif
-		&& !file.StartsWith(g_systemTemp.data)
 		;
 }
 
 // Shared functions
 
+#if UBA_DEBUG_LOG_ENABLED
+static StringBuffer<256> OpenFlagsToString(int flags)
+{
+	StringBuffer<256> buf;
+	auto append = [&](const char* s) { if (buf.count) buf.Append('|'); buf.Append(s); };
+	int acc = flags & O_ACCMODE;
+	if      (acc == O_RDONLY) append("O_RDONLY");
+	else if (acc == O_WRONLY) append("O_WRONLY");
+	else if (acc == O_RDWR)   append("O_RDWR");
+	if (flags & O_CREAT)     append("O_CREAT");
+	if (flags & O_TRUNC)     append("O_TRUNC");
+	if (flags & O_APPEND)    append("O_APPEND");
+	if (flags & O_NONBLOCK)  append("O_NONBLOCK");
+	if (flags & O_CLOEXEC)   append("O_CLOEXEC");
+	if (flags & O_EXCL)      append("O_EXCL");
+	if (flags & O_NOCTTY)    append("O_NOCTTY");
+	if (flags & O_NOFOLLOW)  append("O_NOFOLLOW");
+	if (flags & O_SYNC)      append("O_SYNC");
+	#if PLATFORM_LINUX
+	if (flags & O_TMPFILE)   append("O_TMPFILE");
+	#endif
+	if (flags & O_DIRECTORY) append("O_DIRECTORY");
+	return buf;
+}
+#endif
+
 template<typename TrueOpen>
 int Shared_open(const char* funcName, const char* file, int flags, int mode, const TrueOpen& trueOpen)
 {
@@ -393,58 +427,88 @@ int Shared_open(const char* funcName, const char* file, int flags, int mode, con
 
 
 	StringKey fileNameKey = ToFilenameKey(fileName);
-
-	bool keepInMemory = false;
-
-	u64 size = InvalidValue;
-	u32 closeId = 0;
 	DirTableOffset dirTableOffset;
 
-	if (g_allowDirectoryCache)
+	bool foundEntry = Rpc_GetEntryOffset(dirTableOffset, fileNameKey, fileName, false);
+	DirectoryTable::EntryInformation entryInfo;
+	if (foundEntry)
+		g_directoryTable.GetEntryInformation(entryInfo, dirTableOffset);
+	bool fileExists = foundEntry && entryInfo.attributes != 0;
+
+	if (!(flags & O_CREAT))
 	{
-		if (!isWrite) // We need to skip SystemTemp.. lots of stuff going on there.
+		// File must exist if O_CREATE is not set
+		if (!fileExists)
 		{
-			bool allowEarlyOut = true;
-			if (!Rpc_GetEntryOffset(dirTableOffset, fileNameKey, fileName, false))
+			errno = ENOENT;
+			DEBUG_LOG_DETOURED(funcName, "(DIRTABLE) %s -> -1 (ENOENT)", fileName.data);
+			return -1;
+		}
+	}
+	else if (const tchar* lastSeparator = fileName.Last(PathSeparator))
+	{
+		if (fileExists)
+		{
+			if (flags & O_EXCL)
 			{
-				// This could be a written file not reported to server yet
-				{
-					SCOPED_READ_LOCK(g_mappedFileTable.m_lookupLock, lock);
-					auto findIt = g_mappedFileTable.m_lookup.find(fileNameKey);
-					if (findIt != g_mappedFileTable.m_lookup.end())
-						allowEarlyOut = findIt->second.deleted;
-				}
-				if (allowEarlyOut)
-				{
-					//SetLastError(ERROR_FILE_NOT_FOUND); // Don't think this is needed
-					errno = ENOENT;
-					DEBUG_LOG_DETOURED(funcName, "NOTFOUND_USINGTABLE (%s) (%s) -> -1", fileName.data, KeyToString(fileNameKey).data);
-					return -1;
-				}
+				errno = EEXIST;
+				DEBUG_LOG_DETOURED(funcName, "(DIRTABLE) %s -> -1 (EEXIST)", fileName.data);
+				return -1;
 			}
-			else
+		}
+		else
+		{
+			// Check so parent directory exist
+			StringView dirName(fileName.data, u32(lastSeparator - fileName.data));
+			DirHash hash(dirName);
+
+			SCOPED_WRITE_LOCK(g_directoryTable.m_lookupLock, lookupLock);
+			auto insres = g_directoryTable.m_lookup.try_emplace(hash.key, g_memoryBlock);
+			DirectoryTable::Directory& dir = insres.first->second;
+			if (insres.second)
+				if (g_directoryTable.EntryExistsNoLock(hash.key, dirName) != DirectoryTable::Exists_No)
+					Rpc_UpdateDirectory(hash.key, dirName.data, dirName.count, false);
+			DirTableOffset dirTableOffset2;
+			if (!g_directoryTable.GetLatestOffset(dirTableOffset2, dir))
 			{
-				// File could have been deleted.
-				DirectoryTable::EntryInformation entryInfo;
-				g_directoryTable.GetEntryInformation(entryInfo, dirTableOffset);
-				if (entryInfo.attributes == 0)
-				{
-					DEBUG_LOG_DETOURED(funcName, "DELETED (%s) (%s) -> -1", fileName.data, KeyToString(fileNameKey).data);
-					errno = ENOENT;
-					return -1;
-				}
+				errno = ENOENT;
+				DEBUG_LOG_DETOURED(funcName, "(DIRTABLE) %s -> -1 (ENOENT)", fileName.data);
+				return -1;
+			}
+		}
+	}
+	else
+	{
+		UBA_ASSERT(false); // Should never end up here
+	}
+
+	if (fileExists)
+	{
+		if (IsDirectory(entryInfo.attributes))
+		{
+			if ((flags & (O_WRONLY | O_RDWR | O_CREAT)) != 0)
+			{
+				errno = EISDIR;
+				DEBUG_LOG_DETOURED(funcName, "(DIRTABLE) %s -> -1 (EISDIR)", fileName.data);
+				return -1;
 			}
 		}
 	}
 
+	bool keepInMemory = false;
+
+	u64 size = InvalidValue;
+	u32 closeId = 0;
+
 	const char* realFileName = fileName.data;
+	const char* prevFileName = nullptr;
 
 	SCOPED_WRITE_LOCK(g_mappedFileTable.m_lookupLock, fileTableLock);
 	auto insres = g_mappedFileTable.m_lookup.try_emplace(fileNameKey);
 	FileInfo& info = insres.first->second;
 	FileInfo* fileInfo = &info;
 	u32 lastDesiredAccess = info.lastDesiredAccess;
-	if (insres.second)
+	if (insres.second || info.recursiveCall)
 	{
 		info.originalName = g_memoryBlock.Strdup(fileName).data;
 		info.name = info.originalName;
@@ -466,6 +530,7 @@ int Shared_open(const char* funcName, const char* file, int flags, int mode, con
 			info.originalName = g_memoryBlock.Strdup(fileName).data;
 		if (isWrite)
 		{
+			prevFileName = info.name;
 			//UBA_ASSERT(!info.isFileMap);
 			bool isTruncating = (flags & O_TRUNC) != 0;
 			bool shouldReport = !(info.lastDesiredAccess & AccessFlag_Write) || info.deleted || isTruncating;
@@ -536,57 +601,115 @@ int Shared_open(const char* funcName, const char* file, int flags, int mode, con
 	else
 		tempFileName = info.name;
 
+	int trueFlags = flags;
+
+	Vector<u8> dataToWrite;
+
+	auto readDataFromOldFile = [&]()
+		{
+			DEBUG_LOG_DETOURED(funcName, "INTERNAL READ FOR MEMORYWRITE (%s) (%s)", fileName.data, OpenFlagsToString(trueFlags).data);
+			UBA_ASSERT(!info.recursiveCall);
+			FileInfo temp = info;
+			info.recursiveCall = true;
+			fileTableLock.Leave();
+
+			int readFd = Shared_open(funcName, fileName.data, O_RDONLY|O_NONBLOCK, mode, trueOpen);
+			UBA_ASSERTF(readFd != -1, "%s", fileName.data);
+			dataToWrite.resize(entryInfo.size);
+			int readRes = read(readFd, dataToWrite.data(), entryInfo.size);
+			close(readFd);
+			UBA_ASSERT(readRes == entryInfo.size);
+
+			fileTableLock.Enter();
+			info = temp;
+		};
+
+	// When running remote, if O_APPEND is set we need to copy the source file into the destination file since they are not the same files
+	if (fileExists && (flags & O_APPEND) && !fileName.Equals(tempFileName))
+	{
+		// It might already have been created as a write file and in that case we don't need to copy old state over
+		if (!prevFileName || !Equals(prevFileName, tempFileName))
+		{
+			readDataFromOldFile();
+			// trueFlags &= ~O_APPEND; // O_APPEND actually gives different behavior on writes so we can't remove that flag
+			trueFlags |= O_CREAT;
+			mode = entryInfo.attributes;
+		}
+	}
+
+	if (flags == O_RDWR && !fileName.Equals(tempFileName))
+	{
+		UBA_ASSERT(fileExists);
+		if (!prevFileName || !Equals(prevFileName, tempFileName))
+		{
+			readDataFromOldFile();
+			trueFlags |= O_CREAT;
+			mode = entryInfo.attributes;
+		}
+	}
+
+	// If we write a file that we know exists there might be a missing O_CREAT but on remote builds the output file actually doesn't exist
+	// so we need to add the flag to make sure it is created. Note, this only works when O_TRUC is there, otherwise we need to download
+	// the read-file, create the new file and write the read-file content into the write
+	if (flags == (O_WRONLY|O_TRUNC) && !fileName.Equals(tempFileName))
+	{
+		trueFlags |= O_CREAT;
+		if (mode == 0) // Caller did not plan on creating a file, let's create the default mode
+			mode = S_IRUSR | S_IWUSR | S_IRGRP | S_IROTH;
+	}
+
 	TimerScope ts(g_kernelStats.getFileInfo);
-	int fd = trueOpen(tempFileName, flags, mode);
+	int fd = trueOpen(tempFileName, trueFlags, mode);
 	ts.Leave();
 
-	DEBUG_LOG_TRUE(funcName, "%s%s (%s) %i %i -> %i (%s)", file, isWriteStr, tempFileName, flags, mode, fd, StrError(fd, errno));
+	DEBUG_LOG_TRUE(funcName, "%s%s (%s) %s %i -> %i (%s)", file, isWriteStr, tempFileName, OpenFlagsToString(trueFlags).data, mode, fd, StrError(fd, errno));
 	if (fd == -1)
+	{
+		// TODO: Do we need to report failure if closeId is set?
 		return fd;
+	}
+
+	if (auto toWrite = dataToWrite.size())
+	{
+		int wrres = write(fd, dataToWrite.data(), toWrite);
+		UBA_ASSERT(wrres == toWrite);
+	}
 
 	TrackInput(fileName);
 
 	SCOPED_WRITE_LOCK(g_fileHandlesLock, lock);
 	auto insres2 = g_fileHandles.insert({ fd, DetouredHandle() });
-	UBA_ASSERTF(insres2.second, "File handle already added");
+	UBA_ASSERTF(insres2.second, "File handle %i already added. %s (%s)", fd, funcName, tempFileName);
 	DetouredHandle& h = insres2.first->second;
 	auto fo = new FileObject();
 	fo->closeId = closeId;
 	fo->fileInfo = fileInfo;
 	fo->desiredAccess = desiredAccess;
 	h.fileObject = fo;
+	lock.Leave();
+
+	if (isWrite)
+		Rpc_RegisterFileForWrite(fileName, fileNameKey);
+
 	return fd;
 }
 
 template<typename TrueOpen>
 FILE* Shared_fopen(const char* funcName, const char* path, const char* mode, const char* trueOpenName, const TrueOpen& trueOpen)
 {
-	bool r = strchr(mode, 'r');
-	bool w = strchr(mode, 'w');
-	bool a = strchr(mode, 'a');
 	bool p = strchr(mode, '+');
 	//bool b = strchr(mode, 'b');
 	//bool t = strchr(mode, 't');
-	if (a)
-	{
-		//if (StartsWith(path, g_systemTemp.data))
-		{
-			FILE* res = TRUE_WRAPPER(fopen)(path, mode);
-			DEBUG_LOG_TRUE(funcName, "(%s  %s) -> %p", path, mode, res);
-			if (res)
-				TrackInput(ToView(path));
-			return res;
-		}
-		//UBA_ASSERTF(false, "%s with append (%s) not implemented (%s)", funcName, mode, path);
-	}
 
 	int flags = 0;
-	if (r)
+	if (strchr(mode, 'r'))
 	{
-		UBA_ASSERTF(!p, "%s with + not implemented (%s)", funcName, mode);
-		flags = O_RDONLY | O_NONBLOCK;
+		if (p)
+			flags = O_RDWR;
+		else
+			flags = O_RDONLY | O_NONBLOCK;
 	}
-	else if (w)
+	else if (strchr(mode, 'w'))
 	{
 		flags = O_CREAT | O_TRUNC;
 		if (p)
@@ -594,6 +717,15 @@ FILE* Shared_fopen(const char* funcName, const char* path, const char* mode, con
 		else
 			flags |= O_WRONLY;
 	}
+	else if (strchr(mode, 'a'))
+	{
+		flags = O_CREAT | O_APPEND;
+		if (p)
+			flags |= O_RDWR;
+		else
+			flags |= O_WRONLY;
+	}
+
 	int openMode = S_IRUSR | S_IWUSR | S_IRGRP | S_IROTH;
 	int fd = Shared_open(trueOpenName, path, flags, openMode, trueOpen);
 	if (fd == -1)
@@ -607,28 +739,28 @@ FILE* Shared_fopen(const char* funcName, const char* path, const char* mode, con
 	FILE* res = TRUE_WRAPPER(fdopen)(fd, mode);
 	ts.Leave();
 
-	DEBUG_LOG_TRUE(funcName, "%i (%s  %s) -> %p", fd, path, mode, res);
+	DEBUG_LOG_TRUE(funcName, "%i (%s %s) -> %p", fd, path, mode, res);
 	return res;
 }
 
 template<typename TrueClose>
-void Shared_close(int fd, const TrueClose& trueClose)
+int Shared_close(const char* funcName, int fd, const TrueClose& trueClose)
 {
+	auto runClose = [&]()
+		{
+			TimerScope ts(g_kernelStats.closeFile);
+			int res = trueClose();
+			DEBUG_LOG_TRUE(funcName, "(%i) -> %i (%s)", fd, res, StrError(res, errno));
+			return res;
+		};
+
 	if (!g_isDetouring)
-	{
-		TimerScope ts(g_kernelStats.closeFile);
-		trueClose();
-		return;
-	}
+		return runClose();
 
 	SCOPED_WRITE_LOCK(g_fileHandlesLock, lock);
 	auto findIt = g_fileHandles.find(fd);
 	if (findIt == g_fileHandles.end())
-	{
-		TimerScope ts(g_kernelStats.closeFile);
-		trueClose();
-		return;
-	}
+		return runClose();
 
 	DetouredHandle& h = findIt->second;
 	FileObject* fo = h.fileObject;
@@ -636,45 +768,48 @@ void Shared_close(int fd, const TrueClose& trueClose)
 	g_fileHandles.erase(findIt);
 	lock.Leave();
 
-	auto closeGuard = MakeGuard([&]()
-	{
-		TimerScope ts(g_kernelStats.closeFile);
-		trueClose();
-	});
-
 	if (--fo->refCount)
-		return;
+		return runClose();
+
+	auto dg = MakeGuard([fo]() { delete fo; });
+
+	if (!fo->closeId)
+		return runClose();
 
 	SharedMemoryHandle memoryHandle;
-	u64 memoryWritten = 0;
+	u64 fileSize = InvalidValue;
+	u64 lastWritten = 0;
+	u32 attributes = 0;
+
 	FileInfo& fi = *fo->fileInfo;
 	const tchar* path = fi.name;
 
-	if (fo->closeId)
-	{
-		fi.created = true; // TODO: Should be set earlier
-		// Dup the fd before closing so we can measure the true file size AFTER fclose
-		// flushes any stdio buffers. Measuring via lseek before fclose misses any bytes
-		// still sitting in the stdio buffer (e.g. when 'as' uses fdopen on top of open64).
-		int dupFd = TRUE_WRAPPER(dup)(fd);
-		closeGuard.Execute();
-		if (dupFd != -1)
-		{
-			fi.size = TRUE_WRAPPER(lseek)(dupFd, 0, SEEK_END);
-			TRUE_WRAPPER(close)(dupFd);
-		}
-		//for (auto& kv : g_fileHandles)
-		//	UBA_ASSERT(&fi != kv.second.fileObject->fileInfo);
-		if (fi.size == (off_t)-1)
-		{
-			// Non-seekable fd (pipe etc.) — should not happen for tracked output files
-			UBA_ASSERT(false);
-			delete fo;
-			return;
-		}
-		Rpc_UpdateCloseHandle(path, fo->closeId, fo->deleteOnClose, fo->newName.c_str(), memoryHandle, memoryWritten, true);
-	}
-	delete fo;
+	fi.created = true; // TODO: Should be set earlier
+
+	// Dup the fd before closing so we can measure the true file size AFTER fclose
+	// flushes any stdio buffers. Measuring via lseek before fclose misses any bytes
+	// still sitting in the stdio buffer (e.g. when 'as' uses fdopen on top of open64).
+	int dupFd = TRUE_WRAPPER(dup)(fd);
+	UBA_ASSERT(dupFd != -1);
+	int res = trueClose();
+	int error = errno;
+
+	struct stat attr;
+	int res2 = TRUE_WRAPPER(fstat)(dupFd, &attr);
+	UBA_ASSERT(res2 == 0);(void)res2; 
+	fileSize = fi.size = attr.st_size;
+	lastWritten = FromTimeSpec(attr.st_mtimespec);
+	attributes = attr.st_mode;
+	//fi.size = TRUE_WRAPPER(lseek)(dupFd, 0, SEEK_END);
+	//UBA_ASSERT(fi.size != (off_t)-1);
+	TRUE_WRAPPER(close)(dupFd);
+	fileSize = fi.size;
+
+	DEBUG_LOG_TRUE(funcName, "(%i) Size: %llu -> %i (%s)", fd, fileSize, res, StrError(res, error));
+
+	Rpc_UpdateCloseHandle(path, fo->closeId, fo->deleteOnClose, fo->newName.c_str(), memoryHandle, fileSize, lastWritten, attributes, true);
+	errno = error;
+	return res;
 }
 
 template<typename True_fstat>
@@ -838,6 +973,48 @@ int Shared_stat(const char* funcName, const char* file, struct stat* attr, const
 	return res;
 }
 
+#if PLATFORM_LINUX
+int Shared_statx(const char* funcName, int dirfd, const char* pathname, int flags, unsigned int mask, struct statx* statxbuf)
+{
+	if (dirfd != AT_FDCWD)
+	{
+		UBA_ASSERTF(!g_runningRemote, TC("Handle statx with dirfd set and pathname %s"), pathname); // TODO: We need to properly implement this
+		int res = (int)TRUE_WRAPPER(syscall)(SYS_statx, dirfd, pathname, flags, mask, statxbuf);
+		DEBUG_LOG_TRUE("syscall.statx", "(NODETOUR) %s -> %s", pathname, res);
+		return res;
+	}
+	bool trueStatCalled = false;
+	struct stat st;
+	int res = Shared_stat(funcName, pathname, &st, [&](const char* path, struct stat* /*unused*/) -> int {
+		trueStatCalled = true;
+		return (int)TRUE_WRAPPER(syscall)(SYS_statx, AT_FDCWD, path, flags, mask, statxbuf);
+	});
+	if (trueStatCalled)
+		return res;
+	if (res != 0)
+		return res;
+
+	memset(statxbuf, 0, sizeof(*statxbuf));
+	statxbuf->stx_mask       = STATX_BASIC_STATS;
+	statxbuf->stx_blksize    = (u32)st.st_blksize;
+	statxbuf->stx_nlink      = (u32)st.st_nlink;
+	statxbuf->stx_uid        = st.st_uid;
+	statxbuf->stx_gid        = st.st_gid;
+	statxbuf->stx_mode       = (u16)st.st_mode;
+	statxbuf->stx_ino        = st.st_ino;
+	statxbuf->stx_size       = (u64)st.st_size;
+	statxbuf->stx_blocks     = (u64)st.st_blocks;
+	statxbuf->stx_atime      = { st.st_atim.tv_sec, (u32)st.st_atim.tv_nsec };
+	statxbuf->stx_mtime      = { st.st_mtim.tv_sec, (u32)st.st_mtim.tv_nsec };
+	statxbuf->stx_ctime      = { st.st_ctim.tv_sec, (u32)st.st_ctim.tv_nsec };
+	statxbuf->stx_rdev_major = major(st.st_rdev);
+	statxbuf->stx_rdev_minor = minor(st.st_rdev);
+	statxbuf->stx_dev_major  = major(st.st_dev);
+	statxbuf->stx_dev_minor  = minor(st.st_dev);
+	return res;
+}
+#endif
+
 //template<typename True_stat>
 int Shared_access(const char* funcName, const char* pathname, int mode)
 {
@@ -990,35 +1167,44 @@ UBA_EXPORT int UBA_WRAPPER(mkdir)(const char* path, mode_t mode)
 {
 	UBA_INIT_DETOUR(mkdir, path, mode);
 
-	StringBuffer<> pathName;
-	if (!FixPath(pathName, path) || !CanDetour2(pathName) || (pathName.count == 1 && pathName[0] == '/'))
+	StringBuffer<> dirPath;
+	if (!FixPath(dirPath, path) || !CanDetour2(dirPath) || (dirPath.count == 1 && dirPath[0] == '/'))
 	{
 		int res = TRUE_WRAPPER(mkdir)(path, mode);
 		DEBUG_LOG_TRUE("mkdir", "%s -> %i", path, res);
 		return res;
 	}
-	DevirtualizePath(pathName);
+	DevirtualizePath(dirPath);
+	StringKey dirPathKey = ToFilenameKey(dirPath);
 
-	DirTableSize directoryTableSize;
-	int res;
-	u32 errorCode = 0;
-	StringKey pathNameKey = ToFilenameKey(pathName);
+	DirTableOffset dirTableOffset;
+	if (Rpc_GetEntryOffset(dirTableOffset, dirPathKey, dirPath, false))
+	{
+		errno = EEXIST;
+		DEBUG_LOG_DETOURED("mkdir", "%s -> -1 (EEXIST)", dirPath.data);
+		return -1;
+	}
 
+	const tchar* lastSeparator = dirPath.Last(PathSeparator);
+	UBA_ASSERT(lastSeparator);
+	StringView parentPath(dirPath.data, u32(lastSeparator - dirPath.data));
+	if (!Rpc_GetEntryOffset(dirTableOffset, ToFilenameKey(parentPath), parentPath, true))
 	{
-		RPC_MESSAGE(CreateDirectory, createFile);
-		writer.WriteStringKey(pathNameKey);
-		writer.WriteString(pathName);
-		BinaryReader reader = writer.Flush();
-		res = reader.ReadBool() ? 0 : -1;
-		errorCode = reader.ReadU32();
-		directoryTableSize = FromU64(reader.ReadU64());
+		errno = ENOENT;
+		DEBUG_LOG_DETOURED("mkdir", "%s -> -1 (ENOENT)", dirPath.data);
+		return false;
 	}
 
-	g_directoryTable.ParseDirectoryTable(directoryTableSize);
+	RPC_MESSAGE(CreateDirectory, createFile);
+	writer.WriteStringKey(dirPathKey);
+	writer.WriteString(dirPath);
+	BinaryReader reader = writer.Flush();
+	DirTableSize directoryTableSize = FromU64(reader.ReadU64());
+	RPC_MESSAGE_END
 
-	errno = errorCode;
-	DEBUG_LOG_DETOURED("mkdir", "%s -> %i (%u)", path, res, errorCode);
-	return res;
+	g_directoryTable.ParseDirectoryTable(directoryTableSize);
+	DEBUG_LOG_DETOURED("mkdir", "%s -> 0", path);
+	return 0;
 }
 
 UBA_EXPORT int UBA_WRAPPER(rmdir)(const char* path)
@@ -1035,26 +1221,19 @@ UBA_EXPORT int UBA_WRAPPER(rmdir)(const char* path)
 
 	DevirtualizePath(pathName);
 
-	DirTableSize directoryTableSize;
-	bool res;
-	u32 errorCode = 0;
 	StringKey pathNameKey = ToFilenameKey(pathName);
 
-	{
-		RPC_MESSAGE(RemoveDirectory, deleteFile);
-		writer.WriteStringKey(pathNameKey);
-		writer.WriteString(pathName);
-		BinaryReader reader = writer.Flush();;
-		res = reader.ReadBool();
-		errorCode = reader.ReadU32();
-		directoryTableSize = FromU64(reader.ReadU64());
-	}
+	RPC_MESSAGE(RemoveDirectory, deleteFile);
+	writer.WriteStringKey(pathNameKey);
+	writer.WriteString(pathName);
+	BinaryReader reader = writer.Flush();
+	DirTableSize directoryTableSize = FromU64(reader.ReadU64());
+	RPC_MESSAGE_END
 
 	g_directoryTable.ParseDirectoryTable(directoryTableSize);
 
-	errno = errorCode;
-	DEBUG_LOG_DETOURED("rmdir", "%s -> %i (%u)", path, res, errorCode);
-	return res ? 0 : -1;
+	DEBUG_LOG_DETOURED("rmdir", "%s -> 0", path);
+	return 0;
 }
 
 UBA_EXPORT int UBA_WRAPPER(chroot)(const char* path)
@@ -1217,10 +1396,35 @@ UBA_EXPORT ssize_t UBA_WRAPPER(readlink)(const char* pathname, char* buf, size_t
 	else
 	{
 		UBA_ASSERTF(!StartsWith(pathname, "/UEVFS"), "Need to devirtualize %s", pathname);
+
+		// If the path is managed by the VFS it is a regular file on the host (not a symlink).
+		// Return EINVAL so callers know the file exists but is not a symlink, rather than
+		// getting ENOENT from the OS (which has no physical copy of the file on the agent).
+		// Without this, tools like rustc that probe their own binary path via readlink()
+		// fall back to dladdr(), which returns a path with unresolved ".." components
+		// (e.g. "bin/rustc/../lib/..."). That causes a wrong sysroot computation and an
+		// incorrect --print target-libdir output.
+		if (g_runningRemote || IsVfsEnabled())
+		{
+			StringBuffer<> fixedPath;
+			if (FixPath(fixedPath, pathname) && CanDetour2(fixedPath))
+			{
+				StringBuffer<> devirtualizedPath(fixedPath);
+				DevirtualizePath(devirtualizedPath);
+				FileAttributes fileAttr;
+				Shared_GetFileAttributes(fileAttr, devirtualizedPath);
+				if (fileAttr.useCache && fileAttr.exists)
+				{
+					errno = EINVAL; // File exists in VFS as a regular file, not a symlink
+					DEBUG_LOG_DETOURED("readlink", "(%s) -> -1 (EINVAL, VFS regular file not a symlink)", pathname);
+					return -1;
+				}
+			}
+		}
 	}
 
 	auto res = TRUE_WRAPPER(readlink)(pathname, buf, bufsiz);
-	DEBUG_LOG_TRUE("readlink", "(%s) (%s) -> %llu", pathname, buf, u64(res));
+	DEBUG_LOG_TRUE("readlink", "(%s) (%s) -> %lli (%s)", pathname, buf, res, StrError(int(res), errno));
 
 	if (res && res < bufsiz && IsVfsEnabled())
 	{
@@ -1258,6 +1462,12 @@ ssize_t UBA_WRAPPER(readlinkat)(int dirfd, const char* pathname, char* buf, size
 //}
 
 #if PLATFORM_LINUX
+
+UBA_EXPORT ssize_t UBA_WRAPPER(pread64)(int __fd, void * __buf, size_t __nbyte, off_t __offset)
+{
+	return TRUE_WRAPPER(pread64)(__fd, __buf, __nbyte, __offset);
+}
+
 UBA_EXPORT int UBA_WRAPPER(open64)(const char* file, int flags, ...)
 {
 	va_list args;
@@ -1372,12 +1582,7 @@ UBA_EXPORT int UBA_WRAPPER(dup2)(int oldfd, int newfd)
 UBA_EXPORT int UBA_WRAPPER(close)(int fd)
 {
 	UBA_INIT_DETOUR(close, fd);
-	int res = 0;
-	int error = 0;
-	Shared_close(fd, [&]() { res = TRUE_WRAPPER(close)(fd); error = errno; });
-	DEBUG_LOG_TRUE("close", "(%i) -> %i (%s)", fd, res, StrError(res, error));
-	errno = error;
-	return res;
+	return Shared_close("close", fd, [&]() { return TRUE_WRAPPER(close)(fd); });
 }
 
 #if PLATFORM_LINUX
@@ -1386,6 +1591,17 @@ UBA_EXPORT FILE* UBA_WRAPPER(fopen64)(const char* path, const char* mode)
 	UBA_INIT_DETOUR(fopen64, path, mode);
 	return Shared_fopen("fopen64", path, mode, "open64", [](const char* realFile, int flags, int openMode) { return TRUE_WRAPPER(open64)(realFile, flags, openMode); });
 }
+UBA_EXPORT int UBA_WRAPPER(fstat64)(int fd, struct stat64* buf)
+{
+	UBA_INIT_DETOUR(fstat64, fd, buf);
+	return Shared_fstat("fstat64", fd, (struct stat*)buf, [](int fd, struct stat* buf) { return TRUE_WRAPPER(fstat64)(fd, (struct stat64*)buf); });
+}
+UBA_EXPORT int __fxstat64(int ver, int fd, struct stat64* attr)
+{
+	UBA_INIT_DETOUR(__fxstat64, ver, fd, attr);
+	UBA_ASSERT(ver == _STAT_VER);
+	return Shared_fstat("__fxstat64", fd, (struct stat*)attr, [](int fd, struct stat* buf) { return TRUE_WRAPPER(__fxstat64)(_STAT_VER, fd, (struct stat64*)buf); });
+}
 #endif
 
 UBA_EXPORT FILE* UBA_WRAPPER(fopen)(const char* path, const char* mode)
@@ -1477,14 +1693,7 @@ UBA_EXPORT int UBA_WRAPPER(fclose)(FILE* stream)
 	FCLOSE_STD(stdin)
 	FCLOSE_STD(stdout)
 	FCLOSE_STD(stderr)
-
-	int fd = fileno(stream);
-	int res = 0;
-	int error = 0;
-	Shared_close(fd, [&]() { res = TRUE_WRAPPER(fclose)(stream); error = errno; });
-	DEBUG_LOG_TRUE("fclose", "(%p) -> %i (%s)", stream, res, StrError(res, error));
-	errno = error;
-	return res;
+	return Shared_close("fclose", fileno(stream), [&]() { return TRUE_WRAPPER(fclose)(stream); });
 }
 
 struct DirInfo
@@ -1497,6 +1706,37 @@ struct DirInfo
 bool IsDirInfo(DIR* dir) { return (u64(dir) & 0x1000'0000'0000'0000) != 0; }
 DirInfo* AsDirInfo(DIR* dir) { return (DirInfo*)(uintptr_t(dir) & ~0x1000'0000'0000'0000); }
 
+UBA_EXPORT int UBA_WRAPPER(mkstemp)(char* templ)
+{
+	UBA_INIT_DETOUR(mkstemp, templ);
+
+	static const char alphabet[] = "abcdefghijklmnopqrstuvwxyzABCDEFGHIJKLMNOPQRSTUVWXYZ0123456789";
+	const uint64_t base = (sizeof(alphabet) - 1);
+
+	DEBUG_LOG_TRUE("mkstemp", "%s", templ);
+
+	auto len = strlen(templ);
+	char* posToWrite = templ + len - 6;
+	while (true)
+	{
+		Guid g;
+		CreateGuid(g);
+		u64 v = *(u64*)&g;
+		for (size_t i = 0; i < 6; ++i)
+		{
+			posToWrite[i] = alphabet[v % base];
+			v /= base;
+		}
+
+		int flags = O_RDWR | O_CREAT | O_EXCL;
+		int fd = UBA_WRAPPER(open)(templ, flags, 0600);
+		if (fd >= 0)
+			return fd;
+		if (errno != EEXIST)
+			return -1;
+	}
+}
+
 UBA_EXPORT DIR* UBA_WRAPPER(opendir)(const char* name)
 {
 	UBA_INIT_DETOUR(opendir, name);
@@ -1563,8 +1803,14 @@ UBA_EXPORT int UBA_WRAPPER(dirfd)(DIR* dirp)
 
 	if (IsDirInfo(dirp))
 	{
-		UBA_ASSERT(false);
-		return 1;
+		// Solution can be to use this and register it with file handles
+		// open("/dev/null", O_RDONLY)
+
+		// Virtual directories have no real file descriptor; return the standard error.
+		// TODO: Might need to fix this
+		DEBUG_LOG_DETOURED("dirfd", "(%p) -> -1 (EBADF)", dirp);
+		errno = EBADF;
+		return -1;
 	}
 
 	int res = TRUE_WRAPPER(dirfd)(dirp);
@@ -1618,7 +1864,62 @@ UBA_EXPORT dirent* UBA_WRAPPER(readdir)(DIR* dirp)
 	return res;
 }
 
+UBA_EXPORT int UBA_WRAPPER(closedir)(DIR* dirp)
+{
+	UBA_INIT_DETOUR(closedir, dirp);
+
+	if (IsDirInfo(dirp))
+	{
+		delete (DirInfo*)AsDirInfo(dirp);
+		DEBUG_LOG_DETOURED("closedir", "(%p)", dirp);
+		return 0;
+	}
+
+	DEBUG_LOG_TRUE("closedir", "(%p)", dirp);
+	return TRUE_WRAPPER(closedir)(dirp);
+}
+
+template<typename Dirent, typename ReadDir>
+int Shared_scandir(const char* func, const char* dirp, Dirent*** namelist, int (*filter)(const Dirent*), int (*compar)(const Dirent**, const Dirent**), ReadDir readDir)
+{
+	DIR* dir = UBA_WRAPPER(opendir)(dirp);
+	if (!dir)
+	{
+		DEBUG_LOG_DETOURED(func, "(%s) -> -1 (%s)", dirp, strerror(errno));
+		return -1;
+	}
+
+	Vector<Dirent*> entries;
+	Dirent* ent;
+	while ((ent = readDir(dir)) != nullptr)
+	{
+		if (filter && !filter(ent))
+			continue;
+		auto copy = (Dirent*)malloc(ent->d_reclen);
+		memcpy(copy, ent, ent->d_reclen);
+		entries.push_back(copy);
+	}
+	UBA_WRAPPER(closedir)(dir);
+
+	if (compar && entries.size() > 1)
+		qsort(entries.data(), entries.size(), sizeof(Dirent*), (int(*)(const void*, const void*))compar);
+
+	*namelist = (Dirent**)malloc(entries.size() * sizeof(Dirent*));
+	for (size_t i = 0; i < entries.size(); i++)
+		(*namelist)[i] = entries[i];
+
+	DEBUG_LOG_DETOURED(func, "(%s) -> %zu", dirp, entries.size());
+	return (int)entries.size();
+}
+
 #if PLATFORM_LINUX
+
+UBA_EXPORT int UBA_WRAPPER(scandir64)(const char* dirp, dirent64*** namelist, int (*filter)(const dirent64*), int (*compar)(const dirent64**, const dirent64**))
+{
+	UBA_INIT_DETOUR(scandir64, dirp, namelist, filter, compar);
+	return Shared_scandir("scandir64", dirp, namelist, filter, compar, [](DIR* d) { return UBA_WRAPPER(readdir64)(d); });
+}
+
 UBA_EXPORT dirent64* UBA_WRAPPER(readdir64)(DIR* dirp)
 {
 	UBA_INIT_DETOUR(readdir64, dirp);
@@ -1641,9 +1942,7 @@ UBA_EXPORT void UBA_WRAPPER(rewinddir)(DIR* dirp)
 UBA_EXPORT int UBA_WRAPPER(scandir)(const char* dirp, dirent*** namelist, int (*filter)(const dirent*), int (*compar)(const dirent**, const dirent**))
 {
 	UBA_INIT_DETOUR(scandir, dirp, namelist, filter, compar);
-	UBA_ASSERTF(!g_runningRemote, "scandir not implemented for remote");
-	DEBUG_LOG_TRUE("scandir", "(%p)", dirp);
-	return TRUE_WRAPPER(scandir)(dirp, namelist, filter, compar);
+	return Shared_scandir("scandir", dirp, namelist, filter, compar, [](DIR* d) { return UBA_WRAPPER(readdir)(d); });
 }
 UBA_EXPORT void UBA_WRAPPER(seekdir)(DIR* dirp, long loc)
 {
@@ -1664,24 +1963,20 @@ UBA_EXPORT long UBA_WRAPPER(telldir)(DIR* dirp)
 UBA_EXPORT DIR* UBA_WRAPPER(fdopendir)(int fd)
 {
 	UBA_INIT_DETOUR(fdopendir, fd);
-	UBA_ASSERTF(false, "fdopendir");
-	DEBUG_LOG_TRUE("fdopendir", "(%i)", fd);
-	return TRUE_WRAPPER(fdopendir)(fd);
-}
 
-UBA_EXPORT int UBA_WRAPPER(closedir)(DIR* dirp)
-{
-	UBA_INIT_DETOUR(closedir, dirp);
-
-	if (IsDirInfo(dirp))
+	SCOPED_READ_LOCK(g_fileHandlesLock, lock);
+	auto findIt = g_fileHandles.find(fd);
+	if (findIt == g_fileHandles.end())
 	{
-		delete (DirInfo*)AsDirInfo(dirp);
-		DEBUG_LOG_DETOURED("closedir", "(%p)", dirp);
-		return 0;
+		DEBUG_LOG_TRUE("fdopendir", "(%i)", fd);
+		return TRUE_WRAPPER(fdopendir)(fd);
 	}
-
-	DEBUG_LOG_TRUE("closedir", "(%p)", dirp);
-	return TRUE_WRAPPER(closedir)(dirp);
+	DetouredHandle& h = findIt->second;
+	lock.Leave();
+	FileObject* fo = h.fileObject;
+	FileInfo& info = *fo->fileInfo;
+	DEBUG_LOG_DETOURED("fdopendir", "(%i) (%s)", info.originalName);
+	return opendir(info.originalName);
 }
 
 UBA_EXPORT int UBA_WRAPPER(glob)(const char* pattern, int flags, int (*errfunc)(const char* epath, int eerrno), glob_t* pglob)
@@ -1722,9 +2017,7 @@ UBA_EXPORT long UBA_WRAPPER(pathconf)(char *path, int name)
 UBA_EXPORT int UBA_WRAPPER(lstat)(const char *path, struct stat *buf)
 {
 	UBA_INIT_DETOUR(lstat, path, buf);
-	int res = TRUE_WRAPPER(lstat)(path, buf);
-	DEBUG_LOG_TRUE("lstat", "%s -> %i", path, res);
-	return res;
+	return Shared_stat("lstat", path, buf, [](const char* path, struct stat* buf) { return TRUE_WRAPPER(lstat)(path, buf); });
 }
 
 UBA_EXPORT off_t  UBA_WRAPPER(lseek)(int fd, off_t offset, int whence)
@@ -1741,6 +2034,27 @@ UBA_EXPORT int UBA_WRAPPER(stat)(const char* file, struct stat* attr)
 	return Shared_stat("stat", file, attr, [](const char* file, struct stat* attr) { return TRUE_WRAPPER(stat)(file, attr); });
 }
 
+#if PLATFORM_LINUX
+UBA_EXPORT int UBA_WRAPPER(stat64)(const char* file, struct stat64* attr)
+{
+	UBA_INIT_DETOUR(stat64, file, attr);
+	return Shared_stat("stat64", file, (struct stat*)attr, [](const char* file, struct stat* attr) { return TRUE_WRAPPER(stat64)(file, (struct stat64*)attr); });
+}
+
+UBA_EXPORT int __xstat64(int ver, const char* file, struct stat64* attr)
+{
+	UBA_INIT_DETOUR(__xstat64, ver, file, attr);
+	UBA_ASSERT(ver == _STAT_VER);
+	return Shared_stat("__xstat64", file, (struct stat*)attr, [](const char* file, struct stat* attr) { return TRUE_WRAPPER(__xstat64)(_STAT_VER, file, (struct stat64*)attr); });
+}
+
+UBA_EXPORT int UBA_WRAPPER(lstat64)(const char* file, struct stat64* attr)
+{
+	UBA_INIT_DETOUR(lstat64, file, attr);
+	return Shared_stat("lstat64", file, (struct stat*)attr, [](const char* file, struct stat* attr) { return TRUE_WRAPPER(lstat64)(file, (struct stat64*)attr); });
+}
+#endif
+
 UBA_EXPORT int UBA_WRAPPER(truncate)(const char* path, off_t length)
 {
 	UBA_INIT_DETOUR(truncate, path, length);
@@ -1796,6 +2110,18 @@ UBA_EXPORT int UBA_WRAPPER(__fxstat)(int ver, int fd, struct stat* attr)
 	return Shared_fstat("__fxstat", fd, attr, [ver](int fd, struct stat* attr) { return TRUE_WRAPPER(__fxstat)(ver, fd, attr); });
 }
 
+UBA_EXPORT int UBA_WRAPPER(__lxstat64)(int ver, const char* file, struct stat64* attr)
+{
+	static_assert(sizeof(struct stat64) == sizeof(struct stat));
+	UBA_INIT_DETOUR(__lxstat64, ver, file, attr);
+	return Shared_stat("__lxstat64", file, (struct stat*)attr, [ver](const char* file, struct stat* attr) { return TRUE_WRAPPER(__lxstat64)(ver, file, (struct stat64*)attr); });
+}
+
+UBA_EXPORT int UBA_WRAPPER(statx)(int dirfd, const char* pathname, int flags, unsigned int mask, struct statx* statxbuf)
+{
+	UBA_INIT_DETOUR(statx, dirfd, pathname, flags, mask, statxbuf);
+	return Shared_statx("statx", dirfd, pathname, flags, mask, statxbuf);
+}
 #endif
 
 UBA_EXPORT int UBA_WRAPPER(rename)(const char* oldpath, const char* newpath)
@@ -1812,6 +2138,15 @@ UBA_EXPORT int UBA_WRAPPER(rename)(const char* oldpath, const char* newpath)
 	DevirtualizePath(fixedNewPath);
 	StringKey newKey = ToFilenameKey(fixedNewPath);
 
+	DirTableOffset dirTableOffset;
+	if (!Rpc_GetEntryOffset(dirTableOffset, oldKey, fixedOldPath, false))
+	{
+		errno = ENOENT;
+		return -1;
+	}
+	DirectoryTable::EntryInformation info;
+	g_directoryTable.GetEntryInformation(info, dirTableOffset);
+
 	// TODO: This might be really slow but it seems you can rename files on linux while they are open and they won't be properly renamed until closed
 	{
 		SCOPED_READ_LOCK(g_fileHandlesLock, lock);
@@ -1872,6 +2207,9 @@ UBA_EXPORT int UBA_WRAPPER(rename)(const char* oldpath, const char* newpath)
 		RPC_MESSAGE(MoveFile, moveFile);
 		writer.WriteStringKey(oldKey);
 		writer.WriteString(fixedOldPath);
+		writer.WriteU64(info.size);
+		writer.WriteU64(info.lastWrite);
+		writer.WriteU32(info.attributes);
 		writer.WriteStringKey(newKey);
 		writer.WriteString(fixedNewPath);
 		writer.WriteU32(0);// dwFlags);
@@ -1889,14 +2227,14 @@ UBA_EXPORT int UBA_WRAPPER(rename)(const char* oldpath, const char* newpath)
 	{
 		SCOPED_WRITE_LOCK(g_mappedFileTable.m_lookupLock, _);
 		auto insres = g_mappedFileTable.m_lookup.try_emplace(newKey);
-		FileInfo& info = insres.first->second;
+		FileInfo& fi = insres.first->second;
 
 		auto findIt = g_mappedFileTable.m_lookup.find(oldKey);
 		if (findIt != g_mappedFileTable.m_lookup.end())
-			info = findIt->second;
-		info.originalName = g_memoryBlock.Strdup(fixedNewPath).data;
-		info.name = info.originalName;
-		info.created = true;
+			fi = findIt->second;
+		fi.originalName = g_memoryBlock.Strdup(fixedNewPath).data;
+		fi.name = fi.originalName;
+		fi.created = true;
 	}
 
 	g_directoryTable.ParseDirectoryTable(directoryTableSize);
@@ -1949,80 +2287,88 @@ UBA_EXPORT int UBA_WRAPPER(renameat)(int olddirfd, const char* oldpath, int newd
 {
 	UBA_INIT_DETOUR(renameat, olddirfd, oldpath, newdirfd, newpath);
 	DEBUG_LOG_TRUE("renameat", "(from %s to %s)", oldpath, newpath);
-	UBA_ASSERTF(false, "Not implemented");
-	return TRUE_WRAPPER(renameat)(olddirfd, oldpath, newdirfd, newpath);
+	UBA_ASSERT(olddirfd == AT_FDCWD);
+	UBA_ASSERT(newdirfd == AT_FDCWD);
+	return rename(oldpath, newpath);
 }
 
 UBA_EXPORT int UBA_WRAPPER(utimensat)(int dirfd, const char* pathname, const struct timespec* times, int flags)
 {
 	UBA_INIT_DETOUR(utimensat, dirfd, pathname, times, flags);
 	DEBUG_LOG_TRUE("utimensat", "(%s)", pathname);
-	UBA_ASSERTF(false, "Not implemented");
+	UBA_ASSERT(dirfd== AT_FDCWD);
 	return TRUE_WRAPPER(utimensat)(dirfd, pathname, times, flags);
 }
 
-UBA_EXPORT int UBA_WRAPPER(symlink)(const char* path1, const char* path2)
+UBA_EXPORT int UBA_WRAPPER(symlink)(const char* symlinkContent, const char* symlinkFile)
 {
-	UBA_INIT_DETOUR(symlink, path1, path2);
+	UBA_INIT_DETOUR(symlink, symlinkContent, symlinkFile);
+
+	StringBuffer<> symlinkFilePath;
+	FixPath(symlinkFilePath, symlinkFile);
+	DevirtualizePath(symlinkFilePath);
 
-	StringBuffer<> fromPath;
-	FixPath(fromPath, path1);
-	DevirtualizePath(fromPath);
-	StringBuffer<> toPath;
-	FixPath(toPath, path2);
-	DevirtualizePath(toPath);
+	StringBuffer<> symlinkContentPath;
+	if (IsAbsolutePath(symlinkContent))
+		FixPath(symlinkContentPath, symlinkContent);
+	else
+		FixPath(symlinkContentPath, ToView(symlinkContent), ToView(symlinkFilePath).GetPath(true));
+
+	DevirtualizePath(symlinkContentPath);
+	if (IsAbsolutePath(symlinkContent))
+		symlinkContent = symlinkContentPath.data;
 
-	if (!CanDetour2(fromPath) && !CanDetour2(toPath))
+	if (!CanDetour2(symlinkContentPath) && !CanDetour2(symlinkFilePath))
 	{
-		DEBUG_LOG_TRUE("symlink", "(from %s to %s)", path1, path2);
-		return TRUE_WRAPPER(symlink)(path1, path2);
+		DEBUG_LOG_TRUE("symlink", "(file: %s content: %s)", symlinkFile, symlinkContent);
+		return TRUE_WRAPPER(symlink)(symlinkContent, symlinkFile);
 	}
 
-	StringKey fromKey = ToStringKeyLower(fromPath);
-	StringKey toKey = ToStringKeyLower(toPath);
+	StringKey contentPathKey = ToFilenameKey(symlinkContentPath);
+	StringKey filePathKey = ToFilenameKey(symlinkFilePath);
 
-	StringBuffer<> newFromName;
-	StringBuffer<> newToName;
-	u32 closeId;
-	u32 lastError;
-	DirTableSize directoryTableSize;
+	// This is wrong since the symlink target (content) does not need to exist. But for now we make that assumption
+	// Adding support for "real" symlinks require quite a lot of changes in directory table etc
+	DirTableOffset dirTableOffset;
+	if (!Rpc_GetEntryOffset(dirTableOffset, contentPathKey, symlinkContentPath, false))
 	{
-		RPC_MESSAGE(CopyFile, copyFile)
-		writer.WriteStringKey(fromKey);
-		writer.WriteString(fromPath);
-		writer.WriteStringKey(toKey);
-		writer.WriteString(toPath);
-		BinaryReader reader = writer.Flush();
-		reader.ReadString(newFromName);
-		reader.ReadString(newToName);
-		closeId = reader.ReadU32();
-		lastError = reader.ReadU32();
-		directoryTableSize = FromU64(reader.ReadU64());
+		errno = ENOENT;
+		DEBUG_LOG_DETOURED("symlink", "(DIRTABLE) THIS IS WRONG, NEED FIX (file: %s content: %s) -> -1 (ENOENT)", symlinkFilePath.data, symlinkContentPath.data);
+		return -1;
 	}
+	DirectoryTable::EntryInformation info;
+	g_directoryTable.GetEntryInformation(info, dirTableOffset);
 
-	UBA_ASSERT(closeId != ~0u);
+	if (!info.attributes)
+	{
+		DEBUG_LOG("symlink - TODO: This is wrong. Need to add support in dirtable for symlinks! (file: %s content: %s)", symlinkFilePath.data, symlinkContentPath.data);
+		info.attributes = S_IFLNK;
+	}
 
-	//UBA_ASSERTF(false, "symlink not implemented (from %s to %s)", path1, path2);
-	DEBUG_LOG_DETOURED("symlink", "(from %s to %s)", path1, path2);
-	int res = TRUE_WRAPPER(symlink)(newFromName.data, newToName.data);
 
-	if (!closeId)
-		return res;
+	RPC_MESSAGE(Symlink, copyFile)
+	writer.WriteStringKey(filePathKey);
+	writer.WriteString(symlinkFilePath);
+	writer.WriteStringKey(contentPathKey);
+	writer.WriteString(symlinkContentPath);
+	writer.WriteString(symlinkContent);
+	writer.WriteFileSize(info.size);
+	writer.WriteFileTime(info.lastWrite);
+	writer.WriteFileAttributes(info.attributes);
+	BinaryReader reader = writer.Flush();
+	bool success = reader.ReadBool();
+	u32 error = reader.ReadU32();
+	DirTableSize directoryTableSize = FromU64(reader.ReadU64());
+	RPC_MESSAGE_END
 
-	bool deleteOnClose = res != 0; // If failing to copy we set deleteOnClose
-	Rpc_UpdateCloseHandle(newToName.data, closeId, deleteOnClose, "", {}, 0, true);
+	g_directoryTable.ParseDirectoryTable(directoryTableSize);
 
-	return res;
+	errno = error;
+	return success ? 0 : -1;
 }
 
 UBA_EXPORT ssize_t UBA_WRAPPER(pread)(int __fd, void * __buf, size_t __nbyte, off_t __offset)
 {
-	// char filePath[PATH_MAX];
-	// if (fcntl(__fd, F_GETPATH, filePath) != -1)
-	// {
-	// 	// printf("***** PREAD: %d %s bytes: 0x%lx\n",__fd, filePath, __nbyte);
-	// 	// do something with the file path
-	// }
 	return TRUE_WRAPPER(pread)(__fd, __buf, __nbyte, __offset);
 }
 
@@ -2035,7 +2381,9 @@ UBA_EXPORT ssize_t UBA_WRAPPER(read)(int fd, void *buf, size_t nbyte)
 	// 	// printf("***** READ: %d %s bytes: %lu\n",fd, filePath, nbyte);
 	// 	// do something with the file path
 	// }
-	return TRUE_WRAPPER(read)(fd, buf, nbyte);
+	ssize_t res = TRUE_WRAPPER(read)(fd, buf, nbyte);
+	DEBUG_LOG_TRUE("read", "%u %llu -> %llu", fd, nbyte, res);
+	return res;
 }
 
 int Shared_DeleteFile(const char* funcName, const char* pathname)
@@ -2051,26 +2399,26 @@ int Shared_DeleteFile(const char* funcName, const char* pathname)
 		return res;
 	}
 
-
 	StringKey fileNameKey = ToFilenameKey(fixedName);
 
-	DirTableSize directoryTableSize;
-	bool result;
-	u32 errorCode;
+	DirTableOffset dirTableOffset;
+	if (!Rpc_GetEntryOffset(dirTableOffset, fileNameKey, fixedName, false))
 	{
-		u32 closeId = 0;
-		RPC_MESSAGE(DeleteFile, deleteFile);
-		writer.WriteString(fixedName);
-		writer.WriteStringKey(fileNameKey);
-		writer.WriteU32(closeId);
-		BinaryReader reader = writer.Flush();
-		result = reader.ReadBool();
-		errorCode = reader.ReadU32();
-		directoryTableSize = FromU64(reader.ReadU64());
-		pcs.Leave();
-		DEBUG_LOG_PIPE(L"DeleteFile", L"%s", lpFileName);
+		DEBUG_LOG_DETOURED(funcName, "(%s) -> Error (ENOENT)", pathname);
+		errno = ENOENT;
+		return -1;
 	}
-	//DEBUG_LOG_DETOURED(L"DeleteFile", L"(%ls) -> %ls", pathname, ToString(result));
+
+	u32 closeId = 0;
+	RPC_MESSAGE(DeleteFile, deleteFile);
+	writer.WriteString(fixedName);
+	writer.WriteStringKey(fileNameKey);
+	writer.WriteU32(closeId);
+	BinaryReader reader = writer.Flush();
+	bool result = reader.ReadBool();
+	u32 errorCode = reader.ReadU32();
+	DirTableSize directoryTableSize = FromU64(reader.ReadU64());
+	RPC_MESSAGE_END
 
 	g_directoryTable.ParseDirectoryTable(directoryTableSize);
 	g_mappedFileTable.SetDeleted(fileNameKey, fixedName.data, true);
@@ -2118,12 +2466,40 @@ UBA_EXPORT int UBA_WRAPPER(unlink)(const char* pathname)
 	return Shared_DeleteFile("unlink", pathname);
 }
 
+const char* ResolvePath(int dirfd, const char* pathname, StringBuffer<>& temp)
+{
+	if (pathname && *pathname == '/')
+		return pathname;
+	if (dirfd == AT_FDCWD)
+		return pathname;
+	SCOPED_WRITE_LOCK(g_fileHandlesLock, lock);
+	auto findIt = g_fileHandles.find(dirfd);
+	if (findIt == g_fileHandles.end())
+		return nullptr;
+	DetouredHandle& h = findIt->second;
+	lock.Leave();
+	FileObject* fo = h.fileObject;
+	FileInfo& info = *fo->fileInfo;
+	temp.Append(info.originalName);
+	if (pathname)
+		temp.EnsureEndsWithSlash().Append(pathname);
+	return temp.data;
+}
+
 UBA_EXPORT int UBA_WRAPPER(unlinkat)(int dirfd, const char *pathname, int flags)
 {
 	UBA_INIT_DETOUR(unlinkat, dirfd, pathname, flags);
-	UBA_ASSERT(!g_runningRemote);
-	UBA_ASSERT(dirfd == AT_FDCWD);
-	return Shared_DeleteFile("unlinkat", pathname);
+
+	StringBuffer<> temp;
+	if (const char* resolvedPath = ResolvePath(dirfd, pathname, temp))
+	{
+		if (flags == AT_REMOVEDIR)
+			return rmdir(resolvedPath);
+		UBA_ASSERT(flags == 0);
+		return Shared_DeleteFile("unlinkat", resolvedPath);
+	}
+	DEBUG_LOG_TRUE("unlinkat", "(%i) %s", dirfd, pathname)
+	return TRUE_WRAPPER(unlinkat)(dirfd, pathname, flags);
 }
 
 void FlattenArgs(StringBufferBase& out, const char* const argv[])
@@ -2218,13 +2594,17 @@ struct ProcessCreator
 			u32 plen = result.count;
 			if (plen >= 3 && result.data[plen-3] == '.' && result.data[plen-2] == 's' && result.data[plen-1] == 'h')
 			{
-				u32 argc2 = 0;
-				tempArgv[argc2++] = "/bin/sh";
-				tempArgv[argc2++] = path; // the .sh script (still in result.data)
+				u32 count = 0;
+				for (int i = 1; argv[i]; ++i)
+					++count;
+				tempArgv.reserve(count + 4);
+
+				tempArgv.push_back("/bin/sh");
+				tempArgv.push_back(path); // the .sh script (still in result.data)
 				for (int i = 1; argv[i]; ++i)
-					tempArgv[argc2++] = argv[i];
-				tempArgv[argc2] = nullptr;
-				fixedArgv = (char* const*)tempArgv;
+					tempArgv.push_back(argv[i]);
+				tempArgv.push_back(nullptr);
+				fixedArgv = (char* const*)tempArgv.data();
 				path = "/bin/sh";
 			}
 		}
@@ -2261,68 +2641,71 @@ struct ProcessCreator
 				cmdLineWithoutApplication.append(arg);
 		}
 
-		{
-			RPC_MESSAGE2(messageType, messageName, createProcess);
+		RPC_MESSAGE2(messageType, messageName, createProcess);
+		writer.WriteString(path); // Application
+		writer.WriteLongString(cmdLineWithoutApplication, CommunicationMemSize - 4096);
+		writer.WriteString(g_virtualWorkingDir); // Current dir
+		writer.WriteBool(false); // Start suspended
+		writer.WriteBool(true); // Is child
 
-			const char* pwd = "";
-#if 0
-			for (char* const* env = envp; *env;++env)
-			{
-				if (StartsWith(*env, "PWD="))
-				{
-					printf("%s\n", *env);
-					pwd = *env + 4;
-					break;
-				}
-			}
-#endif
-
-			writer.WriteString(path); // Application
-			writer.WriteLongString(cmdLineWithoutApplication, CommunicationMemSize - 4096);
-			writer.WriteString(pwd); // Current dir
-			writer.WriteBool(false); // Start suspended
-			writer.WriteBool(true); // Is child
-			BinaryReader reader = writer.Flush();
+		BinaryReader reader = writer.Flush();
+		processId = reader.ReadU32();
 
-			processId = reader.ReadU32();
+		if (!processId) // Can happen if session client got disconnected from server session
+		{
+			errno = EINVAL; // This is not really correct but there is no errno for this failure
+			return -1;
+		}
 
-			if (!processId) // Can happen if session client got disconnected from server session
-			{
-				errno = EINVAL; // This is not really correct but there is no errno for this failure
-				return -1;
-			}
+		StringBuffer<1024> temp;
 
-			rulesStr.Append("UBA_RULES=").AppendValue(reader.ReadU32());
+		rulesStr = temp.Append("UBA_RULES=").AppendValue(reader.ReadU32()).ToStringAndClear();
 
-			while (true)
-			{
-				u32 dllNameSize = reader.ReadU32();
-				if (!dllNameSize)
-					break;
-				reader.Skip(dllNameSize);
-			}
 
-			currentDir.Append("UBA_CWD=");
-			reader.ReadString(currentDir);
+		while (true)
+		{
+			u32 dllNameSize = reader.ReadU32();
+			if (!dllNameSize)
+				break;
+			reader.Skip(dllNameSize);
+		}
 
-			realApplication = reader.ReadString();
+		currentDir = temp.Append("UBA_CWD=").ReadString(reader).ToStringAndClear();
+		realApplication = reader.ReadString();
 
-			#if UBA_USE_NATIVE_MAC_SEMAPHORES
-			pidVar.Append("UBA_PID=").AppendValue(reader.ReadU32());
-			#endif
+		#if UBA_USE_NATIVE_MAC_SEMAPHORES
+		pidVar = temp.Append("UBA_PID=").AppendValue(reader.ReadU32()).ToStringAndClear();
+		#endif
 
-			comIdVar.Append("UBA_COMID=").AppendValue(reader.ReadU64()).Append('+').AppendValue(reader.ReadU32());
+		comIdVar = temp.Append("UBA_COMID=").AppendValue(reader.ReadU64()).Append('+').AppendValue(reader.ReadU32()).ToStringAndClear();
+		logFile = temp.Append("UBA_LOGFILE=").ReadString(reader).ToStringAndClear();
+		RPC_MESSAGE_END
 
-			logFile.Append("UBA_LOGFILE=");
-			reader.ReadString(logFile);
-		}
+		constexpr auto ldLibraryPathName = IsLinux ? TCV("LD_LIBRARY_PATH=") : TCV("DYLD_LIBRARY_PATH=");
 
 		for (u32 i = 0; envp[i]; ++i)
+		{
+			if (g_runningRemote && StartsWith(envp[i], ldLibraryPathName.data))
+			{
+				UBA_ASSERT(!temp.count);
+				if (char* colon = strchr(envp[i], ':'))
+					temp.Append(envp[i], colon - envp[i]);
+				else
+					temp.Append(envp[i]);
+				temp.Append(':').Append(ToView(realApplication).GetPath());
+				ldLibraryPath = temp.ToStringAndClear();
+				continue;
+			}
 			envvars.push_back(envp[i]);
-		envvars.push_back(comIdVar.data);
-		envvars.push_back(currentDir.data);
-		envvars.push_back(rulesStr.data);
-		envvars.push_back(logFile.data);
+		}
+		UBA_ASSERT(!g_runningRemote || !ldLibraryPath.empty());
+
+		envvars.push_back(comIdVar.data());
+		envvars.push_back(currentDir.data());
+		envvars.push_back(rulesStr.data());
+		envvars.push_back(logFile.data());
+		if (!ldLibraryPath.empty())
+			envvars.push_back(ldLibraryPath.data());
 		if (isExec)
 			envvars.push_back("UBA_EXEC=1");
 		//envvars.push_back("LD_DEBUG=bindings");
@@ -2346,17 +2729,18 @@ struct ProcessCreator
 	TString realApplication;
 	TString cmdLineWithoutApplication;
 	u32 processId = 0;
-	StringBuffer<512> currentDir;
-	StringBuffer<256> comIdVar;
-	StringBuffer<32> rulesStr;
-	StringBuffer<512> logFile;
+	TString currentDir;
+	TString comIdVar;
+	TString rulesStr;
+	TString logFile;
+	TString ldLibraryPath;
 	#if UBA_USE_NATIVE_MAC_SEMAPHORES
-	StringBuffer<64> pidVar;
+	TString pidVar;
 	#endif
 	StringBuffer<> result;
-	const char* tempArgv[1024];
+	Vector<const char*> tempArgv;
 	char* const* fixedArgv;
-	std::vector<const char*> envvars;
+	Vector<const char*> envvars;
 };
 
 void SendStartProcess(u32 processId, u32 pid, bool success, u32 res)
@@ -2454,7 +2838,7 @@ int shared_posix_spawn(pid_t* pid, const char* path, const posix_spawn_file_acti
 		return res;
 	
 	#if UBA_DEBUG_LOG_ENABLED
-	DEBUG_LOG_TRUE("posix_spawn", "%s (%s)", pc.realApplication.data(), pc.logFile.data);
+	DEBUG_LOG_TRUE("posix_spawn", "%s (%s)", pc.realApplication.data(), pc.logFile.data());
 	for (u32 i = 0; pc.fixedArgv[i]; ++i)
 		DEBUG_LOG("            %s", pc.fixedArgv[i]);
 	#endif
@@ -2596,7 +2980,7 @@ UBA_EXPORT void* UBA_WRAPPER(dlopen)(const char* path, int mode)
 			const char* loaderPaths[] = { "/", 0 };
 			if (g_handledLibraries.insert(path).second)
 				Shared_LoadLibrary(path, loaderPaths, tempBuf);
-			if (const char* name = TStrrchr(path, '/'))
+			if (const char* name = TStrrchr(path, '/')) // We want to use LD_LIBRARY_PATH to find the .so file after downloading it
 				path = name + 1;
 		}
 	#endif
@@ -2614,7 +2998,7 @@ UBA_EXPORT int UBA_WRAPPER(dladdr)(const void *addr, Dl_info *info)
 	if (StartsWith(info->dli_fname, g_exeDir.data))
 	{
 		StringBuffer<> newPath;
-		newPath.Append(g_virtualApplication.data).Append(info->dli_fname + g_exeDir.count - 1); // TODO: Make g_handledLibraries a map and store the real path
+		newPath.AppendDir(g_virtualApplication).Append('/').Append(info->dli_fname + g_exeDir.count - 1); // TODO: Make g_handledLibraries a map and store the real path
 		info->dli_fname = g_memoryBlock.Strdup(newPath).data;
 	}
 	DEBUG_LOG_TRUE("dladdr", "%s -> %i", info->dli_fname, res);
@@ -2785,75 +3169,83 @@ UBA_EXPORT pid_t UBA_WRAPPER(vfork)(void)
 	DEBUG_LOG_DETOURED("vfork", "");
 	return Internal_fork();
 }
-/*
-UBA_EXPORT int UBA_WRAPPER(clone)(int (*fn)(void *), void *stack, int flags, void *args, ...)
-{
-	va_list ap;
-	va_start(ap, args);
-
-	UBA_INIT_DETOUR(clone, fn, stack, flags, va_arg(ap, void *), va_arg(ap, void *), va_arg(ap, void *));
 
-	pid_t *parent_tid = NULL;
-	void *tls = NULL;
-	pid_t *child_tid = NULL;
+#if PLATFORM_LINUX
+UBA_EXPORT long UBA_WRAPPER(syscall)(long number, ...)
+{
+	va_list args;
+	va_start(args, number);
+	void* a1 = va_arg(args, void*);
+	void* a2 = va_arg(args, void*);
+	void* a3 = va_arg(args, void*);
+	void* a4 = va_arg(args, void*);
+	void* a5 = va_arg(args, void*);
+	void* a6 = va_arg(args, void*);
+	va_end(args);
 
-	if (flags & CLONE_PARENT_SETTID)
-		parent_tid = va_arg(ap, pid_t *);
+	UBA_INIT_DETOUR(syscall, number, a1, a2, a3, a4, a5, a6);
 
-	if (flags & CLONE_SETTLS)
-		tls = va_arg(ap, void *);
+	switch (number)
+	{
+	case SYS_write:
+	{
+		DEBUG_LOG_TRUE("syscall.write", "fd=%ld", (long)(uintptr_t)a1);
+		return TRUE_WRAPPER(syscall)(number, a1, a2, a3, a4, a5, a6);
+	}
 
-	if (flags & CLONE_CHILD_CLEARTID)
-		child_tid = va_arg(ap, pid_t *);
+	case SYS_close:
+	{
+		int fd = (int)(long)(uintptr_t)a1;
+		return Shared_close("syscall.close", fd, [fd]() { return TRUE_WRAPPER(close)(fd); });
+	}
 
-	va_end(ap);
+	case SYS_gettid:
+	case SYS_capget:
+	case SYS_futex:
+		break;
 
-	DEBUG_LOG_TRUE("CLONE!!!", "");
-	//fprintf(stderr, "clone intercepted: flags=0x%x\n", flags);
+#if !PLATFORM_LINUX || !PLATFORM_CPU_ARM_FAMILY
+	case SYS_vfork:
+		{ DEBUG_LOG_DETOURED("syscall.vfork", "");
+		return Internal_fork(); }
+	case SYS_fork:
+		{ DEBUG_LOG_DETOURED("syscall.fork", "");
+		return Internal_fork(); }
+#endif
 
-	if ((flags & CLONE_PARENT_SETTID) && (flags & CLONE_SETTLS) && (flags & CLONE_CHILD_CLEARTID))
-		return TRUE_WRAPPER(clone)(fn, stack, flags, args, parent_tid, tls, child_tid);
-	if ((flags & CLONE_PARENT_SETTID) && (flags & CLONE_SETTLS))
-		return TRUE_WRAPPER(clone)(fn, stack, flags, args, parent_tid, tls);
-	if ((flags & CLONE_PARENT_SETTID))
-		return TRUE_WRAPPER(clone)(fn, stack, flags, args, parent_tid);
-	return TRUE_WRAPPER(clone)(fn, stack, flags, args);
-}
+	case SYS_clone:
+		{ DEBUG_LOG_DETOURED("syscall.clone", "");
+		return Internal_fork(); }
 
-UBA_EXPORT long UBA_WRAPPER(clone3)(struct clone_args *cl_args, size_t size)
-{
-	UBA_INIT_DETOUR(clone3, cl_args, size);
-	DEBUG_LOG_TRUE("CLONE3!!!", "");
-	long ret = TRUE_WRAPPER(clone3)(cl_args, size);
-	if (ret == 0)
+	case 217: // SYS_getdents64
 	{
-		#if PLATFORM_LINUX
-		prctl(PR_SET_PDEATHSIG, SIGHUP, 0, 0, 0); // We want the process to die if the parent die
-		#endif
+		UBA_NOT_IMPLEMENTED("syscall.getdents64")
+		break;
 	}
-	return ret;
-}
-*/
-#if PLATFORM_LINUX
-UBA_EXPORT long UBA_WRAPPER(syscall)(long number, ...)
-{
-	va_list args;
-	va_start(args, number);
 
-	UBA_INIT_DETOUR(syscall, number, va_arg(args, void *), va_arg(args, void *), va_arg(args, void *), va_arg(args, void *), va_arg(args, void *));
-	DEBUG_LOG_TRUE("syscall", "%ld", number);
+	case SYS_statx:
+		return Shared_statx("syscall.statx", (int)(long)(uintptr_t)a1, (const char*)a2, (int)(long)(uintptr_t)a3, (unsigned int)(uintptr_t)a4, (struct statx*)a5);
 
-	bool isClone = number == SYS_clone;
-	#ifdef SYS_clone3
-	isClone |= number == SYS_clone3;
-	#endif
+	case 425: // SYS_io_uring_setup
+		errno = ENOSYS;
+		return -1;
 
-	if (isClone)
-		return Internal_fork();
+	case 435: // SYS_clone3
+		UBA_NOT_IMPLEMENTED("syscall.clone3");
+		break;
 
-	long ret = TRUE_WRAPPER(syscall)(number, va_arg(args, void *), va_arg(args, void *), va_arg(args, void *), va_arg(args, void *), va_arg(args, void *));
-	va_end(args);
-	return ret;
+	case 436: // SYS_close_range
+		UBA_NOT_IMPLEMENTED("syscall.close_range")
+		break;
+
+#if UBA_DEBUG_LOG_ENABLED
+	default:
+		DEBUG_LOG_TRUE("syscall", "%ld", number);
+		break;
+#endif
+	}
+
+	return TRUE_WRAPPER(syscall)(number, a1, a2, a3, a4, a5, a6);
 }
 #endif
 
@@ -2956,33 +3348,44 @@ UBA_EXPORT int UBA_WRAPPER(system)(const char* command)
 	return TRUE_WRAPPER(system)(command);
 }
 
-/*
+#if PLATFORM_LINUX
 
-UBA_EXPORT int __fxstat64(int ver, int fd, struct stat64* attr)
+UBA_EXPORT int openat64(int dirfd, const char* pathname, int flags, ...)
 {
-	UBA_NOT_IMPLEMENTED(__fxstat64);
-	return TRUE_WRAPPER(__fxstat64)(ver, fd, attr);
+	va_list args;
+	va_start(args, flags);
+	int mode = va_arg(args, int);
+	va_end(args);
+	UBA_INIT_DETOUR(openat64, dirfd, pathname, flags, mode);
 
+	StringBuffer<> temp;
+	if (const char* resolvedPath = ResolvePath(dirfd, pathname, temp))
+		return Shared_open("openat64", resolvedPath, flags, mode, [](const char* realFile, int flags, int mode) { return TRUE_WRAPPER(open64)(realFile, flags, mode); });
+
+	UBA_ASSERT(!g_runningRemote);
+
+	int res = TRUE_WRAPPER(openat64)(dirfd, pathname, flags, mode);
+	DEBUG_LOG_TRUE("openat64", "(NODETOUR) %s -> %i", pathname, res);
+	if (res)
+		TrackInput(ToView(pathname));
+	return res;
 }
-UBA_EXPORT int __xstat64(int ver, const char* file, struct stat64* attr)
-{
-	UBA_NOT_IMPLEMENTED(__xstat64);
-	return TRUE_WRAPPER(__xstat64)(ver, file, attr);
-}
-UBA_EXPORT int openat64(int __fd, const char* __file, int __oflag, ...)
+
+UBA_EXPORT int UBA_WRAPPER(dup3)(int oldfd, int newfd, int flags)
 {
-	UBA_NOT_IMPLEMENTED(openat64);
-	return -1;
+	UBA_INIT_DETOUR(dup3, oldfd, newfd, flags);
+	auto res = TRUE_WRAPPER(dup3)(oldfd, newfd, flags);
+	if (res != -1)
+		DuplicateFd(oldfd, newfd);
+	return res;
 }
 
-UBA_EXPORT int openat64(int __fd, const char* __file, int __oflag, mode_t mode, ...)
+UBA_EXPORT pid_t UBA_WRAPPER(clone)(int (*__fn) (void *__arg), void *__child_stack, int __flags, void *__arg, ...)
 {
-	UBA_NOT_IMPLEMENTED(open64);
-	return -1;
+	DEBUG_LOG_TRUE("clone", "");
+	UBA_ASSERT(false);
+	return -1;//TRUE_WRAPPER(clone)();
 }
-*/
-
-#if PLATFORM_LINUX
 
 UBA_EXPORT char* UBA_WRAPPER(get_current_dir_name)(void)
 {
@@ -3203,17 +3606,23 @@ namespace uba
 			SCOPED_WRITE_LOCK(g_fileHandlesLock, lock);
 			for (auto& kv : g_fileHandles)
 			{
-				TRUE_WRAPPER(close)(kv.first);
+				int fd = kv.first;
 				DetouredHandle& h = kv.second;
 				FileObject* fo = h.fileObject;
-				if (!fo->closeId)
-					continue;
-				TRUE_WRAPPER(close)(kv.first);
-				SharedMemoryHandle memoryHandle;
-				u64 memoryWritten = 0;
-				FileInfo& fi = *fo->fileInfo;
-				const tchar* path = fi.name;
-				Rpc_UpdateCloseHandle(path, fo->closeId, fo->deleteOnClose, fo->newName.c_str(), memoryHandle, memoryWritten, true);
+				if (fo->closeId)
+				{
+					SharedMemoryHandle memoryHandle;
+					struct stat attr;
+					int res = TRUE_WRAPPER(fstat)(fd, &attr);
+					UBA_ASSERT(res == 0);(void)res;
+					u64 fileSize = attr.st_size;
+					u64 fileLastWriteTime = FromTimeSpec(attr.st_mtimespec);
+					u32 fileAttributes = attr.st_mode;
+					FileInfo& fi = *fo->fileInfo;
+					const tchar* path = fi.name;
+					Rpc_UpdateCloseHandle(path, fo->closeId, fo->deleteOnClose, fo->newName.c_str(), memoryHandle, fileSize, fileLastWriteTime, fileAttributes, true);
+				}
+				TRUE_WRAPPER(close)(fd);
 			}
 		}
 
diff --git a/Engine/Source/Programs/UnrealBuildAccelerator/Detours/Private/UbaDetoursFileMappingTable.cpp b/Engine/Source/Programs/UnrealBuildAccelerator/Detours/Private/UbaDetoursFileMappingTable.cpp
--- a/Engine/Source/Programs/UnrealBuildAccelerator/Detours/Private/UbaDetoursFileMappingTable.cpp
+++ b/Engine/Source/Programs/UnrealBuildAccelerator/Detours/Private/UbaDetoursFileMappingTable.cpp
@@ -150,18 +150,22 @@ namespace uba
 		return tableOffset != InvalidTableOffset;
 	}
 
-	void Rpc_UpdateCloseHandle(const tchar* handleName, u32 closeId, bool deleteOnClose, const tchar* newName, SharedMemoryHandle memoryHandle, u64 memoryWritten, bool success)
+	void Rpc_UpdateCloseHandle(const tchar* handleName, u32 closeId, bool deleteOnClose, const tchar* newName, SharedMemoryHandle memoryHandle, u64 fileSize, u64 lastWriteTime, u32 attributes, bool success)
 	{
+		UBA_ASSERT(fileSize != InvalidValue);
+
 		DirTableSize directoryTableSize;
 		{
 			RPC_MESSAGE(CloseFile, closeFile)
 			writer.WriteString(handleName);
 			writer.WriteU32(closeId);
-			//writer.WriteU32(attributes); // TODO THIS
 			writer.WriteBool(deleteOnClose);
 			writer.WriteBool(success);
 			writer.WriteU64(memoryHandle.ToU64());
-			writer.WriteU64(memoryWritten);
+			writer.WriteU64(fileSize);
+			writer.WriteU64(lastWriteTime);
+			writer.WriteU32(attributes);
+
 			if (*newName)
 			{
 				StringBuffer<> fixedName;
@@ -364,18 +368,15 @@ namespace uba
 		if (exists == DirectoryTable::Exists_Yes)
 			return true;
 
-		const tchar* lastPathSeparator = TStrrchr(entryName.data, PathSeparator);
-		if (!lastPathSeparator)
-		{
-			UBA_ASSERTF(lastPathSeparator, TC("No path separator found in %s"), entryName.count > 0 ? entryName.data : TC("(NULL)"));
-			return false;
-		}
-
-		#if PLATFORM_WINDOWS
-		UBA_ASSERT(wcsncmp(entryName.data, g_systemTemp.data, g_systemTemp.count) != 0);
-		#endif
+		u32 dirNameLen = entryName.count;
+		if (const tchar* lastPathSeparator = TStrrchr(entryName.data, PathSeparator))
+			dirNameLen = u32(lastPathSeparator - entryName.data);
+		//else
+		//{
+		//	UBA_ASSERTF(lastPathSeparator, TC("No path separator found in %s"), entryName.count > 0 ? entryName.data : TC("(NULL)"));
+		//	return false;
+		//}
 
-		u32 dirNameLen = u32(lastPathSeparator - entryName.data);
 		DirHash hash(StringView(entryNameForKey.data, dirNameLen));
 
 		if (!Rpc_UpdateDirectory(hash.key, entryName.data, dirNameLen))
diff --git a/Engine/Source/Programs/UnrealBuildAccelerator/Detours/Private/UbaDetoursFileMappingTable.h b/Engine/Source/Programs/UnrealBuildAccelerator/Detours/Private/UbaDetoursFileMappingTable.h
--- a/Engine/Source/Programs/UnrealBuildAccelerator/Detours/Private/UbaDetoursFileMappingTable.h
+++ b/Engine/Source/Programs/UnrealBuildAccelerator/Detours/Private/UbaDetoursFileMappingTable.h
@@ -51,6 +51,8 @@ namespace uba
 		// If set to true, this means file was created during process group lifetime (which means it has its own file id etc)
 		bool created = false;
 
+		bool recursiveCall = false;
+
 		// File mapping handle coming from session process (Set in NtCreateFile)
 		SharedMemoryHandle trueFileMapHandle;
 
@@ -97,7 +99,7 @@ namespace uba
 	void Rpc_RegisterFileForWrite(StringView fileName, const StringKey& fileNameKey);
 	void Rpc_CheckRemapping(const StringView& fileName, const StringKey& fileNameKey);
 	bool Rpc_UpdateDirectory(const StringKey& dirKey, const tchar* dirName, u64 dirNameLen, bool lockDirTable = true);
-	void Rpc_UpdateCloseHandle(const tchar* handleName, u32 closeId, bool deleteOnClose, const tchar* newName, SharedMemoryHandle memoryHandle, u64 memoryWritten, bool success);
+	void Rpc_UpdateCloseHandle(const tchar* handleName, u32 closeId, bool deleteOnClose, const tchar* newName, SharedMemoryHandle memoryHandle, u64 fileSize, u64 lastWriteTime, u32 attributes, bool success);
 	void Rpc_UpdateTables();
 	bool Rpc_GetEntryOffset(DirTableOffset& out, const StringKey& entryNameKey, StringView entryName, bool checkIfDir = false);
 	void Rpc_GetFullFileName(const tchar*& path, u64& pathLen, StringBufferBase& tempBuf, bool useVirtualName, const tchar* const* loaderPaths = nullptr);
diff --git a/Engine/Source/Programs/UnrealBuildAccelerator/Detours/Private/UbaDetoursShared.cpp b/Engine/Source/Programs/UnrealBuildAccelerator/Detours/Private/UbaDetoursShared.cpp
--- a/Engine/Source/Programs/UnrealBuildAccelerator/Detours/Private/UbaDetoursShared.cpp
+++ b/Engine/Source/Programs/UnrealBuildAccelerator/Detours/Private/UbaDetoursShared.cpp
@@ -332,6 +332,13 @@ namespace uba
 		return false;
 	}
 
+	tchar DevirtualizeDrive(tchar drive)
+	{
+		if (!g_vfsEntryCount)
+			return '\0';
+		return g_vfsEntries[0].local[0];
+	}
+
 	bool VirtualizePath(StringBufferBase& path)
 	{
 		if (!g_vfsEntryCount)
@@ -568,67 +575,6 @@ namespace uba
 			return fileName.data;
 		}
 		#endif
-/*
-		{
-			SCOPED_READ_LOCK(g_mappedFileTable.m_lookupLock, lock);
-			auto it = g_mappedFileTable.m_lookup.find(fileNameKey);
-			if (it != g_mappedFileTable.m_lookup.end())
-			{
-				FileInfo& fi = it->second;
-
-				if (fi.deleted)
-				{
-					//if (isInsideSystemTemp)
-					//{
-					//	outAttr.useCache = false;
-					//	return fileName.data;
-					//}
-					outAttr.useCache = true;
-					outAttr.exists = false;
-					outAttr.lastError = ErrorFileNotFound;
-					return fileName.data;
-				}
-
-				// TODO: need to implement
-				//outAttr.data.ftLastWriteTime = ;
-				//outAttr.volumeSerial =
-				//outAttr.fileIndex = 
-
-#if PLATFORM_WINDOWS
-				foundMapping = true;
-				fileSize = fi.size;
-				outAttr.useCache = true;
-				outAttr.exists = true;
-				outAttr.lastError = ErrorSuccess;
-				LARGE_INTEGER li = ToLargeInteger(fileSize);
-				outAttr.data.nFileSizeLow = li.LowPart;
-				outAttr.data.nFileSizeHigh = li.HighPart;
-				outAttr.data.dwFileAttributes = FILE_ATTRIBUTE_NORMAL;
-				if (fi.created)
-					return fileName.data;
-				// We can't exit here because we need fileIndex and volume serial among other things.
-				// TODO: make sure fi.created is correct on windows everywhere. CopyFile/MoveFile should probably set this in certain conditions
-#else
-				if (fi.created)
-				{
-					outAttr.useCache = true;
-					outAttr.exists = true;
-					outAttr.lastError = ErrorSuccess;
-					outAttr.data.st_mode = (mode_t)(S_IFREG | S_IRUSR | S_IWUSR | S_IRGRP | S_IROTH);
-					outAttr.data.st_size = fi.size;
-					return fileName.data;
-				}
-#endif
-			}
-		}
-*/
-		if (!g_allowDirectoryCache || fileName.StartsWith(g_systemTemp.data)) // We need to skip SystemTemp.. lots of stuff going on there.
-		{
-			if (foundMapping)
-				return fileName.data;
-			outAttr.useCache = false;
-			return fileName.data;
-		}
 
 		// This is an optimization where we populate directory table and use that to figure out if file exists or not..
 		// .. in msvc's case it doesn't matter much because these tables are already up to date when msvc use CreateFile.
@@ -697,6 +643,10 @@ namespace uba
 		outAttr.data.st_dev = info.volumeSerial;
 		outAttr.data.st_ino = info.fileIndex;
 		outAttr.data.st_size = fileSize;
+		static auto uid = getuid();
+		static auto gid = getgid();
+		outAttr.data.st_uid = uid; // needed by git
+		outAttr.data.st_gid = gid; // needed by git
 #endif
 
 #if 0//UBA_DEBUG_VALIDATE
diff --git a/Engine/Source/Programs/UnrealBuildAccelerator/Detours/Private/UbaDetoursShared.h b/Engine/Source/Programs/UnrealBuildAccelerator/Detours/Private/UbaDetoursShared.h
--- a/Engine/Source/Programs/UnrealBuildAccelerator/Detours/Private/UbaDetoursShared.h
+++ b/Engine/Source/Programs/UnrealBuildAccelerator/Detours/Private/UbaDetoursShared.h
@@ -29,6 +29,7 @@ namespace uba
 	void PopulateVfs(BinaryReader& vfsReader);
 	bool IsVfsEnabled();
 	bool DevirtualizePath(StringBufferBase& path);
+	tchar DevirtualizeDrive(tchar drive);
 	bool VirtualizePath(StringBufferBase& path);
 	void LogHeader(StringView cmdLine);
 	void LogVfsInfo();
@@ -159,6 +160,10 @@ namespace uba
 		BinaryWriter writer(ProcessCommunication_Write); \
 		writer.WriteByte(MessageType_##messageName);
 
+	#define RPC_MESSAGE_END \
+		ts.Leave(); \
+		pcs.Leave();
+
 	// Dummy code
 	enum LogEntryType : u8 { LogEntryType_Error = 0, LogEntryType_Warning = 1, };
 	class Logger { public: void Logf(...) {} void Info(...) {} bool Warning(...) { return false; } bool Error(...) { return false; } };
diff --git a/Engine/Source/Programs/UnrealBuildAccelerator/Detours/Private/Windows/UbaDetoursFunctionsKernelBase.inl b/Engine/Source/Programs/UnrealBuildAccelerator/Detours/Private/Windows/UbaDetoursFunctionsKernelBase.inl
--- a/Engine/Source/Programs/UnrealBuildAccelerator/Detours/Private/Windows/UbaDetoursFunctionsKernelBase.inl
+++ b/Engine/Source/Programs/UnrealBuildAccelerator/Detours/Private/Windows/UbaDetoursFunctionsKernelBase.inl
@@ -221,67 +221,59 @@ BOOL Detoured_CreateDirectoryW(LPCWSTR lpPathName, LPSECURITY_ATTRIBUTES lpSecur
 {
 	DETOURED_CALL(CreateDirectoryW);
 
-	StringBuffer<> pathName;
-	FixPath(pathName, lpPathName);
-
-	StringBuffer<> lowerName(pathName);
+	StringBuffer<> dirPath;
+	FixPath(dirPath, lpPathName);
 
-	if (IsTrackingInput() && pathName.EndsWith(TCV("__pycache__"))) // If we are populating cache we need to make sure not to use pycache
+	if (IsTrackingInput() && dirPath.EndsWith(TCV("__pycache__"))) // If we are populating cache we need to make sure not to use pycache
 	{
 		SetLastError(ERROR_ACCESS_DENIED);
+		DEBUG_LOG_DETOURED(L"CreateDirectoryW", L"(NOPYCACHE) %s -> %s (ERROR_ACCESS_DENIED)", dirPath.data, ToString(false));
 		return false;
 	}
 
-	if (pathName.StartsWith(g_systemTemp.data))
+	DevirtualizePath(dirPath);
+	
+	bool checkDir = false;
+	if (dirPath.count == 3 && dirPath[2] == PathSeparator) // Root
 	{
-		SuppressCreateFileDetourScope s;
-		BOOL res = True_CreateDirectoryW(lpPathName, lpSecurityAttributes);
-		DEBUG_LOG_TRUE(L"CreateDirectoryW", L"%ls -> %ls", lpPathName, ToString(res));
+		dirPath.Resize(2);
+		checkDir = true;
+	}
 
-		// Some applications (icx.exe) creates directory where outputfiles are.. so we need to create an entry with '$' (directory) in the mapping table
-		if (res && g_allowOutputFiles && g_rules->IsOutputFile(pathName, g_systemTemp))
+	auto runError = [&](u32 error, const tchar* errorStr)
 		{
-			SCOPED_WRITE_LOCK(g_mappedFileTable.m_lookupLock, lookupLock);
-			auto insres = g_mappedFileTable.m_lookup.try_emplace(ToStringKey(lowerName.MakeLower()));
-			FileInfo& info = insres.first->second;
-			info.originalName = g_memoryBlock.Strdup(pathName).data;
-			info.name = TC("$");
-		}
-		return res;
-	}
+			SetLastError(error);
+			DEBUG_LOG_DETOURED(L"CreateDirectoryW", L"%s -> %s (%s)", dirPath.data, ToString(false), errorStr);
+			return false;
+		};
 
-	DevirtualizePath(pathName);
-	
-	lowerName.MakeLower();
-	StringKey pathNameKey = ToStringKey(lowerName);
+	StringKey dirPathKey = ToStringKeyLower(dirPath);
 
-	bool checkIfDir = true;
-	DirectoryTable::Exists exists = g_directoryTable.EntryExists(pathNameKey, lowerName, checkIfDir, nullptr);
-	if (exists == DirectoryTable::Exists_Yes)
-	{
-		SetLastError(ERROR_ALREADY_EXISTS);
-		DEBUG_LOG_DETOURED(L"CreateDirectoryW", L"(DIRTABLE) %ls -> %ls (%u)", pathName.data, ToString(false), ERROR_ALREADY_EXISTS);
-		return false;
-	}
+	DirTableOffset dirTableOffset;
+	if (Rpc_GetEntryOffset(dirTableOffset, dirPathKey, dirPath, checkDir))
+		return runError(ERROR_ALREADY_EXISTS, TC("ERROR_ALREADY_EXISTS"));
 
-	DirTableSize directoryTableSize;
-	BOOL res;
-	u32 errorCode = 0;
-	{
-		RPC_MESSAGE(CreateDirectory, createFile)
-		writer.WriteStringKey(pathNameKey);
-		writer.WriteString(pathName);
-		BinaryReader reader = writer.Flush();
-		res = reader.ReadBool();
-		errorCode = reader.ReadU32();
-		directoryTableSize = FromU64(reader.ReadU64());
-	}
+	const tchar* lastSeparator = dirPath.Last(PathSeparator);
+	if (!lastSeparator)
+		return runError(ERROR_PATH_NOT_FOUND, TC("ERROR_PATH_NOT_FOUND"));
+
+	StringView parentDirPath(dirPath.data, u32(lastSeparator - dirPath.data));
+	if (!Rpc_GetEntryOffset(dirTableOffset, ToStringKeyLower(parentDirPath), parentDirPath, true))
+		return runError(ERROR_PATH_NOT_FOUND, TC("ERROR_PATH_NOT_FOUND"));
+
+	// This should be a guaranteed create
+	RPC_MESSAGE(CreateDirectory, createFile)
+	writer.WriteStringKey(dirPathKey);
+	writer.WriteString(dirPath);
+	BinaryReader reader = writer.Flush();
+	DirTableSize directoryTableSize = FromU64(reader.ReadU64());
+	RPC_MESSAGE_END
 
 	g_directoryTable.ParseDirectoryTable(directoryTableSize);
 
-	SetLastError(errorCode);
-	DEBUG_LOG_DETOURED(L"CreateDirectoryW", L"%ls -> %ls (%u)", pathName.data, ToString(res), errorCode);
-	return res;
+	SetLastError(ERROR_SUCCESS);
+	DEBUG_LOG_DETOURED(L"CreateDirectoryW", L"%s -> %s", dirPath.data, ToString(true));
+	return true;
 }
 
 #if DETOURED_INCLUDE_DEBUG
@@ -296,38 +288,62 @@ BOOL Detoured_RemoveDirectoryW(LPCWSTR lpPathName)
 	t_calledRemoveDirectoryW = true;
 	#endif
 
-	StringBuffer<> pathName;
-	FixPath(pathName, lpPathName);
+	StringBuffer<> dirPath;
+	FixPath(dirPath, lpPathName);
+	DevirtualizePath(dirPath);
+
+	StringKey dirPathKey = ToStringKeyLower(dirPath);
 
-	if (pathName.StartsWith(g_systemTemp.data))
+	DirTableOffset dirTableOffset;
+	if (!Rpc_GetEntryOffset(dirTableOffset, dirPathKey, dirPath, true))
 	{
-		DEBUG_LOG_TRUE(L"RemoveDirectoryW", L"(NODETOUR) %ls", lpPathName);
-		SuppressCreateFileDetourScope s; // TODO: Revisit this.. will not work remotely
-		return True_RemoveDirectoryW(lpPathName);
+		SetLastError(ERROR_FILE_NOT_FOUND);
+		DEBUG_LOG_DETOURED(L"RemoveDirectoryW", L"%s -> %s (ERROR_FILE_NOT_FOUND)", dirPath.data, ToString(false));
+		return false;
 	}
 
-	DevirtualizePath(pathName);
+	DirectoryTable::EntryInformation info;
+	g_directoryTable.GetEntryInformation(info, dirTableOffset);
+	if (!IsDirectory(info.attributes))
+	{
+		SetLastError(ERROR_DIRECTORY);
+		DEBUG_LOG_DETOURED(L"RemoveDirectoryW", L"%s -> %s (ERROR_DIRECTORY)", dirPath.data, ToString(false));
+		return false;
+	}
 
-	DirTableSize directoryTableSize;
-	BOOL res;
-	u32 errorCode = 0;
-	StringKey pathNameKey = ToStringKeyLower(pathName);
+	auto findIt = g_directoryTable.m_lookup.find(dirPathKey);
+	UBA_ASSERT(findIt != g_directoryTable.m_lookup.end());
+	auto& dir = findIt->second;
+	if (!g_directoryTable.IsDirectoryUpToDate(dir))
+	{
+		StringBuffer<> dirPathLower(dirPath);
+		dirPathLower.MakeLower();
+		DirHash hash(dirPathLower);
+		g_directoryTable.PopulateDirectory(hash.open, dir);
+	}
 
+	for (auto& f : dir.files)
 	{
-		RPC_MESSAGE(RemoveDirectory, deleteFile)
-		writer.WriteStringKey(pathNameKey);
-		writer.WriteString(pathName);
-		BinaryReader reader = writer.Flush();
-		res = reader.ReadBool();
-		errorCode = reader.ReadU32();
-		directoryTableSize = FromU64(reader.ReadU64());
+		if (g_directoryTable.GetAttributes(f.second) == 0)
+			continue;
+
+		SetLastError(ERROR_DIR_NOT_EMPTY);
+		DEBUG_LOG_DETOURED(L"RemoveDirectoryW", L"%s -> %s (ERROR_DIR_NOT_EMPTY)", dirPath.data, ToString(false));
+		return false;
 	}
 
+	RPC_MESSAGE(RemoveDirectory, deleteFile)
+	writer.WriteStringKey(dirPathKey);
+	writer.WriteString(dirPath);
+	BinaryReader reader = writer.Flush();
+	DirTableSize directoryTableSize = FromU64(reader.ReadU64());
+	RPC_MESSAGE_END
+
 	g_directoryTable.ParseDirectoryTable(directoryTableSize);
 
-	SetLastError(errorCode);
-	DEBUG_LOG_DETOURED(L"RemoveDirectoryW", L"%ls -> %ls (%u)", pathName.data, ToString(res), errorCode);
-	return res;
+	SetLastError(ERROR_SUCCESS);
+	DEBUG_LOG_DETOURED(L"RemoveDirectoryW", L"%s -> %s", dirPath.data, ToString(true));
+	return true;
 }
 
 BOOL Detoured_LockFile(HANDLE hFile, DWORD dwFileOffsetLow, DWORD dwFileOffsetHigh, DWORD nNumberOfBytesToLockLow, DWORD nNumberOfBytesToLockHigh)
@@ -536,6 +552,16 @@ BOOL Detoured_GetVolumeInformationByHandleW(HANDLE hFile, LPWSTR lpVolumeNameBuf
 BOOL Detoured_GetVolumeInformationW(LPCWSTR lpRootPathName, LPWSTR lpVolumeNameBuffer, DWORD nVolumeNameSize, LPDWORD lpVolumeSerialNumber, LPDWORD lpMaximumComponentLength, LPDWORD lpFileSystemFlags, LPWSTR lpFileSystemNameBuffer, DWORD nFileSystemNameSize)
 {
 	DETOURED_CALL(GetVolumeInformationW);
+
+	tchar buffer[128];
+	if (IsVfsEnabled() && lpRootPathName)
+	{
+		UBA_ASSERT(lpRootPathName[1] == ':');
+		TStrcpy_s(buffer, sizeof_array(buffer), lpRootPathName);
+		buffer[0] = DevirtualizeDrive(buffer[0]);
+		lpRootPathName = buffer;
+	}
+
 	if (g_runningRemote)
 	{
 		if (lpVolumeSerialNumber)
@@ -795,6 +821,7 @@ BOOL Detoured_SetFileAttributesW(LPCWSTR lpFileName, DWORD dwFileAttributes)
 		return true;
 	}
 
+	// TODO: This should be tracked in overlay instead and applied afterwards
 	if (g_runningRemote && !fixedFile.StartsWith(g_systemTemp))
 	{
 		StringKey fileNameKey = ToStringKeyLower(fixedFile);
@@ -1140,30 +1167,40 @@ BOOL Detoured_GetModuleHandleExW(DWORD dwFlags, LPCWSTR lpModuleName, HMODULE *p
 
 BOOL Shared_CopyFileExW(const tchar* func, LPCWSTR lpExistingFileName, LPCWSTR lpNewFileName, LPPROGRESS_ROUTINE lpProgressRoutine, LPVOID lpData, LPBOOL pbCancel, DWORD dwCopyFlags)
 {
-	StringBuffer<> fromName;
-	FixPath(fromName, lpExistingFileName);
-	DevirtualizePath(fromName);
-	StringKey fromKey = ToStringKeyLower(fromName);
+	StringBuffer<> fromPath;
+	FixPath(fromPath, lpExistingFileName);
+	DevirtualizePath(fromPath);
+	StringKey fromKey = ToStringKeyLower(fromPath);
 
-	StringBuffer<> toName;
-	FixPath(toName, lpNewFileName);
-	DevirtualizePath(toName);
-	StringKey toKey = ToStringKeyLower(toName);
+	DirTableOffset dirTableOffset;
+	if (!Rpc_GetEntryOffset(dirTableOffset, fromKey, fromPath, false))
+	{
+		SetLastError(ERROR_FILE_NOT_FOUND);
+		return false;
+	}
 
-	StringBuffer<> newFromName;
-	StringBuffer<> newToName;
+	DirectoryTable::EntryInformation info;
+	g_directoryTable.GetEntryInformation(info, dirTableOffset);
+
+	StringBuffer<> toPath;
+	FixPath(toPath, lpNewFileName);
+	DevirtualizePath(toPath);
+	StringKey toKey = ToStringKeyLower(toPath);
+
+	StringBuffer<> newFromPath;
+	StringBuffer<> newToPath;
 	u32 closeId;
 	u32 lastError;
 	DirTableSize directoryTableSize;
 	{
 		RPC_MESSAGE(CopyFile, copyFile)
 		writer.WriteStringKey(fromKey);
-		writer.WriteString(fromName);
+		writer.WriteString(fromPath);
 		writer.WriteStringKey(toKey);
-		writer.WriteString(toName);
+		writer.WriteString(toPath);
 		BinaryReader reader = writer.Flush();
-		reader.ReadString(newFromName);
-		reader.ReadString(newToName);
+		reader.ReadString(newFromPath);
+		reader.ReadString(newToPath);
 		closeId = reader.ReadU32();
 		lastError = reader.ReadU32();
 		directoryTableSize = FromU64(reader.ReadU64());
@@ -1180,28 +1217,28 @@ BOOL Shared_CopyFileExW(const tchar* func, LPCWSTR lpExistingFileName, LPCWSTR l
 
 	// TODO: This copy should probably be moved to session process instead.. to handle failing to copy better
 
-	bool keepToFileInMemory = g_allowKeepFilesInMemory && g_rules->KeepInMemory(newToName, g_systemTemp, g_runningRemote, true);
+	bool keepToFileInMemory = g_allowKeepFilesInMemory && g_rules->KeepInMemory(newToPath, g_systemTemp, g_runningRemote, true);
 
 	bool res;
 	if (keepToFileInMemory)
 	{
-		UBA_ASSERTF(false, TC("%s copying to memory file is not implemented (Used by link.exe -PGOOptimize). Either change application rules or implement code path (Destination file: %s)"), func, newToName.data);
+		UBA_ASSERTF(false, TC("%s copying to memory file is not implemented (Used by link.exe -PGOOptimize). Either change application rules or implement code path (Destination file: %s)"), func, newToPath.data);
 		SetLastError(ERROR_IO_DEVICE);
 		return false;
 	}
 	else
 	{
 		SuppressCreateFileDetourScope cfs;
-		res = True_CopyFileExW(newFromName.data, newToName.data, lpProgressRoutine, lpData, pbCancel, dwCopyFlags);
+		res = True_CopyFileExW(newFromPath.data, newToPath.data, lpProgressRoutine, lpData, pbCancel, dwCopyFlags);
 	}
-	DEBUG_LOG_TRUE(func, L"%ls to %ls flags: %u (%ls to %ls) -> %ls", lpExistingFileName, lpNewFileName, dwCopyFlags, newFromName.data, newToName.data, ToString(res));
+	DEBUG_LOG_TRUE(func, L"%ls to %ls flags: %u (%ls to %ls) -> %ls", lpExistingFileName, lpNewFileName, dwCopyFlags, newFromPath.data, newToPath.data, ToString(res));
 
 	// We need to report the new file that has been added (and we must do it _after_ it has been copied
 	if (!closeId)
 		return res;
 
 	bool deleteOnClose = res == false; // If failing to copy we set deleteOnClose
-	Rpc_UpdateCloseHandle(newToName.data, closeId, deleteOnClose, L"", {}, 0, true);
+	Rpc_UpdateCloseHandle(newToPath.data, closeId, deleteOnClose, L"", {}, info.size, info.lastWrite, info.attributes, true);
 
 	return res;
 }
@@ -1259,7 +1296,7 @@ BOOL Detoured_CreateHardLinkW(LPCWSTR lpFileName, LPCWSTR lpExistingFileName, LP
 
 	// We need to report the new file that has been added (and we must do it _after_ it has been copied
 	if (closeId)
-		Rpc_UpdateCloseHandle(newToName.data, closeId, false, L"", {}, 0, true);
+		Rpc_UpdateCloseHandle(newToName.data, closeId, false, L"", {}, 0, 0, 0, true);
 
 	return res;
 }
@@ -1286,22 +1323,16 @@ BOOL Detoured_DeleteFileW(LPCWSTR lpFileName)
 
 	StringBuffer<> fixedName;
 	FixPath(fixedName, dosName);
-
 	DevirtualizePath(fixedName);
 
-	StringBuffer<> fixedNameLower(fixedName);
-	fixedNameLower.MakeLower();
-
-	StringKey fileNameKey = ToStringKey(fixedNameLower);
+	StringKey fileNameKey = ToStringKeyLower(fixedName);
 
-	if (!fixedName.StartsWith(g_systemTemp))
+	DirTableOffset dirTableOffset;
+	if (!Rpc_GetEntryOffset(dirTableOffset, fileNameKey, fixedName, false))
 	{
-		DirTableOffset dirTableOffset;
-		if (!Rpc_GetEntryOffset(dirTableOffset, fileNameKey, fixedName, false))
-		{
-			SetLastError(ERROR_FILE_NOT_FOUND);
-			return false;
-		}
+		DEBUG_LOG_DETOURED(L"DeleteFileW", L"(%ls) -> Error (ERROR_FILE_NOT_FOUND)", lpFileName);
+		SetLastError(ERROR_FILE_NOT_FOUND);
+		return false;
 	}
 
 	DirTableSize directoryTableSize;
@@ -1342,7 +1373,16 @@ bool Shared_MoveFile(const tchar* func, LPCWSTR lpExistingFileName, LPCWSTR lpNe
 
 	StringKey sourceKey = ToStringKeyLower(source);
 
-#if 0 // This path needs to be revisited... [hka] I can't figure out how to hit this... but KeepInMemory should not be used outside of NtCreateFile anymore
+	DirTableOffset dirTableOffset;
+	if (!Rpc_GetEntryOffset(dirTableOffset, sourceKey, source, false))
+	{
+		SetLastError(ERROR_FILE_NOT_FOUND);
+		return false;
+	}
+	DirectoryTable::EntryInformation info;
+	g_directoryTable.GetEntryInformation(info, dirTableOffset);
+
+	#if 0 // This path needs to be revisited... [hka] I can't figure out how to hit this... but KeepInMemory should not be used outside of NtCreateFile anymore
 	if (KeepInMemory(source, false))
 	{
 		UBA_ASSERT(false);
@@ -1416,6 +1456,9 @@ bool Shared_MoveFile(const tchar* func, LPCWSTR lpExistingFileName, LPCWSTR lpNe
 		RPC_MESSAGE(MoveFile, moveFile)
 		writer.WriteStringKey(sourceKey);
 		writer.WriteString(source);
+		writer.WriteU64(info.size);
+		writer.WriteU64(info.lastWrite);
+		writer.WriteU32(info.attributes);
 		writer.WriteStringKey(destKey);
 		writer.WriteString(dest);
 		writer.WriteU32(dwFlags);
@@ -1529,7 +1572,8 @@ __forceinline HANDLE Shared_FindFirstFileExW(LPCWSTR lpFileName, FINDEX_INFO_LEV
 	FixPath(lowerName, lpFileName);
 	DevirtualizePath(lowerName);
 
-	if (lowerName.StartsWith(g_systemTemp.data) || lowerName.StartsWith(g_systemRoot.data))
+	//if (lowerName.StartsWith(g_systemTemp.data) || lowerName.StartsWith(g_systemRoot.data))
+	if (g_systemRoot.count && lowerName.StartsWith(g_systemRoot.data))
 		return Local_FindFirstFileExW(lpFileName, fInfoLevelId, lpFindFileData, fSearchOp, lpSearchFilter, dwAdditionalFlags, funcName);
 
 	lowerName.MakeLower();
@@ -1568,7 +1612,7 @@ __forceinline HANDLE Shared_FindFirstFileExW(LPCWSTR lpFileName, FINDEX_INFO_LEV
 			Rpc_UpdateDirectory(hash.key, buf, bufChars, false);
 	}
 	bool exists = false;
-	if (dir.latestOffset != InvalidTableOffset)
+	if (dir.latestOffset != InvalidTableOffset || dir.latestOverlayOffset != InvalidTableOffset)
 	{
 		DirTableOffset entryOffset;
 		g_directoryTable.GetLatestOffset(entryOffset, dir);
@@ -1592,8 +1636,8 @@ __forceinline HANDLE Shared_FindFirstFileExW(LPCWSTR lpFileName, FINDEX_INFO_LEV
 
 	if (!exists)
 	{
-		if (g_systemTemp.StartsWith(lowerName.data)) // TODO: This is a big hack. We should make sure the uba system temp folder is virtualized and is always some root path that never can collide with the host file system
-			return Local_FindFirstFileExW(lpFileName, fInfoLevelId, lpFindFileData, fSearchOp, lpSearchFilter, dwAdditionalFlags, funcName);
+		//if (g_systemTemp.StartsWith(lowerName.data)) // TODO: This is a big hack. We should make sure the uba system temp folder is virtualized and is always some root path that never can collide with the host file system
+		//	return Local_FindFirstFileExW(lpFileName, fInfoLevelId, lpFindFileData, fSearchOp, lpSearchFilter, dwAdditionalFlags, funcName);
 
 		DEBUG_LOG_DETOURED(funcName, L"(%ls) -> NotFound", lpFileName);
 		SetLastError(ERROR_FILE_NOT_FOUND);
@@ -1636,8 +1680,8 @@ __forceinline HANDLE Shared_FindFirstFileExW(LPCWSTR lpFileName, FINDEX_INFO_LEV
 		if (!Shared_GetNextFile(data, *listHandle))
 		{
 			delete listHandle;
-			if (g_systemTemp.StartsWith(lowerName.data)) // TODO: This is a big hack. We should make sure the uba system temp folder is virtualized and is always some root path that never can collide with the host file system
-				return Local_FindFirstFileExW(lpFileName, fInfoLevelId, lpFindFileData, fSearchOp, lpSearchFilter, dwAdditionalFlags, funcName);
+			//if (g_systemTemp.StartsWith(lowerName.data)) // TODO: This is a big hack. We should make sure the uba system temp folder is virtualized and is always some root path that never can collide with the host file system
+			//	return Local_FindFirstFileExW(lpFileName, fInfoLevelId, lpFindFileData, fSearchOp, lpSearchFilter, dwAdditionalFlags, funcName);
 			DEBUG_LOG_DETOURED(funcName, L"(%ls) -> NotFound(2)", lpFileName);
 			return INVALID_HANDLE_VALUE;
 		}
@@ -2514,7 +2558,7 @@ DWORD Detoured_SearchPathW(LPCWSTR lpPath, LPCWSTR lpFileName, LPCWSTR lpExtensi
 	}
 
 	DWORD res = True_SearchPathW(lpPath, lpFileName, lpExtension, nBufferLength, lpBuffer, lpFilePart);
-	if (res && IsVfsEnabled() && !StartsWith(lpBuffer, g_systemRoot.data))
+	if (res && IsVfsEnabled() && (!g_systemRoot.count || !StartsWith(lpBuffer, g_systemRoot.data)))
 	{
 		UBA_ASSERT(res < nBufferLength);
 		StringBuffer<> temp(lpBuffer);
@@ -2779,7 +2823,10 @@ BOOL Detoured_SetStdHandle(DWORD nStdHandle, HANDLE hHandle)
 	DETOURED_CALL(SetStdHandle);
 
 	if (g_isDetachedProcess)
+	{
+		DEBUG_LOG_DETOURED(L"SetStdHandle", L"(SKIPPED!) %s %llu -> Success", StdToString(nStdHandle), u64(hHandle));
 		return true;
+	}
 
 	if (nStdHandle == STD_OUTPUT_HANDLE)
 		g_stdHandle[1] = (hHandle != g_nullFile && GetFileType(hHandle) == FILE_TYPE_CHAR) ? hHandle : 0;
@@ -2797,11 +2844,11 @@ BOOL Detoured_SetStdHandle(DWORD nStdHandle, HANDLE hHandle)
 		trueHandle = dh.trueHandle;
 		UBA_ASSERT(trueHandle != INVALID_HANDLE_VALUE);
 	}
-	DEBUG_LOG_TRUE(L"SetStdHandle", L"%u -> %llu", nStdHandle, u64(trueHandle));
-
 	g_setStdHandleCalled = true;
 
-	return True_SetStdHandle(nStdHandle, trueHandle);
+	auto res = True_SetStdHandle(nStdHandle, trueHandle);
+	DEBUG_LOG_TRUE(L"SetStdHandle", L"%s %llu -> %s", StdToString(nStdHandle), u64(trueHandle), ToString(res));
+	return res;
 }
 
 
@@ -2901,6 +2948,21 @@ BOOL Detoured_CreateProcessW(LPCWSTR lpApplicationName, LPWSTR lpCommandLine, LP
 		testName = originalCmd;
 	}
 
+	// Don't know if this is sufficient. But we can't have child processes inherit detoured handles and maybe this detouring has to happen further down the stack
+	// TODO: Maybe not modify lpStartupInfo directly
+	if (lpStartupInfo->cbReserved2)
+	{
+		BYTE* p = lpStartupInfo->lpReserved2;
+		DWORD count = *(DWORD*)p;
+		BYTE*   flags   = p + sizeof(DWORD);
+		HANDLE* handles = (HANDLE*)(flags + count);
+		for (DWORD i = 0; i < count; i++)
+		{
+			HANDLE h = handles[i];
+			if (h != INVALID_HANDLE_VALUE && isDetouredHandle(h))
+				handles[i] = asDetouredHandle(h).trueHandle;
+		}
+	}
 
 	// Debug binaries started when process crash... we don't want to detour these.
 	if (testName.Contains(L"winedbg") || testName.Contains(L"werfault.exe") || testName.Contains(L"vsjitdebugger.exe") || testName.Contains(L"crashpad_handler.exe"))
@@ -3007,6 +3069,7 @@ BOOL Detoured_CreateProcessW(LPCWSTR lpApplicationName, LPWSTR lpCommandLine, LP
 		}
 		Rpc_UpdateTables();
 
+		UBA_ASSERT(g_systemRoot.count);
 		StringBuffer<256> cmdExe;
 		cmdExe.Append(g_systemRoot).EnsureEndsWithSlash().Append(TCV("system32\\cmd.exe"));
 		SuppressDetourScope _;
@@ -3659,14 +3722,14 @@ BOOL Detoured_SetEndOfFile(HANDLE hFile)
 DWORD Detoured_GetFileAttributesA(LPCSTR lpFileName)
 {
 	// Is verified that both windows and wine are calling GetFileAttributesW
-	DEBUG_LOG_TRUE(L"GetFileAttributesA", L"");
+	DEBUG_LOG_TRUE(L"GetFileAttributesA", L"%S", lpFileName);
 	return True_GetFileAttributesA(lpFileName);
 }
 
 BOOL Detoured_GetFileAttributesExA(LPCSTR lpFileName, GET_FILEEX_INFO_LEVELS fInfoLevelId, LPVOID lpFileInformation)
 {
 	DETOURED_CALL(GetFileAttributesExA);
-	DEBUG_LOG_TRUE(L"GetFileAttributesExA", L""); // Calls ExW on both windows and wine
+	DEBUG_LOG_TRUE(L"GetFileAttributesExA", L"%S", lpFileName); // Calls ExW on both windows and wine
 	return True_GetFileAttributesExA(lpFileName, fInfoLevelId, lpFileInformation);
 }
 
@@ -3972,6 +4035,15 @@ DWORD Detoured_GetTempPathW(DWORD nBufferLength, LPWSTR lpBuffer)
 	DEBUG_LOG_TRUE(L"GetTempPathW", L"%s -> %u", lpBuffer, res);
 	return res;
 }
+
+DWORD Detoured_GetTempPath2W(DWORD nBufferLength, LPWSTR lpBuffer)
+{
+	DETOURED_CALL(GetTempPathW);
+	DWORD res = True_GetTempPath2W(nBufferLength, lpBuffer);
+	DEBUG_LOG_TRUE(L"GetTempPath2W", L"%s -> %u", lpBuffer, res);
+	return res;
+}
+
 DWORD Detoured_ExpandEnvironmentStringsW(LPCWSTR lpSrc, LPWSTR lpDst, DWORD nSize)
 {
 	DETOURED_CALL(ExpandEnvironmentStringsW);
@@ -4351,8 +4423,26 @@ HMODULE Detoured_GetModuleHandleW(LPCWSTR lpModuleName)
 VOID Detoured_GetStartupInfoW(LPSTARTUPINFOW lpStartupInfo)
 {
 	DETOURED_CALL(GetStartupInfoW);
-	DEBUG_LOG_TRUE(L"GetStartupInfoW", L"");
 	True_GetStartupInfoW(lpStartupInfo);
+	DEBUG_LOG_TRUE(L"GetStartupInfoW", L"stdout: %llu", u64(lpStartupInfo->hStdOutput));
+	
+	if (lpStartupInfo->cbReserved2 && lpStartupInfo->lpReserved2)
+	{
+		BYTE* p = lpStartupInfo->lpReserved2;
+		DWORD count = *(DWORD*)p;
+		BYTE*   flags   = p + sizeof(DWORD);
+		HANDLE* handles = (HANDLE*)(flags + count);  // may be unaligned, but x64 handles it
+
+		for (DWORD i = 0; i < count; i++)
+		{
+			HANDLE h = handles[i];
+			if (h == INVALID_HANDLE_VALUE)
+				continue;
+			UBA_ASSERT(!isDetouredHandle(h)); // Detoured handles can't be inherited. This is not good
+			// flags[i] bits: FOPEN=0x01, FTEXT=0x80, FPIPE=0x08, FDEV=0x40, etc.
+			// do something with h and flags[i]
+		}
+	}
 }
 
 VOID Detoured_GetSystemTimeAsFileTime(LPFILETIME lpSystemTimeAsFileTime)
diff --git a/Engine/Source/Programs/UnrealBuildAccelerator/Detours/Private/Windows/UbaDetoursFunctionsNtDll.inl b/Engine/Source/Programs/UnrealBuildAccelerator/Detours/Private/Windows/UbaDetoursFunctionsNtDll.inl
--- a/Engine/Source/Programs/UnrealBuildAccelerator/Detours/Private/Windows/UbaDetoursFunctionsNtDll.inl
+++ b/Engine/Source/Programs/UnrealBuildAccelerator/Detours/Private/Windows/UbaDetoursFunctionsNtDll.inl
@@ -60,6 +60,17 @@ u8 GetFileAccessFlags(DWORD desiredAccess, u32 createDisposition)
 	return access;
 }
 
+const tchar* StdToString(DWORD stdHandle)
+{
+	if (stdHandle == STD_OUTPUT_HANDLE)
+		return L"StdOut";
+	if (stdHandle == STD_INPUT_HANDLE)
+		return L"StdIn";
+	if (stdHandle == STD_ERROR_HANDLE)
+		return L"StdErr";
+	return L"StdINVALID";
+}
+
 #if UBA_DEBUG_LOG_ENABLED
 StringBuffer<32> ToString(NTSTATUS s)
 {
@@ -262,7 +273,7 @@ NTSTATUS Detoured_NtQueryInformationFile(HANDLE FileHandle, PIO_STATUS_BLOCK IoS
 		else if (FileInformationClass == FileIdInformation)
 		{
 			auto& info = *(FILE_ID_INFORMATION*)FileInformation;
-			if (listHandle.dir.latestOffset != InvalidTableOffset)
+			if (listHandle.dir.latestOffset != InvalidTableOffset || listHandle.dir.latestOverlayOffset != InvalidTableOffset)
 			{
 				DirTableOffset entryOffset;
 				g_directoryTable.GetLatestOffset(entryOffset, listHandle.dir);
@@ -393,7 +404,16 @@ NTSTATUS Detoured_NtQueryInformationFile(HANDLE FileHandle, PIO_STATUS_BLOCK IoS
 
 	if (FileInformationClass == FileStandardInformation)
 	{
-		UBA_ASSERT(dh.type == HandleType_File);
+		if (dh.type != HandleType_File)
+		{
+			if (dh.type == HandleType_StdOut || dh.type == HandleType_StdErr)
+			{
+				if (trueHandle == INVALID_HANDLE_VALUE)
+					trueHandle = True_GetStdHandle(dh.type == HandleType_StdOut ? STD_OUTPUT_HANDLE : STD_ERROR_HANDLE);
+				return runTrue();
+			}
+			UBA_ASSERTF(false, TC("NtQueryInformationFile is called on unsupported detoured handle: %u"), u32(dh.type));
+		}
 		FileInfo& fi = *dh.fileObject->fileInfo;
 		if (!fi.memoryFile && !fi.isFileMap && fi.created) // fi.created means that it was created in process hierarchy and other fields might be outdated when it comes to size
 			return runTrue();
@@ -1394,7 +1414,7 @@ NTSTATUS NTAPI Shared_NtCreateFile(bool IsCreateFunc, PHANDLE hFileHandle, ACCES
 	const wchar_t* isWriteStr = isWrite ? L"WRITE " : L""; (void)isWriteStr;
 	#endif
 
-	bool isSystemFile = fileName.StartsWith(g_systemRoot.data);
+	bool isSystemFile = g_systemRoot.count && fileName.StartsWith(g_systemRoot.data);
 	if (g_runningRemote)
 	{
 		if (fileName.EndsWith(TCV(".tlb")))  // We want the tbl files to be detoured
@@ -1410,14 +1430,13 @@ NTSTATUS NTAPI Shared_NtCreateFile(bool IsCreateFunc, PHANDLE hFileHandle, ACCES
 		isSystemFile = true;
 		if (g_runningRemote)
 		{
+			UBA_ASSERT(g_systemRoot.count);
 			// TODO: Should we create a detoured handle so name calls are correct?
 			ObjectAttributes->ObjectName->Buffer[4] = g_systemRoot[0];
 		}
 	}
 
-	bool isSystemOrTempFile = isSystemFile || fileName.StartsWith(g_systemTemp.data);
-
-	if (isSystemFile || (isSystemOrTempFile && !keepInMemory))
+	if (isSystemFile)
 	{
 		ObjectAttributes->RootDirectory = rootDir;
 		NTSTATUS res = Local_NtCreateFile(IsCreateFunc, hFileHandle, DesiredAccess, ObjectAttributes, IoStatusBlock, AllocationSize, FileAttributes, ShareAccess, CreateDisposition, CreateOptions, EaBuffer, EaLength);
@@ -1426,11 +1445,7 @@ NTSTATUS NTAPI Shared_NtCreateFile(bool IsCreateFunc, PHANDLE hFileHandle, ACCES
 		DEBUG_LOG_TRUE(funcName, L"(NODETOUR) %s%llu (%.*ls) -> %ls", isWriteStr, uintptr_t(*hFileHandle), ObjectAttributes->ObjectName->Length / 2, ObjectAttributes->ObjectName->Buffer, ToString(res).data);
 		if (NT_ERROR(res))
 			return res;
-
-		if (!isSystemOrTempFile && !isWrite && !t_disallowDetour && fileName[fileName.count-1] != ':')
-			TrackInput(fileName);
-		else
-			SkipTrackInput(fileName);
+		SkipTrackInput(fileName);
 		return res;
 	}
 
@@ -1451,11 +1466,6 @@ NTSTATUS NTAPI Shared_NtCreateFile(bool IsCreateFunc, PHANDLE hFileHandle, ACCES
 			return STATUS_OBJECT_NAME_NOT_FOUND;
 		}
 	}
-	else if (isSystemOrTempFile)
-	{
-		// TODO: For now, ignore these when writing..
-		// Once overlay is fully implemented we want all written files to end up in overlay to simplify all the other functions
-	}
 	else if (const tchar* lastSeparator = fileNameLower.Last(PathSeparator))
 	{
 		StringView dirName(fileNameLower.data, u32(lastSeparator - fileNameLower.data));
@@ -1620,20 +1630,19 @@ NTSTATUS NTAPI Shared_NtCreateFile(bool IsCreateFunc, PHANDLE hFileHandle, ACCES
 	SCOPED_WRITE_LOCK(g_mappedFileTable.m_lookupLock, lookupLock);
 	auto insres = g_mappedFileTable.m_lookup.try_emplace(fileNameKey);
 	FileInfo& info = insres.first->second;
+	UBA_ASSERT(!info.recursiveCall || !isWrite);
 	u32 lastDesiredAccess = info.lastDesiredAccess;
-	if (insres.second)
+	if (insres.second || info.recursiveCall)
 	{
 		u64 size = InvalidValue;
 		info.originalName = g_memoryBlock.Strdup(fileName).data;
 		info.name = info.originalName;
-		if (!keepInMemory && !isSystemOrTempFile)
-		{
-			u8 access = GetFileAccessFlags(DesiredAccess, CreateDisposition);
-			wchar_t newFileName[512];
-			Rpc_CreateFileW(fileName, fileNameKey, access, newFileName, sizeof_array(newFileName), size, closeId, false);
-			info.name = g_memoryBlock.Strdup(newFileName);
-			lpFileName = info.name;
-		}
+
+		u8 access = GetFileAccessFlags(DesiredAccess, CreateDisposition);
+		wchar_t newFileName[512];
+		Rpc_CreateFileW(fileName, fileNameKey, access, newFileName, sizeof_array(newFileName), size, closeId, false);
+		info.name = g_memoryBlock.Strdup(newFileName);
+		lpFileName = info.name;
 
 		info.size = size;
 		info.fileNameKey = fileNameKey;
@@ -1831,7 +1840,7 @@ NTSTATUS NTAPI Shared_NtCreateFile(bool IsCreateFunc, PHANDLE hFileHandle, ACCES
 			mf.writtenSize = mappingHandleSize;
 			mf.committedSize = AlignUp(mappingHandleSize, g_pageSize);
 			mf.mappedSize = mf.committedSize;
-			mf.reserveSize = FileTypeMaxSize(fileName, isSystemOrTempFile);
+			mf.reserveSize = FileTypeMaxSize(fileName);
 
 			TimerScope ts2(g_kernelStats.mapViewOfFile);
 
@@ -1860,7 +1869,7 @@ NTSTATUS NTAPI Shared_NtCreateFile(bool IsCreateFunc, PHANDLE hFileHandle, ACCES
 		{
 			bool isOutput = (isWrite || isDeleteOnClose) && g_rules->IsOutputFile(fileName, g_systemTemp);
 			bool isLocal = !isOutput;
-			MemoryFile* mf = new MemoryFile(isLocal, FileTypeMaxSize(fileName, isSystemOrTempFile), isThrowAway, FileTypeCommitSize(fileName), fileName.data);
+			MemoryFile* mf = new MemoryFile(isLocal, FileTypeMaxSize(fileName), isThrowAway, FileTypeCommitSize(fileName), fileName.data);
 			auto mfg = MakeGuard([mf](){ delete mf; });
 
 			if (!isThrowAway && CreateDisposition == FILE_OPEN)
@@ -1876,12 +1885,16 @@ NTSTATUS NTAPI Shared_NtCreateFile(bool IsCreateFunc, PHANDLE hFileHandle, ACCES
 				// (We end up in this code path is used for incremental linking)
 
 				DEBUG_LOG_DETOURED(funcName, L"INTERNAL READ FOR MEMORYWRITE (%ls) (%u %u)", fileName.data, CreateDisposition, CreateOptions);
-
+				UBA_ASSERT(!info.recursiveCall);
+				FileInfo temp = info;
+				info.recursiveCall = true;
+				lookupLock.Leave();
 				HANDLE fileHandle;
 				IO_STATUS_BLOCK ioStatusBlock;
-				lookupLock.Leave();
 				NTSTATUS res2 = Shared_NtCreateFile(IsCreateFunc, &fileHandle, FILE_GENERIC_READ, ObjectAttributes, &ioStatusBlock, NULL, FILE_ATTRIBUTE_NORMAL, FILE_SHARE_READ, FILE_OPEN, FILE_SEQUENTIAL_ONLY|FILE_SYNCHRONOUS_IO_NONALERT, NULL, 0);
 				lookupLock.Enter();
+				info = temp;
+
 				if (res2 != STATUS_SUCCESS)
 				{
 					*hFileHandle = INVALID_HANDLE_VALUE;
@@ -1994,7 +2007,7 @@ NTSTATUS NTAPI Shared_NtCreateFile(bool IsCreateFunc, PHANDLE hFileHandle, ACCES
 		if (closeId)
 		{
 			info.lastDesiredAccess = lastDesiredAccess;
-			Rpc_UpdateCloseHandle(L"", closeId, false, L"", {}, 0, false);
+			Rpc_UpdateCloseHandle(L"", closeId, false, L"", {}, 0, 0, 0, false);
 		}
 		DEBUG_LOG_TRUE(funcName, L"%ls(%ls) (%ls) -> %ls", isWriteStr, lpFileName, (fileName.data != lpFileName ? fileName.data : L""), ToString(res).data);
 		return res;
@@ -2422,7 +2435,7 @@ NTSTATUS NTAPI Detoured_NtClose(HANDLE handle)
 			{
 				if (!log)
 					return STATUS_SUCCESS;
-				DEBUG_LOG_DETOURED(L"NtClose", L"%llu (%ls) -> %ls", uintptr_t(handle), HandleToName(handle), ToString(STATUS_SUCCESS).data);
+				DEBUG_LOG_TRUE(L"NtClose", L"%llu (%ls) -> %ls", uintptr_t(handle), HandleToName(handle), ToString(STATUS_SUCCESS).data);
 				return STATUS_SUCCESS;
 			}
 			TimerScope ts(g_kernelStats.closeFile);
@@ -2467,19 +2480,34 @@ NTSTATUS NTAPI Detoured_NtClose(HANDLE handle)
 		return closeTrueHandle(false);
 	}
 
-	NTSTATUS res = closeTrueHandle(false);
-
 	SharedMemoryHandle memoryHandle;
-	u64 memoryWritten = 0;
 	FileInfo& fi = *fo->fileInfo;
 	const wchar_t* path = fi.name;
 	wchar_t temp[512];
 
+	u64 fileSize = InvalidValue;
+	u64 lastWriteTime = 0;
+	u32 attributes = FILE_ATTRIBUTE_NORMAL;
+
+	// TODO: We can track file size outside of kernel calls
+	if (trueHandle != INVALID_HANDLE_VALUE && (fo->closeId || !fo->newName.empty()) && !fo->deleteOnClose)
+	{
+		IO_STATUS_BLOCK b;
+		FILE_NETWORK_OPEN_INFORMATION info;
+		NTSTATUS res = True_NtQueryInformationFile(trueHandle, &b, &info, sizeof(info), FileNetworkOpenInformation);(void)res;
+		UBA_ASSERTF(res == STATUS_SUCCESS, TC("Failed to get size of file %s"), path);
+		attributes = info.FileAttributes;
+		lastWriteTime = info.LastWriteTime.QuadPart;
+		fileSize = info.EndOfFile.QuadPart;
+	}
+
+	NTSTATUS res = closeTrueHandle(false);
+		
 	if (MemoryFile* mf = fi.memoryFile)
 	{
 		if (IsWrite(fo->desiredAccess, 0))
 		{
-			// TODO: There are race conditions in this code. There could be other file handles accessing the same piece of memory (allthough unlikely)
+			// TODO: There are race conditions in this code. There could be other file handles accessing the same piece of memory (although unlikely)
 			u64 alignedWritten = AlignUp(mf->writtenSize, PageSize);
 			if (alignedWritten < mf->committedSize)
 			{
@@ -2497,7 +2525,11 @@ NTSTATUS NTAPI Detoured_NtClose(HANDLE handle)
 		}
 
 		memoryHandle = mf->memoryHandle;
-		memoryWritten = mf->writtenSize;
+		fileSize = mf->writtenSize;
+
+		FILETIME ft;
+		::GetSystemTimeAsFileTime(&ft);
+		lastWriteTime = u64(ToLargeInteger(ft.dwHighDateTime, ft.dwLowDateTime).QuadPart);
 
 		u32 orginalNameLen = TStrlen(fi.originalName);
 		if ((fo->deleteOnClose || IsWrite(fo->desiredAccess, 0)) && g_rules->IsOutputFile({fi.originalName, orginalNameLen}, g_systemTemp) && !g_rules->IsThrowAway(StringView(fi.originalName, orginalNameLen), g_runningRemote))
@@ -2551,6 +2583,8 @@ NTSTATUS NTAPI Detoured_NtClose(HANDLE handle)
 	}
 	else if (fo->deleteOnClose) // We have used an optimized handle that actually never opens the file so we need to delete it manually
 	{
+		fileSize = 0;
+
 		if (dh.trueHandle == INVALID_HANDLE_VALUE)
 			DeleteFileW(fi.originalName);
 		else if (!fo->closeId && !g_exitMessageSent)
@@ -2594,21 +2628,18 @@ NTSTATUS NTAPI Detoured_NtClose(HANDLE handle)
 		//fo->newName.clear(); // Don't remove newName here since we might want to enter the update close
 	}
 
-	if (fo->closeId)
+	if (fo->closeId && !g_exitMessageSent)
 	{
-		if (!g_exitMessageSent)
-			Rpc_UpdateCloseHandle(path, fo->closeId, fo->deleteOnClose, fo->newName.c_str(), memoryHandle, memoryWritten, true);
+		DEBUG_LOG_DETOURED(L"NtClose", L"%llu (%ls) CloseId:%u, Size:%llu -> %ls", uintptr_t(handle), HandleToName(handle), fo->closeId, fileSize, ToString(res).data);
+		Rpc_UpdateCloseHandle(path, fo->closeId, fo->deleteOnClose, fo->newName.c_str(), memoryHandle, fileSize, lastWriteTime, attributes, true);
 	}
 	else
 	{
-		// TODO: Update g_mappedFileTable.m_lookup?
-		//UBA_ASSERTF(fo->newName.empty(), L"Got close of file that was renamed but had no closeId. Old: %ls New: %ls", fi.originalName, fo->newName.c_str());
+		DEBUG_LOG_DETOURED(L"NtClose", L"%llu (%ls) -> %ls", uintptr_t(handle), HandleToName(handle), ToString(res).data);
 	}
 
 	InterlockedDecrement(&fi.refCount);
 
-	DEBUG_LOG_DETOURED(L"NtClose", L"%llu (%ls) -> %ls", uintptr_t(handle), HandleToName(handle), ToString(res).data);
-
 	if (fo->ownsFileInfo)
 	{
 		UBA_ASSERT(!fi.memoryFile);
diff --git a/Engine/Source/Programs/UnrealBuildAccelerator/Detours/Private/Windows/UbaDetoursFunctionsWin.cpp b/Engine/Source/Programs/UnrealBuildAccelerator/Detours/Private/Windows/UbaDetoursFunctionsWin.cpp
--- a/Engine/Source/Programs/UnrealBuildAccelerator/Detours/Private/Windows/UbaDetoursFunctionsWin.cpp
+++ b/Engine/Source/Programs/UnrealBuildAccelerator/Detours/Private/Windows/UbaDetoursFunctionsWin.cpp
@@ -454,7 +454,7 @@ const wchar_t* HandleToName(HANDLE handle)
 }
 
 
-u64 FileTypeMaxSize(StringView file, bool isSystemOrTempFile) { return g_rules->FileTypeMaxSize(file, isSystemOrTempFile); }
+u64 FileTypeMaxSize(StringView file) { return g_rules->FileTypeMaxSize(file, file.StartsWith(g_systemTemp)); }
 u32 FileTypeCommitSize(StringView file) { return g_rules->FileTypeCommitSize(file); }
 
 bool EnsureMapped(DetouredHandle& handle, DWORD dwFileOffsetHigh = 0, DWORD dwFileOffsetLow = 0, SIZE_T numberOfBytesToMap = 0, void* baseAddress = nullptr, SharedMemoryMapType type = SharedMemoryMapType_ReadOnly)
diff --git a/Engine/Source/Programs/UnrealBuildAccelerator/Detours/Private/Windows/UbaDetoursFunctionsWin.h b/Engine/Source/Programs/UnrealBuildAccelerator/Detours/Private/Windows/UbaDetoursFunctionsWin.h
--- a/Engine/Source/Programs/UnrealBuildAccelerator/Detours/Private/Windows/UbaDetoursFunctionsWin.h
+++ b/Engine/Source/Programs/UnrealBuildAccelerator/Detours/Private/Windows/UbaDetoursFunctionsWin.h
@@ -327,6 +327,7 @@
 	DETOURED_FUNCTION(GetEnvironmentVariableA) \
 	DETOURED_FUNCTION(GetEnvironmentStringsW) \
 	DETOURED_FUNCTION(GetTempPathW) \
+	DETOURED_FUNCTION(GetTempPath2W) \
 	DETOURED_FUNCTION(ExpandEnvironmentStringsW) \
 	DETOURED_FUNCTION(FindFirstFileNameW) \
 	DETOURED_FUNCTION(GetTempFileNameW) \
@@ -532,6 +533,7 @@ extern "C" {
 	HRESULT SHGetFolderPathW(HWND hwnd, int csidl, HANDLE hToken, DWORD dwFlags, LPWSTR pszPath);
 	enum FILE_INFO_BY_NAME_CLASS{};
 	BOOL WINAPI GetFileInformationByName(PCWSTR FileName, FILE_INFO_BY_NAME_CLASS FileInformationClass, PVOID FileInfoBuffer, ULONG FileInfoBufferSize);
+	DWORD WINAPI GetTempPath2W(DWORD  BufferLength, LPWSTR Buffer);
 	__time64_t WINAPI _time64(__time64_t* destTime);
 
 }
diff --git a/Engine/Source/Programs/UnrealBuildAccelerator/Test/Private/UbaTestAll.h b/Engine/Source/Programs/UnrealBuildAccelerator/Test/Private/UbaTestAll.h
--- a/Engine/Source/Programs/UnrealBuildAccelerator/Test/Private/UbaTestAll.h
+++ b/Engine/Source/Programs/UnrealBuildAccelerator/Test/Private/UbaTestAll.h
@@ -90,10 +90,10 @@ namespace uba
 		UBA_TEST(TestRemoteRegisterVirtualFile) \
 		UBA_TEST(TestRemoteProcessSpecialCase1) \
 		UBA_TEST(TestOverlayAccessOrdering) \
-		UBA_TEST(TestOverlayAccessOrderingRemote) \
+		UBA_TEST(TestRemoteOverlayAccessOrdering) \
 		UBA_TEST(TestSessionSpecialCopy) \
 		UBA_TEST(TestFileMapWrite) \
-		UBA_TEST(TestFileMapWriteRemote) \
+		UBA_TEST(TestRemoteFileMapWrite) \
 
 
 #define UBA_LINUX_TESTS \
@@ -101,6 +101,8 @@ namespace uba
 		UBA_POSIX_TESTS \
 		UBA_TEST(TestDetouredTouch) \
 		UBA_TEST(TestDetouredPopen) \
+		UBA_TEST(TestDetouredSymlinks) \
+		UBA_TEST(TestRemoteDetouredSymlinks) \
 
 
 #define UBA_MAC_TESTS \
diff --git a/Engine/Source/Programs/UnrealBuildAccelerator/Test/Private/UbaTestBasics.cpp b/Engine/Source/Programs/UnrealBuildAccelerator/Test/Private/UbaTestBasics.cpp
--- a/Engine/Source/Programs/UnrealBuildAccelerator/Test/Private/UbaTestBasics.cpp
+++ b/Engine/Source/Programs/UnrealBuildAccelerator/Test/Private/UbaTestBasics.cpp
@@ -430,11 +430,13 @@ namespace uba
 		UBA_TEST_CHECK(arguments10.size() == 2, "ParseArguments 10 failed");
 		UBA_TEST_CHECK(Equals(arguments10[1].data(), TC("\"foo\"")), "ParseArguments 10 failed");
 
+		#if PLATFORM_WINDOWS
 		Vector<TString> arguments11;
 		ParseArguments(arguments11, TC("\\\"a\\\\b\\\" \\\"c\\\\d\\\" meh"));
 		UBA_TEST_CHECK(arguments11.size() == 3, "ParseArguments 11 failed");
 		UBA_TEST_CHECK(Equals(arguments11[1].data(), TC("\"c\\\\d\"")), "ParseArguments 11 failed");
 		UBA_TEST_CHECK(Equals(arguments11[2].data(), TC("meh")), "ParseArguments 11 failed");
+		#endif
 		return true;
 	}
 
diff --git a/Engine/Source/Programs/UnrealBuildAccelerator/Test/Private/UbaTestSession.cpp b/Engine/Source/Programs/UnrealBuildAccelerator/Test/Private/UbaTestSession.cpp
--- a/Engine/Source/Programs/UnrealBuildAccelerator/Test/Private/UbaTestSession.cpp
+++ b/Engine/Source/Programs/UnrealBuildAccelerator/Test/Private/UbaTestSession.cpp
@@ -938,7 +938,7 @@ namespace uba
 		return RunLocal(logger, testRootDir, WriteFileMap);
 	}
 
-	bool TestFileMapWriteRemote(LoggerWithWriter& logger, const StringBufferBase& testRootDir)
+	bool TestRemoteFileMapWrite(LoggerWithWriter& logger, const StringBufferBase& testRootDir)
 	{
 		return RunRemote(logger, testRootDir, WriteFileMap);
 	}
@@ -973,7 +973,7 @@ namespace uba
 		return true;
 	}
 
-	bool TestOverlayAccessOrderingRemote(LoggerWithWriter& logger, const StringBufferBase& testRootDir)
+	bool TestRemoteOverlayAccessOrdering(LoggerWithWriter& logger, const StringBufferBase& testRootDir)
 	{
 		Config c;
 		CHECK_TRUE(RunRemote(logger, testRootDir, c, RunOverlayAccessOrdering, 1));
@@ -981,4 +981,29 @@ namespace uba
 		CHECK_TRUE(RunRemote(logger, testRootDir, c, RunOverlayAccessOrdering, 3));
 		return true;
 	}
+
+	#if PLATFORM_LINUX
+	bool RunSymlinks(LoggerWithWriter& logger, SessionServer& session, const tchar* workingDir, const RunProcessFunction& runProcess, uintptr_t extraData)
+	{
+		CHECK_TRUE(RunTestApp(logger, workingDir, runProcess, TC("-createSymlink"), nullptr, TC("")));
+		StringBuffer<> newLink(workingDir);
+		newLink.EnsureEndsWithSlash().Append(TCV("new_file"));
+		struct stat st;
+		CHECK_TRUE(lstat(newLink.data, &st) == 0);
+		CHECK_TRUE(S_ISLNK(st.st_mode));
+		return true;
+	}
+
+	bool TestDetouredSymlinks(LoggerWithWriter& logger, const StringBufferBase& testRootDir)
+	{
+		CHECK_TRUE(RunLocal(logger, testRootDir, RunSymlinks));
+		return true;
+	}
+
+	bool TestRemoteDetouredSymlinks(LoggerWithWriter& logger, const StringBufferBase& testRootDir)
+	{
+		CHECK_TRUE(RunRemote(logger, testRootDir, RunSymlinks));
+		return true;
+	}
+	#endif
 }
diff --git a/Engine/Source/Programs/UnrealBuildAccelerator/TestApp/Private/PosixOS/UbaTestAppPosix.cpp b/Engine/Source/Programs/UnrealBuildAccelerator/TestApp/Private/PosixOS/UbaTestAppPosix.cpp
--- a/Engine/Source/Programs/UnrealBuildAccelerator/TestApp/Private/PosixOS/UbaTestAppPosix.cpp
+++ b/Engine/Source/Programs/UnrealBuildAccelerator/TestApp/Private/PosixOS/UbaTestAppPosix.cpp
@@ -23,13 +23,25 @@ int LogError(const char* format, ...)
 {
 	va_list args;
 	va_start(args, format);
-	char buffer[1024];
-	snprintf(buffer, 1024, format, args);
+	char buffer[2048];
+	vsnprintf(buffer, 2048, format, args);
 	printf("%s\n", buffer);
 	va_end(args);
 	return -1;
 }
 
+#define CHECK_TRUE(x) \
+	do { \
+	if (!(x)) \
+		return LogError("Failed %s (%s:%u)", #x, __FILE__, __LINE__); \
+	} while (false)
+
+#define CHECK_TRUEF(x, ...) \
+	do { \
+	if (!(x)) \
+		return LogError(##__VA_ARGS__); \
+	} while (false)
+
 bool IsLibraryLoaded(const char* libraryToMatch)
 	{
 		bool foundLib = false;
@@ -178,6 +190,10 @@ int main(int argc, char* argv[])
 		if (!S_ISDIR(attrFoo.st_mode))
 			return LogError("stat for dir 'FooDir' returned wrong type");
 
+		// Add file to FooDir
+		// Try to rmdir while files are there
+		// Delete file
+
 		if (rmdir("FooDir") != 0)
 			return LogError("rmdir 'FooDir' failed");
 
@@ -359,6 +375,14 @@ int main(int argc, char* argv[])
 		if (memcmp(data, "FOO", 3) != 0)
 			return LogError("File %s has wrong content", file);
 	}
+	else if (strcmp(argv[1], "-createSymlink") == 0)
+	{
+		int wh = open("empty_file", O_WRONLY | O_CREAT | O_TRUNC, S_IRUSR | S_IWUSR);
+		CHECK_TRUE(wh != -1);
+		CHECK_TRUE(close(wh) != -1);
+		CHECK_TRUE(symlink("empty_file", "new_file") != -1);
+		return 0;
+	}
 	#if PLATFORM_MAC
 	else if (strcmp(argv[1], "-xcode-select") == 0)
 	{
diff --git a/Engine/Source/Programs/UnrealBuildAccelerator/TestApp/Private/Windows/UbaTestAppWin.cpp b/Engine/Source/Programs/UnrealBuildAccelerator/TestApp/Private/Windows/UbaTestAppWin.cpp
--- a/Engine/Source/Programs/UnrealBuildAccelerator/TestApp/Private/Windows/UbaTestAppWin.cpp
+++ b/Engine/Source/Programs/UnrealBuildAccelerator/TestApp/Private/Windows/UbaTestAppWin.cpp
@@ -169,6 +169,15 @@ int wmain(int argc, wchar_t* argv[])
 		}
 
 		{
+#ifdef _DEBUG
+			if (CreateDirectoryW(L"c:\\", NULL))
+				return LogError(L"Should fail to create root dir");
+			if (GetLastError() != ERROR_ALREADY_EXISTS)
+				return LogError(L"Expected ERROR_ALREADY_EXISTS for root dir");
+			if (CreateDirectoryW(L"a:\\", NULL))
+				return LogError(L"Should fail to create root dir");
+#endif
+
 			if (!CreateDirectoryW(L"DirA", NULL))
 				return LogError(L"Failed to create directory");
 			if (CreateDirectoryW(L"DirA", NULL))
@@ -177,6 +186,16 @@ int wmain(int argc, wchar_t* argv[])
 			if (GetFileAttributesW(L"DirA") == 0)
 				return LogError(L"Failed to get attributes of directory DirA");
 
+			if (WriteTestFile("", 0, L"DirA\\File"))
+				return -1;
+
+			if (RemoveDirectoryW(L"DirA"))
+				return LogError(L"Should not succeed in removing directory with files in it");
+			if (GetLastError() != ERROR_DIR_NOT_EMPTY)
+				return LogError(L"Wrong last error code");
+			if (!DeleteFileW(L"DirA\\File"))
+				return LogError(L"Failed to delete file");
+
 			if (!RemoveDirectoryW(L"DirA"))
 				return LogError(L"Failed to remove directory");
 
